Iris experimental replication system.
Initial submit, compiled out by default.

diff --git a/Engine/Plugins/Experimental/Iris/Iris.uplugin b/Engine/Plugins/Experimental/Iris/Iris.uplugin
new file mode 100644
--- /dev/null
+++ b/Engine/Plugins/Experimental/Iris/Iris.uplugin
@@ -0,0 +1,24 @@
+{
+    "FileVersion" : 3,
+    "Version" : 1,
+    "VersionName" : "0.1",
+    "FriendlyName" : "Iris",
+    "Description" : "Iris networking.",
+    "Category": "Networking",
+    "CreatedBy" : "Epic Games, Inc.",
+    "CreatedByURL" : "https://epicgames.com",
+    "DocsURL" : "",
+    "MarketplaceURL" : "",
+    "SupportURL" : "",
+    "EnabledByDefault" : true,
+    "CanContainContent" : false,
+    "IsBetaVersion" : true,
+    "Installed" : false,
+    "Modules": [
+      {
+        "Name": "Iris",
+        "Type": "Runtime",
+        "LoadingPhase": "Default"
+      }
+    ]
+}
diff --git a/Engine/Plugins/Online/OnlineFramework/Source/Lobby/Lobby.Build.cs b/Engine/Plugins/Online/OnlineFramework/Source/Lobby/Lobby.Build.cs
--- a/Engine/Plugins/Online/OnlineFramework/Source/Lobby/Lobby.Build.cs
+++ b/Engine/Plugins/Online/OnlineFramework/Source/Lobby/Lobby.Build.cs
@@ -32,5 +32,7 @@ public class Lobby : ModuleRules
 		else if (Target.Platform == UnrealTargetPlatform.Mac)
 		{
 		}
+
+		SetupIrisSupport(Target);
 	}
 }
diff --git a/Engine/Plugins/Online/OnlineSubsystemUtils/Source/OnlineSubsystemUtils/Private/OnlineBeaconClient.cpp b/Engine/Plugins/Online/OnlineSubsystemUtils/Source/OnlineSubsystemUtils/Private/OnlineBeaconClient.cpp
--- a/Engine/Plugins/Online/OnlineSubsystemUtils/Source/OnlineSubsystemUtils/Private/OnlineBeaconClient.cpp
+++ b/Engine/Plugins/Online/OnlineSubsystemUtils/Source/OnlineSubsystemUtils/Private/OnlineBeaconClient.cpp
@@ -13,6 +13,7 @@
 #include "Misc/NetworkVersion.h"
 #include "Interfaces/OnlineIdentityInterface.h"
 #include "OnlineSubsystemUtils.h"
+#include "Containers/StringFwd.h"
 
 #define BEACON_RPC_TIMEOUT 15.0f
 
@@ -260,7 +261,10 @@ void AOnlineBeaconClient::SendInitialJoin()
 			EncryptionData.Identifier.Empty();
 		}
 
-		FNetControlMessage<NMT_Hello>::Send(NetDriver->ServerConnection, IsLittleEndian, LocalNetworkVersion, EncryptionData.Identifier);
+
+		EEngineNetworkRuntimeFeatures LocalNetworkFeatures = NetDriver->GetNetworkRuntimeFeatures();
+		FNetControlMessage<NMT_Hello>::Send(NetDriver->ServerConnection, IsLittleEndian, LocalNetworkVersion, EncryptionData.Identifier, LocalNetworkFeatures);
+		
 
 		NetDriver->ServerConnection->FlushNet();
 	}
@@ -455,9 +459,21 @@ void AOnlineBeaconClient::NotifyControlMessage(UNetConnection* Connection, uint8
 			{
 				// Report mismatch.
 				uint32 RemoteNetworkVersion;
+				EEngineNetworkRuntimeFeatures RemoteNetworkFeatures = EEngineNetworkRuntimeFeatures::None;
 
-				if (FNetControlMessage<NMT_Upgrade>::Receive(Bunch, RemoteNetworkVersion))
+				if (FNetControlMessage<NMT_Upgrade>::Receive(Bunch, RemoteNetworkVersion, RemoteNetworkFeatures))
 				{
+					TStringBuilder<128> RemoteFeaturesDescription;
+					FNetworkVersion::DescribeNetworkRuntimeFeaturesBitset(RemoteNetworkFeatures, RemoteFeaturesDescription);
+
+					TStringBuilder<128> LocalFeaturesDescription;
+					FNetworkVersion::DescribeNetworkRuntimeFeaturesBitset(NetDriver->GetNetworkRuntimeFeatures(), LocalFeaturesDescription);
+
+					UE_LOG(LogBeacon, Error, TEXT("Beacon is incompatible with the local version of the game: RemoteNetworkVersion=%u, RemoteNetworkFeatures=%s vs LocalNetworkVersion=%u, LocalNetworkFeatures=%s"), 
+						RemoteNetworkVersion, RemoteFeaturesDescription.ToString(),
+						FNetworkVersion::GetLocalNetworkVersion(), LocalFeaturesDescription.ToString()
+					);
+				
 					// Upgrade
 					const FString ConnectionError = NSLOCTEXT("Engine", "ClientOutdated",
 						"The match you are trying to join is running an incompatible version of the game.  Please try upgrading your game version.").ToString();
diff --git a/Engine/Plugins/Online/OnlineSubsystemUtils/Source/OnlineSubsystemUtils/Private/OnlineBeaconHost.cpp b/Engine/Plugins/Online/OnlineSubsystemUtils/Source/OnlineSubsystemUtils/Private/OnlineBeaconHost.cpp
--- a/Engine/Plugins/Online/OnlineSubsystemUtils/Source/OnlineSubsystemUtils/Private/OnlineBeaconHost.cpp
+++ b/Engine/Plugins/Online/OnlineSubsystemUtils/Source/OnlineSubsystemUtils/Private/OnlineBeaconHost.cpp
@@ -169,7 +169,10 @@ bool AOnlineBeaconHost::HandleControlMessage(UNetConnection* Connection, uint8 M
 			uint8 IsLittleEndian = 0;
 			uint32 RemoteNetworkVersion = 0;
 			FString EncryptionToken;
-			if (!FNetControlMessage<NMT_Hello>::Receive(Bunch, IsLittleEndian, RemoteNetworkVersion, EncryptionToken))
+			EEngineNetworkRuntimeFeatures LocalNetworkFeatures = NetDriver->GetNetworkRuntimeFeatures();
+			EEngineNetworkRuntimeFeatures RemoteNetworkFeatures = EEngineNetworkRuntimeFeatures::None;
+	
+			if (!FNetControlMessage<NMT_Hello>::Receive(Bunch, IsLittleEndian, RemoteNetworkVersion, EncryptionToken, RemoteNetworkFeatures))
 			{
 				SendFailurePacket(Connection, ENetCloseResult::BeaconUnableToParsePacket,
 									FText::Format(Error_UnableToParsePacket, FText::FromString(TEXT("NMT_Hello"))));
@@ -179,13 +182,25 @@ bool AOnlineBeaconHost::HandleControlMessage(UNetConnection* Connection, uint8 M
 
 			// check for net compatibility (in this case sent NMT_Upgrade)
 			uint32 LocalNetworkVersion = FNetworkVersion::GetLocalNetworkVersion();
-			if (!FNetworkVersion::IsNetworkCompatible(LocalNetworkVersion, RemoteNetworkVersion))
+			const bool bIsCompatible = FNetworkVersion::IsNetworkCompatible(LocalNetworkVersion, RemoteNetworkVersion) && FNetworkVersion::AreNetworkRuntimeFeaturesCompatible(LocalNetworkFeatures, RemoteNetworkFeatures);
+
+			if (!bIsCompatible)
 			{
-				UE_LOG(LogBeacon, Error, TEXT("%s: Client not network compatible (Local=%d, Remote=%d)"), *GetDebugName(Connection), LocalNetworkVersion, RemoteNetworkVersion);
+				TStringBuilder<128> LocalNetFeaturesDescription;
+				TStringBuilder<128> RemoteNetFeaturesDescription;
+
+				FNetworkVersion::DescribeNetworkRuntimeFeaturesBitset(LocalNetworkFeatures, LocalNetFeaturesDescription);
+				FNetworkVersion::DescribeNetworkRuntimeFeaturesBitset(RemoteNetworkFeatures, RemoteNetFeaturesDescription);
+
+				UE_LOG(LogBeacon, Error, TEXT("Client not network compatible %s: LocalNetVersion=%u, RemoteNetVersion=%u, LocalNetFeatures=%s, RemoteNetFeatures=%s"), 
+					*GetDebugName(Connection), 
+					LocalNetworkVersion, RemoteNetworkVersion,
+					LocalNetFeaturesDescription.ToString(), RemoteNetFeaturesDescription.ToString()
+				);
 				FNetControlMessage<NMT_Upgrade>::Send(Connection, LocalNetworkVersion);
 				return false;
 			}
-
+			
 			// if the client didn't specify an encryption token we're done with Hello
 			if (EncryptionToken.IsEmpty())
 			{
diff --git a/Engine/Plugins/Runtime/GameplayAbilities/Source/GameplayAbilities/GameplayAbilities.Build.cs b/Engine/Plugins/Runtime/GameplayAbilities/Source/GameplayAbilities/GameplayAbilities.Build.cs
--- a/Engine/Plugins/Runtime/GameplayAbilities/Source/GameplayAbilities/GameplayAbilities.Build.cs
+++ b/Engine/Plugins/Runtime/GameplayAbilities/Source/GameplayAbilities/GameplayAbilities.Build.cs
@@ -46,6 +46,8 @@ namespace UnrealBuildTool.Rules
 			{
 				PublicDefinitions.Add("WITH_GAMEPLAY_DEBUGGER=0");
 			}
+
+			SetupIrisSupport(Target);
 		}
 	}
 }
diff --git a/Engine/Plugins/Runtime/GameplayAbilities/Source/GameplayAbilities/Private/Abilities/GameplayAbility.cpp b/Engine/Plugins/Runtime/GameplayAbilities/Source/GameplayAbilities/Private/Abilities/GameplayAbility.cpp
--- a/Engine/Plugins/Runtime/GameplayAbilities/Source/GameplayAbilities/Private/Abilities/GameplayAbility.cpp
+++ b/Engine/Plugins/Runtime/GameplayAbilities/Source/GameplayAbilities/Private/Abilities/GameplayAbility.cpp
@@ -14,6 +14,10 @@
 #include "GameFramework/CharacterMovementComponent.h"
 #include "GameFramework/Character.h"
 
+#if UE_WITH_IRIS
+#include "Iris/ReplicationSystem/ReplicationFragmentUtil.h"
+#endif // UE_WITH_IRIS
+
 namespace FAbilitySystemTweaks
 {
 	int ClearAbilityTimers = 1;
@@ -1977,3 +1981,13 @@ void UGameplayAbility::NotifyAbilityTaskWaitingOnAvatar(class UAbilityTask* Abil
 		AbilitySystemComponent->ForceCancelAbilityDueToReplication(this);
 	}
 }
+
+#if UE_WITH_IRIS
+void UGameplayAbility::RegisterReplicationFragments(UE::Net::FFragmentRegistrationContext& Context, UE::Net::EFragmentRegistrationFlags RegistrationFlags)
+{
+	Super::RegisterReplicationFragments(Context, RegistrationFlags);
+
+	// Build descriptors and allocate PropertyReplicationFragments for this object
+	UE::Net::FReplicationFragmentUtil::CreateAndRegisterFragmentsForObject(this, Context, RegistrationFlags);
+}
+#endif // UE_WITH_IRIS
diff --git a/Engine/Plugins/Runtime/GameplayAbilities/Source/GameplayAbilities/Private/AbilitySystemComponent.cpp b/Engine/Plugins/Runtime/GameplayAbilities/Source/GameplayAbilities/Private/AbilitySystemComponent.cpp
--- a/Engine/Plugins/Runtime/GameplayAbilities/Source/GameplayAbilities/Private/AbilitySystemComponent.cpp
+++ b/Engine/Plugins/Runtime/GameplayAbilities/Source/GameplayAbilities/Private/AbilitySystemComponent.cpp
@@ -2847,7 +2847,7 @@ void UAbilitySystemComponent::SetSpawnedAttributes(const TArray<UAttributeSet*>&
 
 void UAbilitySystemComponent::AddSpawnedAttribute(UAttributeSet* Attribute)
 {
-	if (Attribute == nullptr)
+	if (!IsValid(Attribute))
 	{
 		return;
 	}
@@ -2927,7 +2927,7 @@ void UAbilitySystemComponent::OnRep_SpawnedAttributes(const TArray<UAttributeSet
 		// Find the attributes that got added
 		for (UAttributeSet* NewAttributeSet : SpawnedAttributes)
 		{
-			if (NewAttributeSet)
+			if (IsValid(NewAttributeSet))
 			{
 				const bool bIsAdded = PreviousSpawnedAttributes.Find(NewAttributeSet) == INDEX_NONE;
 				if (bIsAdded)
diff --git a/Engine/Plugins/Runtime/GameplayAbilities/Source/GameplayAbilities/Private/AbilitySystemGlobals.cpp b/Engine/Plugins/Runtime/GameplayAbilities/Source/GameplayAbilities/Private/AbilitySystemGlobals.cpp
--- a/Engine/Plugins/Runtime/GameplayAbilities/Source/GameplayAbilities/Private/AbilitySystemGlobals.cpp
+++ b/Engine/Plugins/Runtime/GameplayAbilities/Source/GameplayAbilities/Private/AbilitySystemGlobals.cpp
@@ -12,6 +12,12 @@
 #include "GameplayTagsManager.h"
 #include "UObject/UObjectIterator.h"
 
+#if UE_WITH_IRIS
+#include "Serialization/GameplayAbilityTargetDataHandleNetSerializer.h"
+#include "Serialization/GameplayEffectContextHandleNetSerializer.h"
+#include "Serialization/PredictionKeyNetSerializer.h"
+#endif // UE_WITH_IRIS
+
 #if WITH_EDITOR
 #include "Editor.h"
 #endif
@@ -252,6 +258,10 @@ UFunction* UAbilitySystemGlobals::GetGameplayCueFunction(const FGameplayTag& Chi
 void UAbilitySystemGlobals::InitTargetDataScriptStructCache()
 {
 	TargetDataStructCache.InitForType(FGameplayAbilityTargetData::StaticStruct());
+#if UE_WITH_IRIS
+	InitGameplayAbilityTargetDataHandleNetSerializerTypeCache();
+	InitGameplayEffectContextHandleNetSerializerTypeCache();
+#endif // UE_WITH_IRIS
 }
 
 // --------------------------------------------------------------------
diff --git a/Engine/Plugins/Runtime/GameplayAbilities/Source/GameplayAbilities/Private/AttributeSet.cpp b/Engine/Plugins/Runtime/GameplayAbilities/Source/GameplayAbilities/Private/AttributeSet.cpp
--- a/Engine/Plugins/Runtime/GameplayAbilities/Source/GameplayAbilities/Private/AttributeSet.cpp
+++ b/Engine/Plugins/Runtime/GameplayAbilities/Source/GameplayAbilities/Private/AttributeSet.cpp
@@ -16,6 +16,10 @@
 #include "Net/Core/PushModel/PushModel.h"
 #include "UObject/UObjectThreadContext.h"
 
+#if UE_WITH_IRIS
+#include "Iris/ReplicationSystem/ReplicationFragmentUtil.h"
+#endif // UE_WITH_IRIS
+
 
 #if ENABLE_VISUAL_LOG
 namespace
@@ -371,6 +375,18 @@ void UAttributeSet::SetNetAddressable()
 	bNetAddressable = true;
 }
 
+#if UE_WITH_IRIS
+void UAttributeSet::RegisterReplicationFragments(UE::Net::FFragmentRegistrationContext& Context, UE::Net::EFragmentRegistrationFlags RegistrationFlags)
+{
+	using namespace UE::Net;
+
+	Super::RegisterReplicationFragments(Context, RegistrationFlags);
+
+	// Build descriptors and allocate PropertyReplicationFragments for this object
+	FReplicationFragmentUtil::CreateAndRegisterFragmentsForObject(this, Context, RegistrationFlags);
+}
+#endif // UE_WITH_IRIS
+
 void UAttributeSet::InitFromMetaDataTable(const UDataTable* DataTable)
 {
 	static const FString Context = FString(TEXT("UAttribute::BindToMetaDataTable"));
diff --git a/Engine/Plugins/Runtime/GameplayAbilities/Source/GameplayAbilities/Private/GameplayEffect.cpp b/Engine/Plugins/Runtime/GameplayAbilities/Source/GameplayAbilities/Private/GameplayEffect.cpp
--- a/Engine/Plugins/Runtime/GameplayAbilities/Source/GameplayAbilities/Private/GameplayEffect.cpp
+++ b/Engine/Plugins/Runtime/GameplayAbilities/Source/GameplayAbilities/Private/GameplayEffect.cpp
@@ -4085,14 +4085,20 @@ bool FActiveGameplayEffectsContainer::NetDeltaSerialize(FNetDeltaSerializeInfo&
 
 	bool RetVal = FastArrayDeltaSerialize<FActiveGameplayEffect>(GameplayEffects_Internal, DeltaParms, *this);
 
-	// After the array has been replicated, invoke GC events ONLY if the effect is not inhibited
-	// We postpone this check because in the same net update we could receive multiple GEs that affect if one another is inhibited
+	// This section has been moved into new PostReplicatedReceive() method that is invoked after every call to FastArrayDeltaSerialize<> that results in data being modified
 	
-	if (DeltaParms.Writer == nullptr && Owner != nullptr)
+	return RetVal;
+}
+
+void FActiveGameplayEffectsContainer::PostReplicatedReceive(const FFastArraySerializer::FPostReplicatedReceiveParameters& Parameters)
+{
+	// After the array has been replicated, invoke GC events ONLY if the effect is not inhibited
+	// We postpone this check because in the same net update we could receive multiple GEs that affect if one another is inhibited	
+	if (Owner != nullptr)
 	{
 		QUICK_SCOPE_CYCLE_COUNTER(STAT_ActiveGameplayEffectsContainer_NetDeltaSerialize_CheckRepGameplayCues);
 
-		if (!DeltaParms.bOutHasMoreUnmapped) // Do not invoke GCs when we have missing information (like AActor*s in EffectContext)
+		if (!Parameters.bHasMoreUnmappedReferences) // Do not invoke GCs when we have missing information (like AActor*s in EffectContext)
 		{
 			if (Owner->IsReadyForGameplayCues())
 			{
@@ -4100,8 +4106,6 @@ bool FActiveGameplayEffectsContainer::NetDeltaSerialize(FNetDeltaSerializeInfo&
 			}
 		}
 	}
-
-	return RetVal;
 }
 
 void FActiveGameplayEffectsContainer::Uninitialize()
diff --git a/Engine/Plugins/Runtime/GameplayAbilities/Source/GameplayAbilities/Public/Abilities/GameplayAbility.h b/Engine/Plugins/Runtime/GameplayAbilities/Source/GameplayAbilities/Public/Abilities/GameplayAbility.h
--- a/Engine/Plugins/Runtime/GameplayAbilities/Source/GameplayAbilities/Public/Abilities/GameplayAbility.h
+++ b/Engine/Plugins/Runtime/GameplayAbilities/Source/GameplayAbilities/Public/Abilities/GameplayAbility.h
@@ -510,6 +510,11 @@ public:
 	virtual bool CallRemoteFunction(UFunction* Function, void* Parameters, FOutParmRec* OutParms, FFrame* Stack) override;
 	virtual bool IsSupportedForNetworking() const override;
 
+#if UE_WITH_IRIS
+	/** Register all replication fragments */
+	virtual void RegisterReplicationFragments(UE::Net::FFragmentRegistrationContext& Context, UE::Net::EFragmentRegistrationFlags RegistrationFlags) override;
+#endif // UE_WITH_IRIS
+
 	// --------------------------------------
 	//	IGameplayTaskOwnerInterface
 	// --------------------------------------	
diff --git a/Engine/Plugins/Runtime/GameplayAbilities/Source/GameplayAbilities/Public/AttributeSet.h b/Engine/Plugins/Runtime/GameplayAbilities/Source/GameplayAbilities/Public/AttributeSet.h
--- a/Engine/Plugins/Runtime/GameplayAbilities/Source/GameplayAbilities/Public/AttributeSet.h
+++ b/Engine/Plugins/Runtime/GameplayAbilities/Source/GameplayAbilities/Public/AttributeSet.h
@@ -252,6 +252,10 @@ public:
 	virtual void PreNetReceive() override;
 	virtual void PostNetReceive() override;
 
+#if UE_WITH_IRIS
+	/** Register all replication fragments */
+	virtual void RegisterReplicationFragments(UE::Net::FFragmentRegistrationContext& Context, UE::Net::EFragmentRegistrationFlags RegistrationFlags) override;
+#endif // UE_WITH_IRIS
 protected:
 	/** Is this attribute set safe to ID over the network by name?  */
 	uint32 bNetAddressable : 1;
diff --git a/Engine/Plugins/Runtime/GameplayAbilities/Source/GameplayAbilities/Public/GameplayCueInterface.h b/Engine/Plugins/Runtime/GameplayAbilities/Source/GameplayAbilities/Public/GameplayCueInterface.h
--- a/Engine/Plugins/Runtime/GameplayAbilities/Source/GameplayAbilities/Public/GameplayCueInterface.h
+++ b/Engine/Plugins/Runtime/GameplayAbilities/Source/GameplayAbilities/Public/GameplayCueInterface.h
@@ -152,7 +152,7 @@ private:
 
 	int32 GetGameStateTime(const UWorld* World) const;
 
-	UPROPERTY()
+	UPROPERTY(NotReplicated)
 	class UAbilitySystemComponent*	Owner = nullptr;
 	
 	friend struct FActiveGameplayCue;
diff --git a/Engine/Plugins/Runtime/GameplayAbilities/Source/GameplayAbilities/Public/GameplayEffect.h b/Engine/Plugins/Runtime/GameplayAbilities/Source/GameplayAbilities/Public/GameplayEffect.h
--- a/Engine/Plugins/Runtime/GameplayAbilities/Source/GameplayAbilities/Public/GameplayEffect.h
+++ b/Engine/Plugins/Runtime/GameplayAbilities/Source/GameplayAbilities/Public/GameplayEffect.h
@@ -1744,6 +1744,9 @@ struct GAMEPLAYABILITIES_API FActiveGameplayEffectsContainer : public FFastArray
 	/** Recomputes the start time for all active abilities */
 	void RecomputeStartWorldTimes(const float WorldTime, const float ServerWorldTime);
 
+	/** Called every time data has been modified by the FastArraySerializer */
+	void PostReplicatedReceive(const FFastArraySerializer::FPostReplicatedReceiveParameters& Parameters);
+
 private:
 
 	/**
@@ -1852,7 +1855,7 @@ private:
 	FGameplayTagCountContainer ApplicationImmunityGameplayTagCountContainer;
 
 	/** Active GEs that have immunity queries. This is an acceleration list to avoid searching through the Active GameplayEffect list frequetly. (We only search for the active GE if immunity procs) */
-	UPROPERTY()
+	UPROPERTY(NotReplicated)
 	TArray<const UGameplayEffect*> ApplicationImmunityQueryEffects;
 
 	FAggregatorRef& FindOrCreateAttributeAggregator(FGameplayAttribute Attribute);
diff --git a/Engine/Plugins/Runtime/GameplayAbilities/Source/GameplayAbilities/Public/GameplayEffectTypes.h b/Engine/Plugins/Runtime/GameplayAbilities/Source/GameplayAbilities/Public/GameplayEffectTypes.h
--- a/Engine/Plugins/Runtime/GameplayAbilities/Source/GameplayAbilities/Public/GameplayEffectTypes.h
+++ b/Engine/Plugins/Runtime/GameplayAbilities/Source/GameplayAbilities/Public/GameplayEffectTypes.h
@@ -28,6 +28,10 @@ class UGameplayAbility;
 struct FActiveGameplayEffect;
 struct FGameplayEffectModCallbackData;
 struct FGameplayEffectSpec;
+namespace UE::Net
+{
+	struct FGameplayEffectContextHandleAccessorForNetSerializer;
+}
 
 /** Wrappers to convert enum to string. These are fairly slow */
 GAMEPLAYABILITIES_API FString EGameplayModOpToString(int32 Type);
@@ -881,6 +885,8 @@ struct GAMEPLAYABILITIES_API FGameplayEffectContextHandle
 	bool NetSerialize(FArchive& Ar, class UPackageMap* Map, bool& bOutSuccess);
 
 private:
+	friend UE::Net::FGameplayEffectContextHandleAccessorForNetSerializer;
+
 	TSharedPtr<FGameplayEffectContext> Data;
 };
 
diff --git a/Engine/Plugins/Runtime/GameplayAbilities/Source/GameplayAbilities/Public/GameplayPrediction.h b/Engine/Plugins/Runtime/GameplayAbilities/Source/GameplayAbilities/Public/GameplayPrediction.h
--- a/Engine/Plugins/Runtime/GameplayAbilities/Source/GameplayAbilities/Public/GameplayPrediction.h
+++ b/Engine/Plugins/Runtime/GameplayAbilities/Source/GameplayAbilities/Public/GameplayPrediction.h
@@ -6,6 +6,10 @@
 #include "GameplayPrediction.generated.h"
 
 class UAbilitySystemComponent;
+namespace UE::Net
+{
+	struct FPredictionKeyNetSerializer;
+}
 
 DECLARE_DELEGATE(FPredictionKeyEvent);
 
@@ -395,6 +399,7 @@ struct GAMEPLAYABILITIES_API FPredictionKey
 	UPTRINT GetPredictiveConnectionKey() const { return PredictiveConnectionKey; }
 
 private:
+	friend UE::Net::FPredictionKeyNetSerializer;
 
 	void GenerateNewPredictionKey();
 
diff --git a/Engine/Plugins/Runtime/NetworkPrediction/Source/NetworkPrediction/NetworkPrediction.Build.cs b/Engine/Plugins/Runtime/NetworkPrediction/Source/NetworkPrediction/NetworkPrediction.Build.cs
--- a/Engine/Plugins/Runtime/NetworkPrediction/Source/NetworkPrediction/NetworkPrediction.Build.cs
+++ b/Engine/Plugins/Runtime/NetworkPrediction/Source/NetworkPrediction/NetworkPrediction.Build.cs
@@ -57,6 +57,7 @@ namespace UnrealBuildTool.Rules
                 });
             }
 
-        }
+			SetupIrisSupport(Target);
+		}
 	}
 }
diff --git a/Engine/Plugins/Runtime/NetworkPrediction/Source/NetworkPrediction/Private/NetworkPredictionReplicatedManager.cpp b/Engine/Plugins/Runtime/NetworkPrediction/Source/NetworkPrediction/Private/NetworkPredictionReplicatedManager.cpp
--- a/Engine/Plugins/Runtime/NetworkPrediction/Source/NetworkPrediction/Private/NetworkPredictionReplicatedManager.cpp
+++ b/Engine/Plugins/Runtime/NetworkPrediction/Source/NetworkPrediction/Private/NetworkPredictionReplicatedManager.cpp
@@ -12,7 +12,8 @@ ANetworkPredictionReplicatedManager::ANetworkPredictionReplicatedManager()
 {
 	bReplicates = true;
 	NetPriority = 1000.f; // We want this to be super high priority when it replicates
-	NetUpdateFrequency = 0.001f; // Very low frequency: we will use ForceNetUpdate when important data changes
+	// Mute very low update frequency ensure. Was 0.001f .
+	NetUpdateFrequency = 0.125f; // Low frequency: we will use ForceNetUpdate when important data changes
 	bAlwaysRelevant = true;
 }
 
diff --git a/Engine/Plugins/Runtime/ReplicationGraph/Source/Private/ReplicationGraph.cpp b/Engine/Plugins/Runtime/ReplicationGraph/Source/Private/ReplicationGraph.cpp
--- a/Engine/Plugins/Runtime/ReplicationGraph/Source/Private/ReplicationGraph.cpp
+++ b/Engine/Plugins/Runtime/ReplicationGraph/Source/Private/ReplicationGraph.cpp
@@ -1784,14 +1784,6 @@ int64 UReplicationGraph::ReplicateSingleActor_FastShared(AActor* Actor, FConnect
 	{
 		TGuardValue<bool> Guard(ActorChannel->bHoldQueuedExportBunchesAndGUIDs, true);		// Don't export queued GUIDs in fast path
 		
-#if !(UE_BUILD_SHIPPING || UE_BUILD_TEST) && (UE_NET_TRACE_ENABLED)
-		// If we want to trace the actual contents of the shared path, we need to create a temporary collector associated with the shared bunch, collect the shared data and report it every time it is reused
-		// For now we just set the debug name of the the bunch to be able to give some context on what this bunch contains
-		if (OutBunch.DebugString.IsEmpty())
-		{
-			OutBunch.DebugString = FString(TEXT("ReplicateSingleActor_FastShared"));
-		}
-#endif
 		ActorChannel->SendBunch(&OutBunch, false);
 	}
 
@@ -2087,14 +2079,50 @@ bool UReplicationGraph::ProcessRemoteFunction(class AActor* Actor, UFunction* Fu
 
 			// It sucks we have to a temp writer like this, but we don't know how big the payload will be until we serialize it
 			FNetBitWriter TempWriter(nullptr, 0);
+
+#if UE_NET_TRACE_ENABLED
+			// Create trace collector if tracing is enabled for the target bunch
+			FNetTraceCollector* Collector = GetTraceCollector(*FastSharedReplicationBunch);
+			if (Collector)
+			{
+				Collector->Reset();
+			}
+			else
+			{
+				Collector = UE_NET_TRACE_CREATE_COLLECTOR(ENetTraceVerbosity::Trace);
+				SetTraceCollector(*FastSharedReplicationBunch, Collector);
+			}
+
+			// We use the collector from the shared bunch
+			SetTraceCollector(TempWriter, Collector);
+#endif // UE_NET_TRACE_ENABLED
+
 			TSharedPtr<FRepLayout> RepLayout = NetDriver->GetFunctionRepLayout( Function );
 			RepLayout->SendPropertiesForRPC(Function, FastSharedReplicationChannel, TempWriter, Parameters);
 
 			FNetBitWriter TempBlockWriter(nullptr, 0);
+
+#if UE_NET_TRACE_ENABLED
+			// Ugliness to get data reported correctly, we basically fold the data from TempWriter into TempBlockWriter and then to Bunch
+			// Create trace collector if tracing is enabled for the target bunch			
+			SetTraceCollector(TempBlockWriter, Collector ? UE_NET_TRACE_CREATE_COLLECTOR(ENetTraceVerbosity::Trace) : nullptr);
+#endif // UE_NET_TRACE_ENABLED
+
 			FastSharedReplicationChannel->WriteFieldHeaderAndPayload( TempBlockWriter, ClassCache, FieldCache, nullptr, TempWriter, true );
 
+#if UE_NET_TRACE_ENABLED
+			// As we have used the collector we need to reset it before injecting the final data
+			if (Collector)
+			{
+				Collector->Reset();
+			}
+			UE_NET_TRACE_OBJECT_SCOPE(FastSharedReplicationChannel->ActorNetGUID, *FastSharedReplicationBunch, Collector, ENetTraceVerbosity::Trace);
+#endif
 			FastSharedReplicationChannel->WriteContentBlockPayload( TargetObj, *FastSharedReplicationBunch, false, TempBlockWriter );
 			
+			// Release temporary collector
+			UE_NET_TRACE_DESTROY_COLLECTOR(GetTraceCollector(TempBlockWriter));
+
 			FastSharedReplicationBunch = nullptr;
 			FastSharedReplicationChannel = nullptr;
 			FastSharedReplicationFuncName = NAME_None;
diff --git a/Engine/Plugins/Runtime/ReplicationSystemTestPlugin/Source/Private/NetworkAutomationTest.cpp b/Engine/Plugins/Runtime/ReplicationSystemTestPlugin/Source/Private/NetworkAutomationTest.cpp
new file mode 100644
--- /dev/null
+++ b/Engine/Plugins/Runtime/ReplicationSystemTestPlugin/Source/Private/NetworkAutomationTest.cpp
@@ -0,0 +1,175 @@
+// Copyright Epic Games, Inc. All Rights Reserved.
+
+#include "NetworkAutomationTest.h"
+#include "NetworkAutomationTestMacros.h"
+#include "ReplicationSystemTestPlugin/NetworkAutomationTest.h"
+
+namespace UE::Net
+{
+
+//
+bool FNetworkAutomationTestConfig::bVerboseLogging = false;
+
+struct FNetworkAutomationTestStats
+{
+	SIZE_T FailureCount;
+	SIZE_T WarningCount;
+};
+
+class FNetworkAutomationTestStatsSummary
+{
+public:
+	FNetworkAutomationTestStatsSummary();
+
+	void AddTestResult(const TCHAR* TestName, const FNetworkAutomationTestStats& Stats);
+
+private:
+	friend void PrintNetworkAutomationTestSummary();
+
+	TMap<const TCHAR*, FNetworkAutomationTestStats> TestsWithFailures;
+	TMap<const TCHAR*, FNetworkAutomationTestStats> TestsWithWarnings;
+	SIZE_T TestCount;
+	SIZE_T SuccessCount;
+};
+
+static FNetworkAutomationTestStatsSummary NetworkAutomationTestStatsSummary;
+
+// 
+FNetworkAutomationTestSuiteFixture::FNetworkAutomationTestSuiteFixture()
+: Stats(nullptr)
+, bSuppressWarningsFromSummary(false)
+{
+}
+
+FNetworkAutomationTestSuiteFixture::~FNetworkAutomationTestSuiteFixture()
+{
+}
+
+void FNetworkAutomationTestSuiteFixture::SetUp()
+{
+}
+
+void FNetworkAutomationTestSuiteFixture::TearDown()
+{
+}
+
+void FNetworkAutomationTestSuiteFixture::RunTest()
+{
+	if (FNetworkAutomationTestConfig::GetVerboseLogging())
+	{
+		UE_LOG(LogNetworkAutomationTest, Display, TEXT("Running TestCase %ls"), GetName());
+	}
+
+	FNetworkAutomationTestStats TempStats = {};
+	Stats = &TempStats;
+
+	SetUp();
+	RunTestImpl();
+	TearDown();
+
+	NetworkAutomationTestStatsSummary.AddTestResult(GetName(), *Stats);
+	Stats = nullptr;
+}
+
+void FNetworkAutomationTestSuiteFixture::SetSuppressWarningsFromSummary(bool bSuppress)
+{
+	bSuppressWarningsFromSummary = bSuppress;
+}
+
+void FNetworkAutomationTestSuiteFixture::AddTestFailure()
+{
+	++Stats->FailureCount;
+}
+
+void FNetworkAutomationTestSuiteFixture::AddTestWarning()
+{
+	Stats->WarningCount += !bSuppressWarningsFromSummary;
+}
+
+#if WITH_AUTOMATION_WORKER
+
+FNetworkAutomationTestWrapper::FNetworkAutomationTestWrapper(FNetworkAutomationTestSuiteFixture& InTestSuite, const TCHAR* Name)
+: FAutomationTestBase(Name, false /* bIsComplex */)
+, TestSuite(InTestSuite)
+{
+}
+
+uint32 FNetworkAutomationTestWrapper::GetTestFlags() const
+{
+	return EAutomationTestFlags::EngineFilter | EAutomationTestFlags::ClientContext | EAutomationTestFlags::ServerContext;
+}
+
+uint32 FNetworkAutomationTestWrapper::GetRequiredDeviceNum() const
+{
+	return 1;
+}
+
+FString FNetworkAutomationTestWrapper::GetBeautifiedTestName() const
+{
+	return TestName;
+}
+
+void FNetworkAutomationTestWrapper::GetTests(TArray<FString>& OutBeautifiedNames, TArray <FString>& OutTestCommands) const
+{
+	OutBeautifiedNames.Add(TestName);
+	OutTestCommands.Add(TestName);
+}
+
+bool FNetworkAutomationTestWrapper::RunTest(const FString&)
+{
+	TestSuite.RunTest();
+	return true;
+}
+
+#else
+
+#endif
+
+//
+FNetworkAutomationTestStatsSummary::FNetworkAutomationTestStatsSummary()
+: TestCount(0)
+, SuccessCount(0)
+{
+}
+
+void FNetworkAutomationTestStatsSummary::AddTestResult(const TCHAR* TestName, const FNetworkAutomationTestStats& Stats)
+{
+	++TestCount;
+	SuccessCount += (Stats.FailureCount == 0 ? 1U : 0U);
+
+	if (Stats.FailureCount)
+	{
+		TestsWithFailures.Add(TestName, Stats);
+	}
+	else if (Stats.WarningCount)
+	{
+		TestsWithWarnings.Add(TestName, Stats);
+	}
+}
+
+void PrintNetworkAutomationTestSummary()
+{
+	UE_LOG(LogNetworkAutomationTest, Display, TEXT("\n-----------------------------\nNetworkAutomationTest Summary\n-----------------------------\nTests executed: %zu\nTests succeeded: %zu\nTests failed: %zu\nTests with warnings: %zu\n-----------------------------"), NetworkAutomationTestStatsSummary.TestCount, NetworkAutomationTestStatsSummary.SuccessCount, SIZE_T(NetworkAutomationTestStatsSummary.TestsWithFailures.Num()), SIZE_T(NetworkAutomationTestStatsSummary.TestsWithWarnings.Num()));
+
+	if (NetworkAutomationTestStatsSummary.TestsWithFailures.Num())
+	{
+		for (auto It = NetworkAutomationTestStatsSummary.TestsWithFailures.CreateConstIterator(); It; ++It)
+		{
+			const TCHAR* Name = It.Key();
+			const FNetworkAutomationTestStats& Stats = It.Value();
+			UE_LOG(LogNetworkAutomationTest, Display, TEXT("Test %ls failed with %zu errors and %zu warnings"), Name, Stats.FailureCount, Stats.WarningCount);
+		}
+	}
+
+	if (NetworkAutomationTestStatsSummary.TestsWithWarnings.Num())
+	{
+		for (auto It = NetworkAutomationTestStatsSummary.TestsWithWarnings.CreateConstIterator(); It; ++It)
+		{
+			const TCHAR* Name = It.Key();
+			const FNetworkAutomationTestStats& Stats = It.Value();
+			UE_LOG(LogNetworkAutomationTest, Display, TEXT("Test %ls had %zu warnings"), Name, Stats.WarningCount);
+		}
+	}
+}
+
+}
diff --git a/Engine/Plugins/Runtime/ReplicationSystemTestPlugin/Source/Private/NetworkAutomationTestMacros.h b/Engine/Plugins/Runtime/ReplicationSystemTestPlugin/Source/Private/NetworkAutomationTestMacros.h
new file mode 100644
--- /dev/null
+++ b/Engine/Plugins/Runtime/ReplicationSystemTestPlugin/Source/Private/NetworkAutomationTestMacros.h
@@ -0,0 +1,142 @@
+// Copyright Epic Games, Inc. All Rights Reserved.
+
+#pragma once
+
+#include "CoreTypes.h"
+#include "TestMessage.h"
+#include "Logging/LogVerbosity.h"
+
+namespace UE::Net
+{
+	class FNetworkAutomationTestSuiteFixture;
+}
+
+namespace UE::Net
+{
+
+enum ETestResult
+{
+	ETestResult_Success,
+	ETestResult_Error,
+	ETestResult_FatalError,
+};
+
+class FTestResult
+{
+public:
+	template<typename T> FTestResult& operator<<(T Value) { Message << Value; return *this; }
+
+	operator bool() const { return TestResult == ETestResult_Success; }
+	bool operator!() const { return TestResult != ETestResult_Success; }
+
+	FTestMessage& GetMessage() { return Message; }
+
+protected:
+	FTestResult();
+
+	FTestMessage Message;
+	ETestResult TestResult;
+};
+
+class FTestMessageLog
+{
+public:
+	FTestMessageLog(FNetworkAutomationTestSuiteFixture& Test, ELogVerbosity::Type LogVerbosity);
+	~FTestMessageLog();
+
+	void operator=(const FTestMessage&);
+
+private:
+	FNetworkAutomationTestSuiteFixture& Test;
+	FTestMessage Message;
+	ELogVerbosity::Type LogVerbosity;
+};
+
+FTestResult CreateTestSuccess();
+FTestResult CreateTestFailure();
+
+inline FTestResult
+BooleanTestResult(bool Value, bool ExpectedValue, const char* ValueText)
+{
+	if (Value == ExpectedValue)
+	{
+		return CreateTestSuccess();
+	}
+	else
+	{
+		return CreateTestFailure() << ValueText <<  TEXT(" is ") << (Value ? TEXT("true") : TEXT("false")) << TEXT(". Expected ") << (ExpectedValue ? TEXT("true. ") : TEXT("false. "));
+	}
+}
+
+}
+
+#define UE_NET_COMPARE_RETURN_TESTRESULT(Name, Operator) \
+template<typename T> \
+UE::Net::FTestResult \
+Name(const T& Value, const T& ExpectedValue, const char* ValueText, const char* ExpectedValueText) \
+{ \
+	if (Value Operator ExpectedValue) \
+	{ \
+		return UE::Net::CreateTestSuccess(); \
+	} \
+	else \
+	{ \
+		return UE::Net::CreateTestFailure() << TEXT("Expected ") << ValueText << TEXT(' ') << TEXT(#Operator) <<  TEXT(' ') << ExpectedValueText << TEXT(". ") << Value << TEXT(" was compared to ") << ExpectedValue << TEXT(". "); \
+	} \
+} \
+template<typename T, typename U, typename TEnableIf<TIsPointer<T>::Value && TIsSame<U, nullptr_t>::Value, U>::Type X = nullptr> \
+UE::Net::FTestResult \
+Name(const T& Value, U ExpectedValue, const char* ValueText, const char* ExpectedValueText) \
+{ \
+	if (Value Operator ExpectedValue) \
+	{ \
+		return UE::Net::CreateTestSuccess(); \
+	} \
+	else \
+	{ \
+		return UE::Net::CreateTestFailure() << TEXT("Expected ") << ValueText << TEXT(' ') << TEXT(#Operator) <<  TEXT(' ') << ExpectedValueText << TEXT(". ") << Value << TEXT(" was compared to ") << ExpectedValue << TEXT(". "); \
+	} \
+}
+
+UE_NET_COMPARE_RETURN_TESTRESULT(TCmpEqual, ==)
+UE_NET_COMPARE_RETURN_TESTRESULT(TCmpNotEqual, !=)
+UE_NET_COMPARE_RETURN_TESTRESULT(TCmpLess, <)
+UE_NET_COMPARE_RETURN_TESTRESULT(TCmpLessOrEqual, <=)
+UE_NET_COMPARE_RETURN_TESTRESULT(TCmpGreater, >)
+UE_NET_COMPARE_RETURN_TESTRESULT(TCmpGreaterOrEqual, >=)
+
+#define UE_NET_FAIL(Message) return this->AddTestFailure(), UE::Net::FTestMessageLog(*this, ELogVerbosity::Error) = UE::Net::FTestMessage() << Message
+#define UE_NET_WARN(Message) this->AddTestWarning(), UE::Net::FTestMessageLog(*this, ELogVerbosity::Warning) = UE::Net::FTestMessage() << Message
+
+#define UE_NET_TEST_(CompareWithTestResult, V1, V2, FailFunction) if (UE::Net::FTestResult Result_ = CompareWithTestResult(V1, V2, #V1, #V2)) {} else FailFunction(Result_.GetMessage())
+#define UE_NET_ASSERT_INTERNAL(CompareWithTestResult, V1, V2) UE_NET_TEST_(CompareWithTestResult, V1, V2, UE_NET_FAIL)
+#define UE_NET_EXPECT_INTERNAL(CompareWithTestResult, V1, V2) UE_NET_TEST_(CompareWithTestResult, V1, V2, UE_NET_WARN)
+
+#define UE_NET_TEST_BOOL_(V1, V2, FailFunction) if (UE::Net::FTestResult Result_ = UE::Net::BooleanTestResult(V1, V2, #V1)) {} else FailFunction(Result_.GetMessage())
+#define UE_NET_ASSERT_BOOL_INTERNAL(V1, V2) UE_NET_TEST_BOOL_(V1, V2, UE_NET_FAIL)
+#define UE_NET_EXPECT_BOOL_INTERNAL(V1, V2) UE_NET_TEST_BOOL_(V1, V2, UE_NET_WARN)
+
+#define UE_NET_ASSERT_EQ(V1, V2) UE_NET_ASSERT_INTERNAL(TCmpEqual, V1, V2)
+#define UE_NET_ASSERT_NE(V1, V2) UE_NET_ASSERT_INTERNAL(TCmpNotEqual, V1, V2)
+#define UE_NET_ASSERT_LT(V1, V2) UE_NET_ASSERT_INTERNAL(TCmpLess, V1, V2)
+#define UE_NET_ASSERT_LE(V1, V2) UE_NET_ASSERT_INTERNAL(TCmpLessOrEqual, V1, V2)
+#define UE_NET_ASSERT_GT(V1, V2) UE_NET_ASSERT_INTERNAL(TCmpGreater, V1, V2)
+#define UE_NET_ASSERT_GE(V1, V2) UE_NET_ASSERT_INTERNAL(TCmpGreaterOrEqual, V1, V2)
+#define UE_NET_ASSERT_TRUE(V) UE_NET_ASSERT_BOOL_INTERNAL(V, true)
+#define UE_NET_ASSERT_FALSE(V) UE_NET_ASSERT_BOOL_INTERNAL(V, false)
+
+#define UE_NET_EXPECT_EQ(V1, V2) UE_NET_EXPECT_INTERNAL(TCmpEqual, V1, V2)
+#define UE_NET_EXPECT_NE(V1, V2) UE_NET_EXPECT_INTERNAL(TCmpNotEqual, V1, V2)
+#define UE_NET_EXPECT_LT(V1, V2) UE_NET_EXPECT_INTERNAL(TCmpLess, V1, V2)
+#define UE_NET_EXPECT_LE(V1, V2) UE_NET_EXPECT_INTERNAL(TCmpLessOrEqual, V1, V2)
+#define UE_NET_EXPECT_GT(V1, V2) UE_NET_EXPECT_INTERNAL(TCmpGreater, V1, V2)
+#define UE_NET_EXPECT_GE(V1, V2) UE_NET_EXPECT_INTERNAL(TCmpGreaterOrEqual, V1, V2)
+#define UE_NET_EXPECT_TRUE(V) UE_NET_EXPECT_BOOL_INTERNAL(V, true)
+#define UE_NET_EXPECT_FALSE(V) UE_NET_EXPECT_BOOL_INTERNAL(V, false)
+
+#define UE_NET_LOG(Message) FTestMessageLog(*this, ELogVerbosity::Display) = FTestMessage() << Message
+
+#define UE_NET_STRASSERT_EQ
+#define UE_NET_STRASSERT_NE
+#define UE_NET_STRCASEASSERT_EQ
+#define UE_NET_STRCASEASSERT_NE
diff --git a/Engine/Plugins/Runtime/ReplicationSystemTestPlugin/Source/Private/Tests/Core/TestBitTwiddling.cpp b/Engine/Plugins/Runtime/ReplicationSystemTestPlugin/Source/Private/Tests/Core/TestBitTwiddling.cpp
new file mode 100644
--- /dev/null
+++ b/Engine/Plugins/Runtime/ReplicationSystemTestPlugin/Source/Private/Tests/Core/TestBitTwiddling.cpp
@@ -0,0 +1,128 @@
+// Copyright Epic Games, Inc. All Rights Reserved.
+
+#include "NetworkAutomationTest.h"
+#include "NetworkAutomationTestMacros.h"
+#include "Iris/Core/BitTwiddling.h"
+#include "Math/NumericLimits.h"
+
+namespace UE::Net::Private
+{
+
+// GetBitsNeeded
+UE_NET_TEST(TestGetBitsNeeded, ReturnsExpectedValuesForInt8)
+{
+	UE_NET_ASSERT_EQ(GetBitsNeeded(TNumericLimits<int8>::Lowest()), 8U);
+	UE_NET_ASSERT_EQ(GetBitsNeeded(int8(-64)), 7U);
+	UE_NET_ASSERT_EQ(GetBitsNeeded(int8(-1)), 1U);
+	UE_NET_ASSERT_EQ(GetBitsNeeded(int8(0)), 1U);
+	UE_NET_ASSERT_EQ(GetBitsNeeded(int8(1)), 2U);
+	UE_NET_ASSERT_EQ(GetBitsNeeded(TNumericLimits<int8>::Max()), 8U);
+}
+
+UE_NET_TEST(TestGetBitsNeeded, ReturnsExpectedValuesForInt32)
+{
+	UE_NET_ASSERT_EQ(GetBitsNeeded(TNumericLimits<int32>::Lowest()), 32U);
+	UE_NET_ASSERT_EQ(GetBitsNeeded(int32(-64)), 7U);
+	UE_NET_ASSERT_EQ(GetBitsNeeded(int32(-1)), 1U);
+	UE_NET_ASSERT_EQ(GetBitsNeeded(int32(0)), 1U);
+	UE_NET_ASSERT_EQ(GetBitsNeeded(int32(1)), 2U);
+	UE_NET_ASSERT_EQ(GetBitsNeeded(TNumericLimits<int32>::Max()), 32U);
+}
+
+UE_NET_TEST(TestGetBitsNeeded, ReturnsExpectedValuesForInt64)
+{
+	UE_NET_ASSERT_EQ(GetBitsNeeded(TNumericLimits<int64>::Lowest()), 64U);
+	UE_NET_ASSERT_EQ(GetBitsNeeded(int64(-64)), 7U);
+	UE_NET_ASSERT_EQ(GetBitsNeeded(int64(-1)), 1U);
+	UE_NET_ASSERT_EQ(GetBitsNeeded(int64(0)), 1U);
+	UE_NET_ASSERT_EQ(GetBitsNeeded(int64(1)), 2U);
+	UE_NET_ASSERT_EQ(GetBitsNeeded(TNumericLimits<int64>::Max()), 64U);
+}
+
+UE_NET_TEST(TestGetBitsNeeded, ReturnsExpectedValuesForUint8)
+{
+	UE_NET_ASSERT_EQ(GetBitsNeeded(uint8(0)), 0U);
+	UE_NET_ASSERT_EQ(GetBitsNeeded(uint8(1)), 1U);
+	UE_NET_ASSERT_EQ(GetBitsNeeded(TNumericLimits<uint8>::Max()), 8U);
+}
+
+UE_NET_TEST(TestGetBitsNeeded, ReturnsExpectedValuesForUint32)
+{
+	UE_NET_ASSERT_EQ(GetBitsNeeded(uint32(0)), 0U);
+	UE_NET_ASSERT_EQ(GetBitsNeeded(uint32(1)), 1U);
+	UE_NET_ASSERT_EQ(GetBitsNeeded(TNumericLimits<uint32>::Max()), 32U);
+}
+
+UE_NET_TEST(TestGetBitsNeeded, ReturnsExpectedValuesForUint64)
+{
+	UE_NET_ASSERT_EQ(GetBitsNeeded(uint64(0)), 0U);
+	UE_NET_ASSERT_EQ(GetBitsNeeded(uint64(1)), 1U);
+	UE_NET_ASSERT_EQ(GetBitsNeeded(TNumericLimits<uint64>::Max()), 64U);
+}
+
+// GetBitsNeededForRange
+UE_NET_TEST(TestGetBitsNeededForRange, ReturnsExpectedValuesForInt8)
+{
+	UE_NET_ASSERT_EQ(GetBitsNeededForRange(int8(47), int8(47)), 0U);
+	UE_NET_ASSERT_EQ(GetBitsNeededForRange(MIN_int8, int8(-1)), 7U);
+	UE_NET_ASSERT_EQ(GetBitsNeededForRange(int8(MIN_int8 + 5), int8(4)), 7U);
+	UE_NET_ASSERT_EQ(GetBitsNeededForRange(MIN_int8, int8(0)), 8U);
+	UE_NET_ASSERT_EQ(GetBitsNeededForRange(int8(11), int8(47)), 6U);
+}
+
+UE_NET_TEST(TestGetBitsNeededForRange, ReturnsExpectedValuesForInt16)
+{
+	UE_NET_ASSERT_EQ(GetBitsNeededForRange(int16(4711), int16(4711)), 0U);
+	UE_NET_ASSERT_EQ(GetBitsNeededForRange(MIN_int16, int16(-1)), 15U);
+	UE_NET_ASSERT_EQ(GetBitsNeededForRange(int16(MIN_int16 + 5), int16(4)), 15U);
+	UE_NET_ASSERT_EQ(GetBitsNeededForRange(MIN_int16, int16(0)), 16U);
+	UE_NET_ASSERT_EQ(GetBitsNeededForRange(int16(11), int16(47)), 6U);
+}
+
+UE_NET_TEST(TestGetBitsNeededForRange, ReturnsExpectedValuesForInt32)
+{
+	UE_NET_ASSERT_EQ(GetBitsNeededForRange(int32(4711), int32(4711)), 0U);
+	UE_NET_ASSERT_EQ(GetBitsNeededForRange(MIN_int32, int32(-1)), 31U);
+	UE_NET_ASSERT_EQ(GetBitsNeededForRange(int32(MIN_int32 + 5), int32(4)), 31U);
+	UE_NET_ASSERT_EQ(GetBitsNeededForRange(MIN_int32, int32(0)), 32U);
+	UE_NET_ASSERT_EQ(GetBitsNeededForRange(int32(11), int32(47)), 6U);
+}
+
+UE_NET_TEST(TestGetBitsNeededForRange, ReturnsExpectedValuesForInt64)
+{
+	UE_NET_ASSERT_EQ(GetBitsNeededForRange(int64(4711), int64(4711)), 0U);
+	UE_NET_ASSERT_EQ(GetBitsNeededForRange(MIN_int64, int64(-1)), 63U);
+	UE_NET_ASSERT_EQ(GetBitsNeededForRange(int64(MIN_int64 + 5), int64(4)), 63U);
+	UE_NET_ASSERT_EQ(GetBitsNeededForRange(MIN_int64, int64(0)), 64U);
+	UE_NET_ASSERT_EQ(GetBitsNeededForRange(int64(11), int64(47)), 6U);
+}
+
+UE_NET_TEST(TestGetBitsNeededForRange, ReturnsExpectedValuesForUint8)
+{
+	UE_NET_ASSERT_EQ(GetBitsNeededForRange(uint8(47), uint8(47)), 0U);
+	UE_NET_ASSERT_EQ(GetBitsNeededForRange(uint8(MAX_uint8/2U + 1U), MAX_uint8), 7U);
+	UE_NET_ASSERT_EQ(GetBitsNeededForRange(MIN_uint8, MAX_uint8), 8U);
+}
+
+UE_NET_TEST(TestGetBitsNeededForRange, ReturnsExpectedValuesForUint16)
+{
+	UE_NET_ASSERT_EQ(GetBitsNeededForRange(uint16(4711), uint16(4711)), 0U);
+	UE_NET_ASSERT_EQ(GetBitsNeededForRange(uint16(MAX_uint16/2U + 1U), MAX_uint16), 15U);
+	UE_NET_ASSERT_EQ(GetBitsNeededForRange(MIN_uint16, MAX_uint16), 16U);
+}
+
+UE_NET_TEST(TestGetBitsNeededForRange, ReturnsExpectedValuesForUint32)
+{
+	UE_NET_ASSERT_EQ(GetBitsNeededForRange(uint32(4711), uint32(4711)), 0U);
+	UE_NET_ASSERT_EQ(GetBitsNeededForRange(uint32(MAX_uint32/2U + 1U), MAX_uint32), 31U);
+	UE_NET_ASSERT_EQ(GetBitsNeededForRange(MIN_uint32, MAX_uint32), 32U);
+}
+
+UE_NET_TEST(TestGetBitsNeededForRange, ReturnsExpectedValuesForUint64)
+{
+	UE_NET_ASSERT_EQ(GetBitsNeededForRange(uint64(4711), uint64(4711)), 0U);
+	UE_NET_ASSERT_EQ(GetBitsNeededForRange(uint64(MAX_uint64/2U + 1U), MAX_uint64), 63U);
+	UE_NET_ASSERT_EQ(GetBitsNeededForRange(MIN_uint64, MAX_uint64), 64U);
+}
+
+}
\ No newline at end of file
diff --git a/Engine/Plugins/Runtime/ReplicationSystemTestPlugin/Source/Private/Tests/ReplicationSystem/Filtering/MockNetObjectFilter.cpp b/Engine/Plugins/Runtime/ReplicationSystemTestPlugin/Source/Private/Tests/ReplicationSystem/Filtering/MockNetObjectFilter.cpp
new file mode 100644
--- /dev/null
+++ b/Engine/Plugins/Runtime/ReplicationSystemTestPlugin/Source/Private/Tests/ReplicationSystem/Filtering/MockNetObjectFilter.cpp
@@ -0,0 +1,157 @@
+// Copyright Epic Games, Inc. All Rights Reserved.
+
+#include "MockNetObjectFilter.h"
+#include "Containers/ArrayView.h"
+
+UMockNetObjectFilter::UMockNetObjectFilter()
+: CallStatus({})
+, AddedCount(0U)
+{
+}
+
+void UMockNetObjectFilter::Init(FNetObjectFilterInitParams& Params)
+{
+	++CallStatus.CallCounts.Init;
+
+	CallStatus.SuccessfulCallCounts.Init += Cast<UMockNetObjectFilterConfig>(Params.Config) != nullptr;
+
+	AddedObjectIndices.Init(Params.MaxObjectCount);
+	AddedConnectionIndices.Init(Params.MaxConnectionCount + 1);
+}
+
+void UMockNetObjectFilter::AddConnection(uint32 ConnectionId)
+{
+	++CallStatus.CallCounts.AddConnection;
+
+	CallStatus.SuccessfulCallCounts.AddConnection += ConnectionId != 0U;
+
+	AddedConnectionIndices.SetBit(ConnectionId);
+}
+
+void UMockNetObjectFilter::RemoveConnection(uint32 ConnectionId)
+{
+	++CallStatus.CallCounts.RemoveConnection;
+
+	CallStatus.SuccessfulCallCounts.RemoveConnection += AddedConnectionIndices.GetBit(ConnectionId);
+	AddedConnectionIndices.ClearBit(ConnectionId);
+}
+
+bool UMockNetObjectFilter::AddObject(uint32 ObjectIndex, FNetObjectFilterAddObjectParams& Params)
+{
+	++CallStatus.CallCounts.AddObject;
+	if (CallSetup.AddObject.bReturnValue)
+	{
+		bool bIsProperCall = true;
+
+		// Validate the passed info is zeroed
+		bIsProperCall &= Params.OutInfo.Data[0] == 0;
+		bIsProperCall &= Params.OutInfo.Data[1] == 0;
+		bIsProperCall &= Params.OutInfo.Data[2] == 0;
+		bIsProperCall &= Params.OutInfo.Data[3] == 0;
+		
+		CallStatus.SuccessfulCallCounts.AddObject += bIsProperCall;
+
+		// Store object index in info so we can check it's passed to RemoveObject
+		Params.OutInfo.Data[0] = uint16(ObjectIndex);
+		Params.OutInfo.Data[1] = uint16(ObjectIndex);
+		Params.OutInfo.Data[2] = uint16(ObjectIndex);
+		Params.OutInfo.Data[3] = uint16(ObjectIndex);
+
+		AddedObjectIndices.SetBit(ObjectIndex);
+		++AddedCount;
+
+		// Check if object has the NetTest_FilterOut RepTag which says whether the object should be filtered out or not.
+		UE::Net::FRepTagFindInfo RepTagInfo;
+		if (UE::Net::FindRepTag(Params.Protocol, RepTag_NetTest_FilterOut, RepTagInfo))
+		{
+			// Warning: Using the internal state requires the proper NetSerializer and dequantization to get the value.
+			ObjectToFilterOutOffset.Add(ObjectIndex, RepTagInfo.InternalStateAbsoluteOffset);
+		}
+	}
+
+	return CallSetup.AddObject.bReturnValue;
+}
+
+void UMockNetObjectFilter::RemoveObject(uint32 ObjectIndex, const FNetObjectFilteringInfo& Info)
+{
+	++CallStatus.CallCounts.RemoveObject;
+
+	// If this object wasn't added to us it shouldn't be removed either.
+	if (AddedObjectIndices.GetBit(ObjectIndex))
+	{
+		bool bIsProperCall = true;
+
+		// Validate the passed info is as we set it
+		bIsProperCall &= Info.Data[0] == uint16(ObjectIndex);
+		bIsProperCall &= Info.Data[1] == uint16(ObjectIndex);
+		bIsProperCall &= Info.Data[2] == uint16(ObjectIndex);
+		bIsProperCall &= Info.Data[3] == uint16(ObjectIndex);
+
+		CallStatus.SuccessfulCallCounts.RemoveObject += bIsProperCall;
+
+		AddedObjectIndices.ClearBit(ObjectIndex);
+		--AddedCount;
+
+		ObjectToFilterOut.Remove(ObjectIndex);
+		ObjectToFilterOutOffset.Remove(ObjectIndex);
+	}
+}
+
+void UMockNetObjectFilter::UpdateObjects(FNetObjectFilterUpdateParams& Params)
+{
+	++CallStatus.CallCounts.UpdateObjects;
+
+	bool bIsProperCall = Params.ObjectCount <= AddedCount;
+	for (uint32 ObjectIndex : MakeArrayView(Params.ObjectIndices, Params.ObjectCount))
+	{
+		bIsProperCall = bIsProperCall && AddedObjectIndices.GetBit(ObjectIndex);
+
+		if (UPTRINT* InternalStateOffset = ObjectToFilterOutOffset.Find(ObjectIndex))
+		{
+			// Warning: This is generally not safe, not even for primitive types.
+			const bool bFilterOut = *reinterpret_cast<const bool*>(Params.StateBuffers[ObjectIndex] + *InternalStateOffset);
+			ObjectToFilterOut.Emplace(ObjectIndex, bFilterOut);
+		}
+	}
+	CallStatus.SuccessfulCallCounts.UpdateObjects += bIsProperCall;
+}
+
+void UMockNetObjectFilter::PreFilter(FNetObjectPreFilteringParams& Params)
+{
+	++CallStatus.CallCounts.PreFilter;
+
+	constexpr bool bIsProperCall = true;
+	CallStatus.SuccessfulCallCounts.PreFilter += bIsProperCall;
+}
+
+void UMockNetObjectFilter::Filter(FNetObjectFilteringParams& Params)
+{
+	++CallStatus.CallCounts.Filter;
+
+	bool bIsProperCall = true;
+	UE::Net::FNetBitArrayView::ForAllExclusiveBits(Params.FilteredObjects, MakeNetBitArrayView(AddedObjectIndices), [](...) {}, [&bIsProperCall](...) {bIsProperCall = false; });
+	CallStatus.SuccessfulCallCounts.Filter += bIsProperCall;
+
+	if (CallSetup.Filter.bFilterOutByDefault)
+	{
+		Params.OutAllowedObjects.Reset();
+	}
+	else
+	{
+		Params.OutAllowedObjects.Copy(MakeNetBitArrayView(AddedObjectIndices));
+	}
+
+	// Go through objects with FilterOut wishes
+	for (const auto& Pair : ObjectToFilterOut)
+	{
+		Params.OutAllowedObjects.SetBitValue(Pair.Key, !Pair.Value);
+	}
+}
+
+void UMockNetObjectFilter::PostFilter(FNetObjectPostFilteringParams& Params)
+{
+	++CallStatus.CallCounts.PostFilter;
+
+	constexpr bool bIsProperCall = true;
+	CallStatus.SuccessfulCallCounts.PostFilter += bIsProperCall;
+}
diff --git a/Engine/Plugins/Runtime/ReplicationSystemTestPlugin/Source/Private/Tests/ReplicationSystem/Filtering/TestFiltering.cpp b/Engine/Plugins/Runtime/ReplicationSystemTestPlugin/Source/Private/Tests/ReplicationSystem/Filtering/TestFiltering.cpp
new file mode 100644
--- /dev/null
+++ b/Engine/Plugins/Runtime/ReplicationSystemTestPlugin/Source/Private/Tests/ReplicationSystem/Filtering/TestFiltering.cpp
@@ -0,0 +1,1257 @@
+#include "NetworkAutomationTest.h"
+#include "NetworkAutomationTestMacros.h"
+#include "Iris/ReplicationSystem/ReplicationSystem.h"
+#include "Iris/ReplicationSystem/Filtering/NetObjectFilter.h"
+#include "Iris/ReplicationSystem/Filtering/NetObjectFilterDefinitions.h"
+#include "Tests/ReplicationSystem/Filtering/MockNetObjectFilter.h"
+#include "Tests/ReplicationSystem/Filtering/TestFilteringObject.h"
+#include "Tests/ReplicationSystem/ReplicationSystemServerClientTestFixture.h"
+
+namespace UE::Net::Private
+{
+
+class FTestFilteringFixture : public FReplicationSystemServerClientTestFixture
+{
+	virtual void SetUp() override
+	{
+		InitNetObjectFilterDefinitions();
+		FReplicationSystemServerClientTestFixture::SetUp();
+		InitMockNetObjectFilter();
+	}
+
+	virtual void TearDown() override
+	{
+		FReplicationSystemServerClientTestFixture::TearDown();
+		RestoreNetObjectPrioritizerDefinitions();
+	}
+
+	void InitNetObjectFilterDefinitions()
+	{
+		const UClass* NetObjectFilterDefinitionsClass = UNetObjectFilterDefinitions::StaticClass();
+		const FProperty* DefinitionsProperty = NetObjectFilterDefinitionsClass->FindPropertyByName("NetObjectFilterDefinitions");
+		check(DefinitionsProperty != nullptr);
+
+		// Save NetObjectFilterDefinitions CDO state.
+		UNetObjectFilterDefinitions* FilterDefinitions = GetMutableDefault<UNetObjectFilterDefinitions>();
+		DefinitionsProperty->CopyCompleteValue(&OriginalFilterDefinitions, (void*)(UPTRINT(FilterDefinitions) + DefinitionsProperty->GetOffset_ForInternal()));
+
+		// Modify definitions to only include our mock prioritizer. Ugly... 
+		TArray<FNetObjectFilterDefinition> NewFilterDefinitions;
+		FNetObjectFilterDefinition& MockDefinition = NewFilterDefinitions.Emplace_GetRef();
+		MockDefinition.FilterName = "MockFilter";
+		MockDefinition.ClassName = "/Script/ReplicationSystemTestPlugin.MockNetObjectFilter";
+		MockDefinition.ConfigClassName = "/Script/ReplicationSystemTestPlugin.MockNetObjectFilterConfig";
+		DefinitionsProperty->CopyCompleteValue((void*)(UPTRINT(FilterDefinitions) + DefinitionsProperty->GetOffset_ForInternal()), &NewFilterDefinitions);
+	}
+
+	void RestoreNetObjectPrioritizerDefinitions()
+	{
+		// Restore NetObjectPrioritizerDefinitions CDO state from the saved state.
+		const UClass* NetObjectFilterDefinitionsClass = UNetObjectFilterDefinitions::StaticClass();
+		const FProperty* DefinitionsProperty = NetObjectFilterDefinitionsClass->FindPropertyByName("NetObjectFilterDefinitions");
+		UNetObjectFilterDefinitions* FilterDefinitions = GetMutableDefault<UNetObjectFilterDefinitions>();
+		DefinitionsProperty->CopyCompleteValue((void*)(UPTRINT(FilterDefinitions) + DefinitionsProperty->GetOffset_ForInternal()), &OriginalFilterDefinitions);
+		OriginalFilterDefinitions.Empty();
+
+		MockFilterHandle = InvalidNetObjectFilterHandle;
+		MockNetObjectFilter = nullptr;
+	}
+
+	void InitMockNetObjectFilter()
+	{
+		MockNetObjectFilter = ExactCast<UMockNetObjectFilter>(Server->GetReplicationSystem()->GetFilter("MockFilter"));
+		MockFilterHandle = Server->GetReplicationSystem()->GetFilterHandle("MockFilter");
+	}
+
+protected:
+	UMockNetObjectFilter* MockNetObjectFilter;
+	FNetObjectFilterHandle MockFilterHandle;
+
+private:
+	TArray<FNetObjectFilterDefinition> OriginalFilterDefinitions;
+};
+
+// Owner filtering tests
+UE_NET_TEST_FIXTURE(FTestFilteringFixture, OwnerFilterPreventsObjectFromReplicatingToNonOwner)
+{
+	// Add client
+	FReplicationSystemTestClient* Client = CreateClient();
+
+	// Spawn object on server
+	UReplicatedTestObject* ServerObject = Server->CreateObject(0, 0);
+
+	// Apply owner filter
+	Server->ReplicationSystem->SetFilter(ServerObject->NetHandle, ToOwnerFilterHandle);
+
+	// Send and deliver packet
+	Server->PreSendUpdate();
+	Server->SendAndDeliverTo(Client, DeliverPacket);
+	Server->PostSendUpdate();
+
+	// Object should not have been created on the client
+	UE_NET_ASSERT_EQ(Client->GetReplicationBridge()->GetReplicatedObject(ServerObject->NetHandle), nullptr);
+
+	Server->DestroyObject(ServerObject);
+}
+
+UE_NET_TEST_FIXTURE(FTestFilteringFixture, OwnerFilterAllowsObjectToReplicateToOwner)
+{
+	// Add client
+	FReplicationSystemTestClient* Client = CreateClient();
+
+	// Spawn object on server
+	UReplicatedTestObject* ServerObject = Server->CreateObject(0, 0);
+	Server->ReplicationSystem->SetFilter(ServerObject->NetHandle, ToOwnerFilterHandle);
+	Server->ReplicationSystem->SetOwningNetConnection(ServerObject->NetHandle, Client->ConnectionIdOnServer);
+
+	// Send and deliver packet
+	Server->PreSendUpdate();
+	Server->SendAndDeliverTo(Client, DeliverPacket);
+	Server->PostSendUpdate();
+
+	// Object should have been created on the client
+	UE_NET_ASSERT_NE(Client->GetReplicationBridge()->GetReplicatedObject(ServerObject->NetHandle), nullptr);
+
+	// Destroy object on server and client
+	Server->DestroyObject(ServerObject);
+
+	Server->PreSendUpdate();
+	Server->SendAndDeliverTo(Client, DeliverPacket);
+	Server->PostSendUpdate();
+}
+
+UE_NET_TEST_FIXTURE(FTestFilteringFixture, OwnerFilterReplicatesOnlyToOwningConnection)
+{
+	// Add clients
+	FReplicationSystemTestClient* ClientArray[] = {CreateClient(), CreateClient(), CreateClient()};
+	constexpr SIZE_T LastClientIndex = UE_ARRAY_COUNT(ClientArray) - 1;
+
+	// Spawn object on server
+	UReplicatedTestObject* ServerObject = Server->CreateObject(0, 0);
+	Server->ReplicationSystem->SetFilter(ServerObject->NetHandle, ToOwnerFilterHandle);
+	Server->ReplicationSystem->SetOwningNetConnection(ServerObject->NetHandle, ClientArray[LastClientIndex]->ConnectionIdOnServer);
+
+	// Send and deliver packets
+	Server->PreSendUpdate();
+	for (FReplicationSystemTestClient* Client : ClientArray)
+	{
+		Server->SendAndDeliverTo(Client, DeliverPacket);
+	}
+	Server->PostSendUpdate();
+
+	// Object should only have been created on the last client
+	for (FReplicationSystemTestClient*& Client : ClientArray)
+	{
+		const SIZE_T ClientIndex = &Client - &ClientArray[0];
+		if (ClientIndex == LastClientIndex)
+		{
+			UE_NET_ASSERT_NE(Client->GetReplicationBridge()->GetReplicatedObject(ServerObject->NetHandle), nullptr);
+		}
+		else
+		{
+			UE_NET_ASSERT_EQ(Client->GetReplicationBridge()->GetReplicatedObject(ServerObject->NetHandle), nullptr);
+		}
+	}
+
+	Server->DestroyObject(ServerObject);
+	Server->PreSendUpdate();
+	for (FReplicationSystemTestClient* Client : ClientArray)
+	{
+		Server->SendAndDeliverTo(Client, DeliverPacket);
+	}
+	Server->PostSendUpdate();
+}
+
+UE_NET_TEST_FIXTURE(FTestFilteringFixture, CanChangeOwningConnection)
+{
+	// Add client
+	FReplicationSystemTestClient* Client = CreateClient();
+
+	// Spawn object on server
+	UReplicatedTestObject* ServerObject = Server->CreateObject(0, 0);
+
+	// Send and deliver packets
+	Server->PreSendUpdate();
+	Server->SendAndDeliverTo(Client, DeliverPacket);
+	Server->PostSendUpdate();
+
+	// Object should now exist on client
+	UE_NET_ASSERT_NE(Client->GetReplicationBridge()->GetReplicatedObject(ServerObject->NetHandle), nullptr);
+
+	// Turn on owner filter
+	Server->ReplicationSystem->SetFilter(ServerObject->NetHandle, ToOwnerFilterHandle);
+
+	// Send and deliver packets
+	Server->PreSendUpdate();
+	Server->SendAndDeliverTo(Client, DeliverPacket);
+	Server->PostSendUpdate();
+
+	// As object is now filtered it should be deleted on the client
+	UE_NET_ASSERT_EQ(Client->GetReplicationBridge()->GetReplicatedObject(ServerObject->NetHandle), nullptr);
+
+	// Set the client as owner
+	Server->ReplicationSystem->SetOwningNetConnection(ServerObject->NetHandle, Client->ConnectionIdOnServer);
+
+	// Send and deliver packets
+	Server->PreSendUpdate();
+	Server->SendAndDeliverTo(Client, DeliverPacket);
+	Server->PostSendUpdate();
+
+	// The object should have been created again
+	UE_NET_ASSERT_NE(Client->GetReplicationBridge()->GetReplicatedObject(ServerObject->NetHandle), nullptr);
+
+	// Finally, remove the owning connection
+	Server->ReplicationSystem->SetOwningNetConnection(ServerObject->NetHandle, 0);
+
+	// Send and deliver packets
+	Server->PreSendUpdate();
+	Server->SendAndDeliverTo(Client, DeliverPacket);
+	Server->PostSendUpdate();
+
+	// The client is no longer owning the object so it should be deleted
+	UE_NET_ASSERT_EQ(Client->GetReplicationBridge()->GetReplicatedObject(ServerObject->NetHandle), nullptr);
+
+	Server->DestroyObject(ServerObject);
+}
+
+UE_NET_TEST_FIXTURE(FTestFilteringFixture, CanChangeOwnerFilter)
+{
+	// Add client
+	FReplicationSystemTestClient* Client = CreateClient();
+
+	// Spawn object on server
+	UReplicatedTestObject* ServerObject = Server->CreateObject(0, 0);
+
+	// Send and deliver packets
+	Server->PreSendUpdate();
+	Server->SendAndDeliverTo(Client, DeliverPacket);
+	Server->PostSendUpdate();
+
+	// Object should now exist on client
+	UE_NET_ASSERT_NE(Client->GetReplicationBridge()->GetReplicatedObject(ServerObject->NetHandle), nullptr);
+
+	// Turn on owner filter
+	Server->ReplicationSystem->SetFilter(ServerObject->NetHandle, ToOwnerFilterHandle);
+
+	// Send and deliver packets
+	Server->PreSendUpdate();
+	Server->SendAndDeliverTo(Client, DeliverPacket);
+	Server->PostSendUpdate();
+
+	// As object is now filtered it should be deleted on the client
+	UE_NET_ASSERT_EQ(Client->GetReplicationBridge()->GetReplicatedObject(ServerObject->NetHandle), nullptr);
+
+	// Remove the owner filter
+	Server->ReplicationSystem->SetFilter(ServerObject->NetHandle, InvalidNetObjectFilterHandle);
+
+	// Send and deliver packets
+	Server->PreSendUpdate();
+	Server->SendAndDeliverTo(Client, DeliverPacket);
+	Server->PostSendUpdate();
+
+	// The object should have been created again
+	UE_NET_ASSERT_NE(Client->GetReplicationBridge()->GetReplicatedObject(ServerObject->NetHandle), nullptr);
+
+	// Destroy the object
+	Server->DestroyObject(ServerObject);
+
+	// Send and deliver packets
+	Server->PreSendUpdate();
+	Server->SendAndDeliverTo(Client, DeliverPacket);
+	Server->PostSendUpdate();
+}
+
+UE_NET_TEST_FIXTURE(FTestFilteringFixture, LateAddedSubObjectGetsOwnerPropagated)
+{
+	// Add client
+	FReplicationSystemTestClient* Client = CreateClient();
+
+	// Spawn objects on server
+	TArray<UReplicatedTestObject*> ServerObjects;
+	{
+		constexpr SIZE_T ObjectCount = 64;
+		ServerObjects.Reserve(ObjectCount);
+		for (SIZE_T It = 0, EndIt = ObjectCount; It != EndIt; ++It)
+		{
+			UReplicatedTestObject* ServerObject = Server->CreateObject(0, 0);
+			ServerObjects.Add(ServerObject);
+			Server->ReplicationSystem->SetOwningNetConnection(ServerObject->NetHandle, Client->ConnectionIdOnServer);
+		}
+	}
+
+	// Net update
+	Server->PreSendUpdate();
+	Server->PostSendUpdate();
+
+	// Create subobject to arbitrary object
+	UReplicatedTestObject* ArbitraryServerObject = ServerObjects[3];
+	UReplicatedTestObject* ServerSubObject = Server->CreateSubObject(ArbitraryServerObject->NetHandle, 1, 1);
+
+	// Net update
+	Server->PreSendUpdate();
+	Server->PostSendUpdate();
+
+	// Verify subobject owner is as expected
+	UE_NET_ASSERT_EQ(Server->ReplicationSystem->GetOwningNetConnection(ServerSubObject->NetHandle), Client->ConnectionIdOnServer);
+}
+
+// Connection filtering tests
+UE_NET_TEST_FIXTURE(FTestFilteringFixture, ConnectionFilterPreventsObjectFromReplicatingToFilteredOutConnections)
+{
+	// Add client
+	FReplicationSystemTestClient* Client = CreateClient();
+
+	// Spawn object on server
+	UReplicatedTestObject* ServerObject = Server->CreateObject(0, 0);
+
+	// Apply connection filter
+	TBitArray<> NoConnections;
+	const bool bFilterWasApplied = Server->ReplicationSystem->SetConnectionFilter(ServerObject->NetHandle, NoConnections, ENetFilterStatus::Allow);
+	UE_NET_ASSERT_TRUE(bFilterWasApplied);
+
+	// Send and deliver packet
+	Server->PreSendUpdate();
+	Server->SendAndDeliverTo(Client, DeliverPacket);
+	Server->PostSendUpdate();
+
+	// Object should not have been created on the client
+	UE_NET_ASSERT_EQ(Client->GetReplicationBridge()->GetReplicatedObject(ServerObject->NetHandle), nullptr);
+}
+
+UE_NET_TEST_FIXTURE(FTestFilteringFixture, ConnectionFilterAllowsObjectToReplicateToAllowedConnections)
+{
+	// Add client
+	FReplicationSystemTestClient* Client = CreateClient();
+
+	// Spawn object on server
+	UReplicatedTestObject* ServerObject = Server->CreateObject(0, 0);
+
+	// Apply connection filter
+	TBitArray<> AllowedConnections;
+	AllowedConnections.Init(false, Client->ConnectionIdOnServer + 1);
+	AllowedConnections[Client->ConnectionIdOnServer] = true;
+	Server->ReplicationSystem->SetConnectionFilter(ServerObject->NetHandle, AllowedConnections, ENetFilterStatus::Allow);
+
+	// Send and deliver packet
+	Server->PreSendUpdate();
+	Server->SendAndDeliverTo(Client, DeliverPacket);
+	Server->PostSendUpdate();
+
+	// Object should have been created on the client
+	UE_NET_ASSERT_NE(Client->GetReplicationBridge()->GetReplicatedObject(ServerObject->NetHandle), nullptr);
+}
+
+UE_NET_TEST_FIXTURE(FTestFilteringFixture, ConnectionFilterAllowsObjectToReplicateToLateJoiningConnections)
+{
+	// Spawn object on server
+	UReplicatedTestObject* ServerObject = Server->CreateObject(0, 0);
+
+	// Apply filtering that allows all connections
+	TBitArray<> NoConnections;
+	constexpr ENetFilterStatus ReplicationStatus = ENetFilterStatus::Disallow;
+	Server->ReplicationSystem->SetConnectionFilter(ServerObject->NetHandle, NoConnections, ReplicationStatus);
+
+	// Send and deliver packets
+	Server->PreSendUpdate();
+	Server->PostSendUpdate();
+
+	// Late join clients
+	FReplicationSystemTestClient* ClientArray[] = {CreateClient(), CreateClient(), CreateClient()};
+
+	// Send and deliver packets
+	Server->PreSendUpdate();
+	for (FReplicationSystemTestClient* Client : ClientArray)
+	{
+		Server->SendAndDeliverTo(Client, DeliverPacket);
+	}
+	Server->PostSendUpdate();
+
+	// Object should have been created on all clients
+	for (FReplicationSystemTestClient*& Client : ClientArray)
+	{
+		UE_NET_ASSERT_NE(Client->GetReplicationBridge()->GetReplicatedObject(ServerObject->NetHandle), nullptr);
+	}
+}
+
+UE_NET_TEST_FIXTURE(FTestFilteringFixture, ConnectionFilterAllowsObjectToReplicateAndDoesNotAffectExistingGroupFilteredObjects)
+{
+	// Add client
+	FReplicationSystemTestClient* Client = CreateClient();
+
+	// Spawn object on server
+	UReplicatedTestObject* ServerObject0 = Server->CreateObject(0, 0);
+	
+	// Setup group filter
+	FNetObjectGroupHandle GroupHandle = Server->ReplicationSystem->CreateGroup();
+	Server->ReplicationSystem->AddToGroup(GroupHandle, ServerObject0->NetHandle);
+	Server->ReplicationSystem->AddGroupFilter(GroupHandle);
+
+	// Spawn object on server
+	UReplicatedTestObject* ServerObject = Server->CreateObject(0, 0);
+	UReplicatedTestObject* ServerSubObject = Server->CreateSubObject(ServerObject->NetHandle, 0, 0);
+
+	// Apply filtering that allows all connections
+	TBitArray<> NoConnections;
+	constexpr ENetFilterStatus ReplicationStatus = ENetFilterStatus::Disallow;
+	Server->ReplicationSystem->SetConnectionFilter(ServerObject->NetHandle, NoConnections, ReplicationStatus);
+
+	// Send and deliver packets
+	Server->PreSendUpdate();
+	Server->PostSendUpdate();
+
+	// Send and deliver packets
+	Server->PreSendUpdate();
+	Server->SendAndDeliverTo(Client, DeliverPacket);
+	Server->PostSendUpdate();
+
+	// Validate status on clients
+	// Object should not have been created on the clients 
+	UE_NET_ASSERT_EQ(Client->GetReplicationBridge()->GetReplicatedObject(ServerObject0->NetHandle), nullptr);
+		
+	// Object should have been crated on all clients
+	UE_NET_ASSERT_NE(Client->GetReplicationBridge()->GetReplicatedObject(ServerObject->NetHandle), nullptr);
+	UE_NET_ASSERT_NE(Client->GetReplicationBridge()->GetReplicatedObject(ServerSubObject->NetHandle), nullptr);
+}
+
+UE_NET_TEST_FIXTURE(FTestFilteringFixture, ConnectionFilterPreventsObjectFromReplicatingToFilteredOutLateJoiningConnections)
+{
+	// Spawn object on server
+	UReplicatedTestObject* ServerObject = Server->CreateObject(0, 0);
+
+	// Apply filtering that allows all connections
+	TBitArray<> Connection1(false, 2);
+	Connection1[1] = true;
+	constexpr ENetFilterStatus ReplicationStatus = ENetFilterStatus::Allow;
+	Server->ReplicationSystem->SetConnectionFilter(ServerObject->NetHandle, Connection1, ReplicationStatus);
+
+	// Send and deliver packets
+	Server->PreSendUpdate();
+	Server->PostSendUpdate();
+
+	// Late join clients
+	FReplicationSystemTestClient* ClientArray[] = {CreateClient(), CreateClient(), CreateClient()};
+
+	// Since we must set up the filtering before the clients are created we need to make sure our assumptions are valid.
+	UE_NET_ASSERT_EQ(ClientArray[0]->ConnectionIdOnServer, 1U);
+
+	// Send and deliver packets
+	Server->PreSendUpdate();
+	for (FReplicationSystemTestClient* Client : ClientArray)
+	{
+		Server->SendAndDeliverTo(Client, DeliverPacket);
+	}
+	Server->PostSendUpdate();
+
+	// Object should have been created on the client with connection ID 1.
+	for (FReplicationSystemTestClient*& Client : ClientArray)
+	{
+		const uint32 ConnectionId = Client->ConnectionIdOnServer;
+		const bool bShouldHaveObject = ConnectionId == 1;
+		if (bShouldHaveObject)
+		{
+			UE_NET_ASSERT_NE(Client->GetReplicationBridge()->GetReplicatedObject(ServerObject->NetHandle), nullptr);
+		}
+		else
+		{
+			UE_NET_ASSERT_EQ(Client->GetReplicationBridge()->GetReplicatedObject(ServerObject->NetHandle), nullptr);
+		}
+	}
+}
+
+UE_NET_TEST_FIXTURE(FTestFilteringFixture, ConnectionFilterReplicatesOnlyToAllowedConnection)
+{
+	// Add clients
+	FReplicationSystemTestClient* ClientArray[] = {CreateClient(), CreateClient(), CreateClient()};
+	constexpr SIZE_T LastClientIndex = UE_ARRAY_COUNT(ClientArray) - 1;
+
+	// Spawn object on server
+	UReplicatedTestObject* ServerObject = Server->CreateObject(0, 0);
+
+	// Apply filtering that allows the last client to receive the object
+	const uint32 ConnectionIdForLastClient = ClientArray[LastClientIndex]->ConnectionIdOnServer;
+	TBitArray<> AllowedConnections;
+	AllowedConnections.Init(false, ConnectionIdForLastClient + 1);
+	AllowedConnections[ConnectionIdForLastClient] = true;
+	Server->ReplicationSystem->SetConnectionFilter(ServerObject->NetHandle, AllowedConnections, ENetFilterStatus::Allow);
+
+	// Send and deliver packets
+	Server->PreSendUpdate();
+	for (FReplicationSystemTestClient* Client : ClientArray)
+	{
+		Server->SendAndDeliverTo(Client, DeliverPacket);
+	}
+	Server->PostSendUpdate();
+
+	// Object should only have been created on the last client
+	for (FReplicationSystemTestClient*& Client : ClientArray)
+	{
+		const SIZE_T ClientIndex = &Client - &ClientArray[0];
+		if (ClientIndex == LastClientIndex)
+		{
+			UE_NET_ASSERT_NE(Client->GetReplicationBridge()->GetReplicatedObject(ServerObject->NetHandle), nullptr);
+		}
+		else
+		{
+			UE_NET_ASSERT_EQ(Client->GetReplicationBridge()->GetReplicatedObject(ServerObject->NetHandle), nullptr);
+		}
+	}
+}
+
+UE_NET_TEST_FIXTURE(FTestFilteringFixture, CanChangeConnectionFilter)
+{
+	// Add client
+	FReplicationSystemTestClient* Client = CreateClient();
+
+	// Spawn object on server
+	UReplicatedTestObject* ServerObject = Server->CreateObject(0, 0);
+
+	// Send and deliver packets
+	Server->PreSendUpdate();
+	Server->SendAndDeliverTo(Client, DeliverPacket);
+	Server->PostSendUpdate();
+
+	// Object should now exist on client
+	UE_NET_ASSERT_NE(Client->GetReplicationBridge()->GetReplicatedObject(ServerObject->NetHandle), nullptr);
+
+	// Apply connection filtering, not allowing any connection to receive the object
+	TBitArray<> NoConnections;
+	Server->ReplicationSystem->SetConnectionFilter(ServerObject->NetHandle, NoConnections, ENetFilterStatus::Allow);
+
+	// Send and deliver packets
+	Server->PreSendUpdate();
+	Server->SendAndDeliverTo(Client, DeliverPacket);
+	Server->PostSendUpdate();
+
+	// As the object is now filtered it should be deleted on the client
+	UE_NET_ASSERT_EQ(Client->GetReplicationBridge()->GetReplicatedObject(ServerObject->NetHandle), nullptr);
+
+	// Set the client as an allowed connection
+	TBitArray<> AllowedConnections;
+	AllowedConnections.Init(false, Client->ConnectionIdOnServer + 1);
+	AllowedConnections[Client->ConnectionIdOnServer] = true;
+	Server->ReplicationSystem->SetConnectionFilter(ServerObject->NetHandle, AllowedConnections, ENetFilterStatus::Allow);
+
+	// Send and deliver packets
+	Server->PreSendUpdate();
+	Server->SendAndDeliverTo(Client, DeliverPacket);
+	Server->PostSendUpdate();
+
+	// The object should have been created again
+	UE_NET_ASSERT_NE(Client->GetReplicationBridge()->GetReplicatedObject(ServerObject->NetHandle), nullptr);
+
+	// Finally, set the filter to not include any connections again.
+	Server->ReplicationSystem->SetConnectionFilter(ServerObject->NetHandle, NoConnections, ENetFilterStatus::Allow);
+
+	// Send and deliver packets
+	Server->PreSendUpdate();
+	Server->SendAndDeliverTo(Client, DeliverPacket);
+	Server->PostSendUpdate();
+
+	// The client is no longer owning the object so it should be deleted
+	UE_NET_ASSERT_EQ(Client->GetReplicationBridge()->GetReplicatedObject(ServerObject->NetHandle), nullptr);
+}
+
+UE_NET_TEST_FIXTURE(FTestFilteringFixture, CanToggleConnectionFilter)
+{
+	// Add client
+	FReplicationSystemTestClient* Client = CreateClient();
+
+	// Spawn object on server
+	UReplicatedTestObject* ServerObject = Server->CreateObject(0, 0);
+
+	// Send and deliver packets
+	Server->PreSendUpdate();
+	Server->SendAndDeliverTo(Client, DeliverPacket);
+	Server->PostSendUpdate();
+
+	// Object should now exist on client
+	UE_NET_ASSERT_NE(Client->GetReplicationBridge()->GetReplicatedObject(ServerObject->NetHandle), nullptr);
+
+	// Apply connection filtering, not allowing any connection to receive the object
+	TBitArray<> NoConnections;
+	Server->ReplicationSystem->SetConnectionFilter(ServerObject->NetHandle, NoConnections, ENetFilterStatus::Allow);
+
+	// Send and deliver packets
+	Server->PreSendUpdate();
+	Server->SendAndDeliverTo(Client, DeliverPacket);
+	Server->PostSendUpdate();
+
+	// As object is now filtered it should be deleted on the client
+	UE_NET_ASSERT_EQ(Client->GetReplicationBridge()->GetReplicatedObject(ServerObject->NetHandle), nullptr);
+
+	// Remove the connection filter
+	Server->ReplicationSystem->SetFilter(ServerObject->NetHandle, InvalidNetObjectFilterHandle);
+
+	// Send and deliver packets
+	Server->PreSendUpdate();
+	Server->SendAndDeliverTo(Client, DeliverPacket);
+	Server->PostSendUpdate();
+
+	// The object should have been created again
+	UE_NET_ASSERT_NE(Client->GetReplicationBridge()->GetReplicatedObject(ServerObject->NetHandle), nullptr);
+
+	// Destroy the object
+	Server->DestroyObject(ServerObject);
+
+	// Send and deliver packets
+	Server->PreSendUpdate();
+	Server->SendAndDeliverTo(Client, DeliverPacket);
+	Server->PostSendUpdate();
+}
+
+
+// Keep last! Toggle between different kinds of filters
+UE_NET_TEST_FIXTURE(FTestFilteringFixture, CanToggleBetweenAllFilters)
+{
+	// Add client
+	FReplicationSystemTestClient* Client = CreateClient();
+
+	// Spawn object on server
+	UReplicatedTestObject* ServerObject = Server->CreateObject(0, 0);
+
+	// Toggle between owner and connection filters
+	{
+		// Apply owner filtering
+		Server->ReplicationSystem->SetFilter(ServerObject->NetHandle, ToOwnerFilterHandle);
+
+		// Send and deliver packets
+		Server->PreSendUpdate();
+		Server->SendAndDeliverTo(Client, DeliverPacket);
+		Server->PostSendUpdate();
+
+		// Object should not exist on client
+		UE_NET_ASSERT_EQ(Client->GetReplicationBridge()->GetReplicatedObject(ServerObject->NetHandle), nullptr);
+
+		// Apply connection filtering, not allowing any connection to receive the object
+		TBitArray<> NoConnections;
+		Server->ReplicationSystem->SetConnectionFilter(ServerObject->NetHandle, NoConnections, ENetFilterStatus::Allow);
+
+		// Send and deliver packets
+		Server->PreSendUpdate();
+		Server->SendAndDeliverTo(Client, DeliverPacket);
+		Server->PostSendUpdate();
+
+		// Object should still not exist on client
+		UE_NET_ASSERT_EQ(Client->GetReplicationBridge()->GetReplicatedObject(ServerObject->NetHandle), nullptr);
+
+		// Finally test going from connection filtering to owner filtering
+		Server->ReplicationSystem->SetFilter(ServerObject->NetHandle, ToOwnerFilterHandle);
+
+		// Send and deliver packets
+		Server->PreSendUpdate();
+		Server->SendAndDeliverTo(Client, DeliverPacket);
+		Server->PostSendUpdate();
+
+		// Object should still not exist on client
+		UE_NET_ASSERT_EQ(Client->GetReplicationBridge()->GetReplicatedObject(ServerObject->NetHandle), nullptr);
+	}
+}
+
+// Owner filtering tests
+UE_NET_TEST_FIXTURE(FTestFilteringFixture, GroupFilterPreventsObjectFromReplicating)
+{
+	// Add client
+	FReplicationSystemTestClient* Client = CreateClient();
+
+	// Spawn object on server
+	UReplicatedTestObject* ServerObject = Server->CreateObject(0, 0);
+
+	FNetObjectGroupHandle GroupHandle = Server->ReplicationSystem->CreateGroup();
+	Server->ReplicationSystem->AddToGroup(GroupHandle, ServerObject->NetHandle);
+	Server->ReplicationSystem->AddGroupFilter(GroupHandle);
+
+	// Filter out objects in group
+	Server->ReplicationSystem->SetGroupFilterStatus(GroupHandle, Client->ConnectionIdOnServer, ENetFilterStatus::Disallow);
+
+	// Send and deliver packet
+	Server->PreSendUpdate();
+	Server->SendAndDeliverTo(Client, DeliverPacket);
+	Server->PostSendUpdate();
+
+	// Object should not have been created on the client
+	UE_NET_ASSERT_EQ(Client->GetReplicationBridge()->GetReplicatedObject(ServerObject->NetHandle), nullptr);
+}
+
+UE_NET_TEST_FIXTURE(FTestFilteringFixture, GroupFilterAllowsObjectToReplicate)
+{
+	// Add client
+	FReplicationSystemTestClient* Client = CreateClient();
+
+	// Spawn object on server
+	UReplicatedTestObject* ServerObject = Server->CreateObject(0, 0);
+
+	FNetObjectGroupHandle GroupHandle = Server->ReplicationSystem->CreateGroup();
+	Server->ReplicationSystem->AddToGroup(GroupHandle, ServerObject->NetHandle);
+	Server->ReplicationSystem->AddGroupFilter(GroupHandle);
+
+	// Filter out objects in group
+	Server->ReplicationSystem->SetGroupFilterStatus(GroupHandle, Client->ConnectionIdOnServer, ENetFilterStatus::Disallow);
+
+	// Send and deliver packet
+	Server->PreSendUpdate();
+	Server->SendAndDeliverTo(Client, DeliverPacket);
+	Server->PostSendUpdate();
+
+	// Object should not have been created on the client
+	UE_NET_ASSERT_EQ(Client->GetReplicationBridge()->GetReplicatedObject(ServerObject->NetHandle), nullptr);
+
+	// Allow replication again
+	Server->ReplicationSystem->SetGroupFilterStatus(GroupHandle, Client->ConnectionIdOnServer, ENetFilterStatus::Allow);
+
+	// Send and deliver packet
+	Server->PreSendUpdate();
+	Server->SendAndDeliverTo(Client, DeliverPacket);
+	Server->PostSendUpdate();
+
+	// Object should now have been created on the client
+	UE_NET_ASSERT_NE(Client->GetReplicationBridge()->GetReplicatedObject(ServerObject->NetHandle), nullptr);
+}
+
+UE_NET_TEST_FIXTURE(FTestFilteringFixture, GroupFilterRestoresConnectionConnectionFilter)
+{
+	// Add clients
+	FReplicationSystemTestClient* ClientArray[] = {CreateClient(), CreateClient(), CreateClient()};
+	constexpr SIZE_T LastClientIndex = UE_ARRAY_COUNT(ClientArray) - 1;
+
+	// Spawn object on server
+	UReplicatedTestObject* ServerObject = Server->CreateObject(0, 0);
+
+	// Apply filtering that allows the last client to receive the object
+	const uint32 ConnectionIdForLastClient = ClientArray[LastClientIndex]->ConnectionIdOnServer;
+	TBitArray<> AllowedConnections;
+	AllowedConnections.Init(false, ConnectionIdForLastClient + 1);
+	AllowedConnections[ConnectionIdForLastClient] = true;
+	Server->ReplicationSystem->SetConnectionFilter(ServerObject->NetHandle, AllowedConnections, ENetFilterStatus::Allow);
+
+	// Send and deliver packets
+	Server->PreSendUpdate();
+	for (FReplicationSystemTestClient* Client : ClientArray)
+	{
+		Server->SendAndDeliverTo(Client, DeliverPacket);
+	}
+	Server->PostSendUpdate();
+
+	// Object should only have been created on the last client
+	for (FReplicationSystemTestClient*& Client : ClientArray)
+	{
+		const SIZE_T ClientIndex = &Client - &ClientArray[0];
+		if (ClientIndex == LastClientIndex)
+		{
+			UE_NET_ASSERT_NE(Client->GetReplicationBridge()->GetReplicatedObject(ServerObject->NetHandle), nullptr);
+		}
+		else
+		{
+			UE_NET_ASSERT_EQ(Client->GetReplicationBridge()->GetReplicatedObject(ServerObject->NetHandle), nullptr);
+		}
+	}
+
+	// Create and set group filter for last client only
+	FNetObjectGroupHandle GroupHandle = Server->ReplicationSystem->CreateGroup();
+	Server->ReplicationSystem->AddToGroup(GroupHandle, ServerObject->NetHandle);
+	Server->ReplicationSystem->AddGroupFilter(GroupHandle);
+	Server->ReplicationSystem->SetGroupFilterStatus(GroupHandle, ConnectionIdForLastClient, ENetFilterStatus::Disallow);
+
+	// Send and deliver packets
+	Server->PreSendUpdate();
+	for (FReplicationSystemTestClient* Client : ClientArray)
+	{
+		Server->SendAndDeliverTo(Client, DeliverPacket);
+	}
+	Server->PostSendUpdate();
+
+	// Object Should now have been destroyed on the last client
+	for (FReplicationSystemTestClient*& Client : ClientArray)
+	{
+		UE_NET_ASSERT_EQ(Client->GetReplicationBridge()->GetReplicatedObject(ServerObject->NetHandle), nullptr);
+	}
+
+	// Clear Group filter
+	Server->ReplicationSystem->SetGroupFilterStatus(GroupHandle, ConnectionIdForLastClient, ENetFilterStatus::Allow);
+
+	// Send and deliver packets
+	Server->PreSendUpdate();
+	for (FReplicationSystemTestClient* Client : ClientArray)
+	{
+		Server->SendAndDeliverTo(Client, DeliverPacket);
+	}
+	Server->PostSendUpdate();
+
+	// Object should now be recreated on last client
+	for (FReplicationSystemTestClient*& Client : ClientArray)
+	{
+		const SIZE_T ClientIndex = &Client - &ClientArray[0];
+		if (ClientIndex == LastClientIndex)
+		{
+			UE_NET_ASSERT_NE(Client->GetReplicationBridge()->GetReplicatedObject(ServerObject->NetHandle), nullptr);
+		}
+		else
+		{
+			UE_NET_ASSERT_EQ(Client->GetReplicationBridge()->GetReplicatedObject(ServerObject->NetHandle), nullptr);
+		}
+	}
+}
+
+UE_NET_TEST_FIXTURE(FTestFilteringFixture, CanGetDynamicFilter)
+{
+	UE_NET_ASSERT_NE(MockNetObjectFilter, nullptr);
+}
+
+UE_NET_TEST_FIXTURE(FTestFilteringFixture, CanGetDynamicFilterHandle)
+{
+	UE_NET_ASSERT_NE(MockFilterHandle, InvalidNetObjectFilterHandle);
+}
+
+UE_NET_TEST_FIXTURE(FTestFilteringFixture, DynamicFilterInitIsCalled)
+{
+	const UMockNetObjectFilter::FFunctionCallStatus& FunctionCallStatus = MockNetObjectFilter->GetFunctionCallStatus();
+	UE_NET_ASSERT_EQ(FunctionCallStatus.CallCounts.Init, 1U);
+	UE_NET_ASSERT_EQ(FunctionCallStatus.SuccessfulCallCounts.Init, 1U);
+}
+
+UE_NET_TEST_FIXTURE(FTestFilteringFixture, DynamicFilterAddObjectAndRemoveObjectIsCalledWhenObjectIsDeleted)
+{
+	// Setup dynamic filter for the test
+	{
+		UMockNetObjectFilter::FFunctionCallSetup CallSetup;
+		CallSetup.AddObject.bReturnValue = true;
+		CallSetup.Filter.bFilterOutByDefault = false;
+		MockNetObjectFilter->SetFunctionCallSetup(CallSetup);
+		MockNetObjectFilter->ResetFunctionCallStatus();
+	}
+
+	UReplicatedTestObject* ServerObject = Server->CreateObject(0, 0);
+	Server->ReplicationSystem->SetFilter(ServerObject->NetHandle, MockFilterHandle);
+
+	Server->PreSendUpdate();
+	// Filter needs to be set now.
+	{
+		const UMockNetObjectFilter::FFunctionCallStatus& CallStatus = MockNetObjectFilter->GetFunctionCallStatus();
+
+		UE_NET_ASSERT_EQ(CallStatus.CallCounts.AddObject, 1U);
+		UE_NET_ASSERT_EQ(CallStatus.SuccessfulCallCounts.AddObject, 1U);
+		UE_NET_ASSERT_EQ(CallStatus.CallCounts.RemoveObject, 0U);
+
+		MockNetObjectFilter->ResetFunctionCallStatus();
+	}
+	Server->PostSendUpdate();
+
+	Server->DestroyObject(ServerObject);
+
+	Server->PreSendUpdate();
+	// Filter needs to be cleared now.
+	{
+		const UMockNetObjectFilter::FFunctionCallStatus& CallStatus = MockNetObjectFilter->GetFunctionCallStatus();
+
+		UE_NET_ASSERT_EQ(CallStatus.CallCounts.RemoveObject, 1U);
+		UE_NET_ASSERT_EQ(CallStatus.SuccessfulCallCounts.RemoveObject, 1U);
+
+		MockNetObjectFilter->ResetFunctionCallStatus();
+	}
+	Server->PostSendUpdate();
+}
+
+UE_NET_TEST_FIXTURE(FTestFilteringFixture, DynamicFilterCanAllowObjectToReplicate)
+{
+	// Setup dynamic filter for the test
+	{
+		UMockNetObjectFilter::FFunctionCallSetup CallSetup;
+		CallSetup.AddObject.bReturnValue = true;
+		CallSetup.Filter.bFilterOutByDefault = false;
+		MockNetObjectFilter->SetFunctionCallSetup(CallSetup);
+		MockNetObjectFilter->ResetFunctionCallStatus();
+	}
+
+	// Add client
+	FReplicationSystemTestClient* Client = CreateClient();
+
+	// Spawn object on server and set filter
+	UReplicatedTestObject* ServerObject = Server->CreateObject(0, 0);
+	Server->ReplicationSystem->SetFilter(ServerObject->NetHandle, MockFilterHandle);
+
+	// Send and deliver packets
+	Server->PreSendUpdate();
+	Server->SendAndDeliverTo(Client, DeliverPacket);
+	Server->PostSendUpdate();
+
+	// Object should now exist on client
+	UE_NET_ASSERT_NE(Client->GetReplicationBridge()->GetReplicatedObject(ServerObject->NetHandle), nullptr);
+}
+
+UE_NET_TEST_FIXTURE(FTestFilteringFixture, DynamicFilterCanDisallowObjectToReplicate)
+{
+	// Setup dynamic filter for the test
+	{
+		UMockNetObjectFilter::FFunctionCallSetup CallSetup;
+		CallSetup.AddObject.bReturnValue = true;
+		CallSetup.Filter.bFilterOutByDefault = true;
+		MockNetObjectFilter->SetFunctionCallSetup(CallSetup);
+		MockNetObjectFilter->ResetFunctionCallStatus();
+	}
+
+	// Add client
+	FReplicationSystemTestClient* Client = CreateClient();
+
+	// Spawn object on server and set filter
+	UReplicatedTestObject* ServerObject = Server->CreateObject(0, 0);
+	Server->ReplicationSystem->SetFilter(ServerObject->NetHandle, MockFilterHandle);
+
+	// Send and deliver packets
+	Server->PreSendUpdate();
+	Server->SendAndDeliverTo(Client, DeliverPacket);
+	Server->PostSendUpdate();
+
+	// Object should not exist on client
+	UE_NET_ASSERT_EQ(Client->GetReplicationBridge()->GetReplicatedObject(ServerObject->NetHandle), nullptr);
+}
+
+UE_NET_TEST_FIXTURE(FTestFilteringFixture, SwitchingFiltersCallsRemoveObjectOnPreviousFilter)
+{
+	// Setup dynamic filter for the test
+	{
+		UMockNetObjectFilter::FFunctionCallSetup CallSetup;
+		CallSetup.AddObject.bReturnValue = true;
+		CallSetup.Filter.bFilterOutByDefault = true;
+		MockNetObjectFilter->SetFunctionCallSetup(CallSetup);
+		MockNetObjectFilter->ResetFunctionCallStatus();
+	}
+
+	// Spawn object on server and set filter
+	UReplicatedTestObject* ServerObject = Server->CreateObject(0, 0);
+	Server->ReplicationSystem->SetFilter(ServerObject->NetHandle, MockFilterHandle);
+
+	// Make sure filter is set
+	Server->PreSendUpdate();
+	Server->PostSendUpdate();
+
+	// Check RemoveObject is called when switching filters.
+	{
+		MockNetObjectFilter->ResetFunctionCallStatus();
+
+		const TBitArray<> NoConnections;
+		Server->ReplicationSystem->SetConnectionFilter(ServerObject->NetHandle, NoConnections, ENetFilterStatus::Disallow);
+
+		Server->PreSendUpdate();
+		Server->PostSendUpdate();
+
+		const UMockNetObjectFilter::FFunctionCallStatus& FunctionCallStatus = MockNetObjectFilter->GetFunctionCallStatus();
+		UE_NET_ASSERT_EQ(FunctionCallStatus.CallCounts.RemoveObject, 1U);
+		UE_NET_ASSERT_EQ(FunctionCallStatus.SuccessfulCallCounts.RemoveObject, 1U);
+	}
+}
+
+UE_NET_TEST_FIXTURE(FTestFilteringFixture, SubObjectsAreReplicatedWhenOwnerDynamicFilterAllows)
+{
+	// Setup dynamic filter for the test
+	{
+		UMockNetObjectFilter::FFunctionCallSetup CallSetup;
+		CallSetup.AddObject.bReturnValue = true;
+		CallSetup.Filter.bFilterOutByDefault = false;
+		MockNetObjectFilter->SetFunctionCallSetup(CallSetup);
+		MockNetObjectFilter->ResetFunctionCallStatus();
+	}
+
+	// Add client
+	FReplicationSystemTestClient* Client = CreateClient();
+
+	// Spawn object on server and set filter
+	UReplicatedTestObject* ServerObject = Server->CreateObject(0, 0);
+	Server->ReplicationSystem->SetFilter(ServerObject->NetHandle, MockFilterHandle);
+
+	UReplicatedTestObject* ServerSubObject = Server->CreateSubObject(ServerObject->NetHandle, 0, 0);
+
+	// Send and deliver packets
+	Server->PreSendUpdate();
+	Server->SendAndDeliverTo(Client, DeliverPacket);
+	Server->PostSendUpdate();
+
+	// Check that both the object and subobject exist.
+	UE_NET_ASSERT_NE(Client->GetReplicationBridge()->GetReplicatedObject(ServerObject->NetHandle), nullptr);
+	UE_NET_ASSERT_NE(Client->GetReplicationBridge()->GetReplicatedObject(ServerSubObject->NetHandle), nullptr);
+}
+
+UE_NET_TEST_FIXTURE(FTestFilteringFixture, SubObjectsAreNotReplicatedWhenOwnerDynamicFilterDisallows)
+{
+	// Setup dynamic filter for the test
+	{
+		UMockNetObjectFilter::FFunctionCallSetup CallSetup;
+		CallSetup.AddObject.bReturnValue = true;
+		CallSetup.Filter.bFilterOutByDefault = true;
+		MockNetObjectFilter->SetFunctionCallSetup(CallSetup);
+		MockNetObjectFilter->ResetFunctionCallStatus();
+	}
+
+	// Add client
+	FReplicationSystemTestClient* Client = CreateClient();
+
+	// Spawn object on server and set filter
+	UReplicatedTestObject* ServerObject = Server->CreateObject(0, 0);
+	Server->ReplicationSystem->SetFilter(ServerObject->NetHandle, MockFilterHandle);
+
+	UReplicatedTestObject* ServerSubObject = Server->CreateSubObject(ServerObject->NetHandle, 0, 0);
+
+	// Send and deliver packets
+	Server->PreSendUpdate();
+	Server->SendAndDeliverTo(Client, DeliverPacket);
+	Server->PostSendUpdate();
+
+	// Check that neither the object nor the subobject exist.
+	UE_NET_ASSERT_EQ(Client->GetReplicationBridge()->GetReplicatedObject(ServerObject->NetHandle), nullptr);
+	UE_NET_ASSERT_EQ(Client->GetReplicationBridge()->GetReplicatedObject(ServerSubObject->NetHandle), nullptr);
+}
+
+UE_NET_TEST_FIXTURE(FTestFilteringFixture, DependentObjectIsUnaffectedByDynamicFilter)
+{
+	// Setup dynamic filter for the test
+	{
+		UMockNetObjectFilter::FFunctionCallSetup CallSetup;
+		CallSetup.AddObject.bReturnValue = true;
+		CallSetup.Filter.bFilterOutByDefault = true;
+		MockNetObjectFilter->SetFunctionCallSetup(CallSetup);
+		MockNetObjectFilter->ResetFunctionCallStatus();
+	}
+
+	// Add client
+	FReplicationSystemTestClient* Client = CreateClient();
+
+	// Spawn objects on server and set filter on dependent object
+	UReplicatedTestObject* ServerObject = Server->CreateObject(0, 0);
+	UReplicatedTestObject* ServerFutureDependentObject = Server->CreateObject(0, 0);
+	Server->ReplicationSystem->SetFilter(ServerFutureDependentObject->NetHandle, MockFilterHandle);
+
+	// Send and deliver packets
+	Server->PreSendUpdate();
+	Server->SendAndDeliverTo(Client, DeliverPacket);
+	Server->PostSendUpdate();
+
+	// We expect the object to exist and the future dependent object not to exist
+	UE_NET_ASSERT_NE(Client->GetReplicationBridge()->GetReplicatedObject(ServerObject->NetHandle), nullptr);
+	UE_NET_ASSERT_EQ(Client->GetReplicationBridge()->GetReplicatedObject(ServerFutureDependentObject->NetHandle), nullptr);
+
+	// Make dependent object and make sure it's now replicated.
+	Server->ReplicationBridge->AddDependentObject(ServerObject->NetHandle, ServerFutureDependentObject->NetHandle);
+	UReplicatedTestObject* ServerDependentObject = ServerFutureDependentObject;
+
+	// Send and deliver packets
+	Server->PreSendUpdate();
+	Server->SendAndDeliverTo(Client, DeliverPacket);
+	Server->PostSendUpdate();
+
+	// We now expect the dependent object to exist
+	UE_NET_ASSERT_NE(Client->GetReplicationBridge()->GetReplicatedObject(ServerDependentObject->NetHandle), nullptr);
+
+	// Remove dependency and make sure the formerly dependent object is removed from the client
+	Server->ReplicationBridge->RemoveDependentObject(ServerObject->NetHandle, ServerDependentObject->NetHandle);
+	UReplicatedTestObject* ServerFormerDependentObject = ServerDependentObject;
+
+	// Send and deliver packets
+	Server->PreSendUpdate();
+	Server->SendAndDeliverTo(Client, DeliverPacket);
+	Server->PostSendUpdate();
+
+	// We expect the former dependent object not to exist
+	UE_NET_ASSERT_EQ(Client->GetReplicationBridge()->GetReplicatedObject(ServerFormerDependentObject->NetHandle), nullptr);
+}
+
+UE_NET_TEST_FIXTURE(FTestFilteringFixture, NestedDependentObjectIsFilteredAsParentsOrIndependent)
+{
+	// Setup dynamic filter for the test
+	{
+		UMockNetObjectFilter::FFunctionCallSetup CallSetup;
+		CallSetup.AddObject.bReturnValue = true;
+		CallSetup.Filter.bFilterOutByDefault = true;
+		MockNetObjectFilter->SetFunctionCallSetup(CallSetup);
+		MockNetObjectFilter->ResetFunctionCallStatus();
+	}
+
+	// Add client
+	FReplicationSystemTestClient* Client = CreateClient();
+
+	// Spawn objects on server and set filter on dependent objects
+	UReplicatedTestObject* ServerObject = Server->CreateObject(0, 0);
+	UReplicatedTestObject* ServerFutureDependentObject = Server->CreateObject(0, 0);
+	UReplicatedTestObject* ServerFutureNestedDependentObject = Server->CreateObject(0, 0);
+
+	Server->ReplicationSystem->SetFilter(ServerFutureDependentObject->NetHandle, MockFilterHandle);
+	Server->ReplicationSystem->SetFilter(ServerFutureNestedDependentObject->NetHandle, MockFilterHandle);
+
+	// Send and deliver packets
+	Server->PreSendUpdate();
+	Server->SendAndDeliverTo(Client, DeliverPacket);
+	Server->PostSendUpdate();
+
+	// We expect the object to exist and the future dependent and future dependent objects not to exist
+	UE_NET_ASSERT_NE(Client->GetReplicationBridge()->GetReplicatedObject(ServerObject->NetHandle), nullptr);
+	UE_NET_ASSERT_EQ(Client->GetReplicationBridge()->GetReplicatedObject(ServerFutureDependentObject->NetHandle), nullptr);
+	UE_NET_ASSERT_EQ(Client->GetReplicationBridge()->GetReplicatedObject(ServerFutureNestedDependentObject->NetHandle), nullptr);
+
+	// Make dependent objects and make sure they're now replicated.
+	Server->ReplicationBridge->AddDependentObject(ServerObject->NetHandle, ServerFutureDependentObject->NetHandle);
+	Server->ReplicationBridge->AddDependentObject(ServerFutureDependentObject->NetHandle, ServerFutureNestedDependentObject->NetHandle);
+
+	UReplicatedTestObject* ServerDependentObject = ServerFutureDependentObject;
+	UReplicatedTestObject* ServerNestedDependentObject = ServerFutureNestedDependentObject;
+
+	// Send and deliver packets
+	Server->PreSendUpdate();
+	Server->SendAndDeliverTo(Client, DeliverPacket);
+	Server->PostSendUpdate();
+
+	// We now expect the dependent objects to exist
+	UE_NET_ASSERT_NE(Client->GetReplicationBridge()->GetReplicatedObject(ServerDependentObject->NetHandle), nullptr);
+	UE_NET_ASSERT_NE(Client->GetReplicationBridge()->GetReplicatedObject(ServerNestedDependentObject->NetHandle), nullptr);
+
+	// Remove dependency on root and make sure the formerly dependent object is removed from the client thanks to the filter.
+	Server->ReplicationBridge->RemoveDependentObject(ServerObject->NetHandle, ServerDependentObject->NetHandle);
+	UReplicatedTestObject* ServerFormerDependentObject = ServerDependentObject;
+
+	// Send and deliver packets
+	Server->PreSendUpdate();
+	Server->SendAndDeliverTo(Client, DeliverPacket);
+	Server->PostSendUpdate();
+
+	// We expect the former dependent object to not to exist, thanks to the filter
+	UE_NET_ASSERT_EQ(Client->GetReplicationBridge()->GetReplicatedObject(ServerFormerDependentObject->NetHandle), nullptr);
+
+	// As the former dependent object is filtered out it's ok for the nested dependent object to be filtered out.
+	UE_NET_ASSERT_EQ(Client->GetReplicationBridge()->GetReplicatedObject(ServerNestedDependentObject->NetHandle), nullptr);
+
+	// Remove filter on the nested dependent object
+	Server->ReplicationSystem->SetFilter(ServerNestedDependentObject->NetHandle, InvalidNetObjectFilterHandle);
+
+	// Send and deliver packets
+	Server->PreSendUpdate();
+	Server->SendAndDeliverTo(Client, DeliverPacket);
+	Server->PostSendUpdate();
+
+	// Verify that dependent object no longer is filtered out even though its parent is
+	UE_NET_ASSERT_EQ(Client->GetReplicationBridge()->GetReplicatedObject(ServerFormerDependentObject->NetHandle), nullptr);
+	UE_NET_ASSERT_NE(Client->GetReplicationBridge()->GetReplicatedObject(ServerNestedDependentObject->NetHandle), nullptr);
+}
+
+UE_NET_TEST_FIXTURE(FTestFilteringFixture, ObjectGetsFilterOutSettingFromStart)
+{
+	// Setup dynamic filter for the test. For this test we want the value of the object's filter property to rule.
+	{
+		UMockNetObjectFilter::FFunctionCallSetup CallSetup;
+		CallSetup.AddObject.bReturnValue = true;
+		CallSetup.Filter.bFilterOutByDefault = false;
+		MockNetObjectFilter->SetFunctionCallSetup(CallSetup);
+		MockNetObjectFilter->ResetFunctionCallStatus();
+	}
+
+	// Add client
+	FReplicationSystemTestClient* Client = CreateClient();
+
+	// Create object and set filter
+	UTestFilteringObject* ServerObject = Server->CreateObject<UTestFilteringObject>();
+	Server->ReplicationSystem->SetFilter(ServerObject->NetHandle, MockFilterHandle);
+
+	// We want the object to be filtered out
+	constexpr bool bFilterOut = true;
+	ServerObject->SetFilterOut(bFilterOut);
+
+	// Send and deliver packets
+	Server->PreSendUpdate();
+	Server->SendAndDeliverTo(Client, DeliverPacket);
+	Server->PostSendUpdate();
+
+	// Check that the object does not exist on the client.
+	UE_NET_ASSERT_EQ(Client->GetReplicationBridge()->GetReplicatedObject(ServerObject->NetHandle), nullptr);
+}
+
+UE_NET_TEST_FIXTURE(FTestFilteringFixture, ObjectGetsUpdatedFilterOutSetting)
+{
+	// Setup dynamic filter for the test. For this test we want the value of the object's filter property to rule.
+	{
+		UMockNetObjectFilter::FFunctionCallSetup CallSetup;
+		CallSetup.AddObject.bReturnValue = true;
+		CallSetup.Filter.bFilterOutByDefault = false;
+		MockNetObjectFilter->SetFunctionCallSetup(CallSetup);
+		MockNetObjectFilter->ResetFunctionCallStatus();
+	}
+
+	// Setup dynamic filter for the test. For this test we want the value of the object's filter property to rule.
+	{
+		UMockNetObjectFilter::FFunctionCallSetup CallSetup;
+		CallSetup.AddObject.bReturnValue = true;
+		CallSetup.Filter.bFilterOutByDefault = false;
+		MockNetObjectFilter->SetFunctionCallSetup(CallSetup);
+		MockNetObjectFilter->ResetFunctionCallStatus();
+	}
+
+	// Add client
+	FReplicationSystemTestClient* Client = CreateClient();
+
+	// Create object and set filter
+	UTestFilteringObject* ServerObject = Server->CreateObject<UTestFilteringObject>();
+	Server->ReplicationSystem->SetFilter(ServerObject->NetHandle, MockFilterHandle);
+
+	// We don't want the object to be filtered out
+	ServerObject->SetFilterOut(false);
+
+	// Send and deliver packets
+	Server->PreSendUpdate();
+	Server->SendAndDeliverTo(Client, DeliverPacket);
+	Server->PostSendUpdate();
+
+	// Check that the object exists on the client.
+	UE_NET_ASSERT_NE(Client->GetReplicationBridge()->GetReplicatedObject(ServerObject->NetHandle), nullptr);
+
+	// Mark the object to be filtered out.
+	ServerObject->SetFilterOut(true);
+
+	// Send and deliver packets
+	Server->PreSendUpdate();
+	Server->SendAndDeliverTo(Client, DeliverPacket);
+	Server->PostSendUpdate();
+
+	// Check that the object does not exist on the client.
+	UE_NET_ASSERT_EQ(Client->GetReplicationBridge()->GetReplicatedObject(ServerObject->NetHandle), nullptr);
+}
+
+UE_NET_TEST_FIXTURE(FTestFilteringFixture, DynamicFilteredOutSubObjectsAreResetWhenIndexIsReused)
+{
+	// Setup dynamic filter for the test
+	{
+		UMockNetObjectFilter::FFunctionCallSetup CallSetup;
+		CallSetup.AddObject.bReturnValue = true;
+		CallSetup.Filter.bFilterOutByDefault = true;
+		MockNetObjectFilter->SetFunctionCallSetup(CallSetup);
+		MockNetObjectFilter->ResetFunctionCallStatus();
+	}
+
+	// Add client
+	FReplicationSystemTestClient* Client = CreateClient();
+
+	// Spawn object on server and set filter
+	UReplicatedTestObject* ServerObject = Server->CreateObject(0, 0);
+	Server->ReplicationSystem->SetFilter(ServerObject->NetHandle, MockFilterHandle);
+
+	// Create and destroy subobject
+	{
+		UReplicatedTestObject* ServerSubObject = Server->CreateSubObject(ServerObject->NetHandle, 0, 0);
+		Server->DestroyObject(ServerSubObject);
+	}
+
+	// Send and deliver packets
+	Server->PreSendUpdate();
+	Server->SendAndDeliverTo(Client, DeliverPacket);
+	Server->PostSendUpdate();
+
+	// Object should not exist on client
+	UE_NET_ASSERT_EQ(Client->GetReplicationBridge()->GetReplicatedObject(ServerObject->NetHandle), nullptr);
+
+	// Create new object which should get the same internal index as the destroyed SubObject
+	UReplicatedTestObject* ServerObject2 = Server->CreateObject(0, 0);
+
+	// Send and deliver packets
+	Server->PreSendUpdate();
+	Server->SendAndDeliverTo(Client, DeliverPacket);
+	Server->PostSendUpdate();
+
+	// Object should exist on client
+	UE_NET_ASSERT_NE(Client->GetReplicationBridge()->GetReplicatedObject(ServerObject2->NetHandle), nullptr);
+}
+
+
+}
