Iris experimental replication system.

diff --git a/Engine/Plugins/Runtime/ReplicationSystemTestPlugin/Source/Private/Tests/ReplicationSystem/Filtering/TestGroups.cpp b/Engine/Plugins/Runtime/ReplicationSystemTestPlugin/Source/Private/Tests/ReplicationSystem/Filtering/TestGroups.cpp
new file mode 100644
--- /dev/null
+++ b/Engine/Plugins/Runtime/ReplicationSystemTestPlugin/Source/Private/Tests/ReplicationSystem/Filtering/TestGroups.cpp
@@ -0,0 +1,572 @@
+#include "NetworkAutomationTest.h"
+#include "NetworkAutomationTestMacros.h"
+#include "Iris/ReplicationSystem/ReplicationSystem.h"
+#include "Iris/ReplicationSystem/Filtering/NetObjectGroups.h"
+#include "Tests/ReplicationSystem/ReplicatedTestObject.h"
+#include "Tests/ReplicationSystem/ReplicationSystemServerClientTestFixture.h"
+#include "Iris/ReplicationSystem/Filtering/ReplicationFiltering.h"
+#include "Iris/ReplicationSystem/ReplicationSystemInternal.h"
+#include "Iris/ReplicationSystem/Filtering/NetObjectFilter.h"
+
+namespace UE::Net::Private
+{
+
+class FTestGroupsFixture : public FReplicationSystemServerClientTestFixture
+{
+public:
+
+	FReplicationFiltering* ServerFiltering = nullptr;
+	FNetObjectGroups* ServerGroups = nullptr;
+	FDirtyNetObjectTracker* ServerDirtyNetObjectTracker = nullptr;
+	
+protected:
+	virtual void SetUp() override
+	{
+		FReplicationSystemServerClientTestFixture::SetUp();
+		ServerFiltering = &Server->ReplicationSystem->GetReplicationSystemInternal()->GetFiltering();
+		ServerGroups = &Server->ReplicationSystem->GetReplicationSystemInternal()->GetGroups();
+		ServerDirtyNetObjectTracker = &Server->ReplicationSystem->GetReplicationSystemInternal()->GetDirtyNetObjectTracker();
+	}
+
+	void Filter()
+	{
+		Server->PreSendUpdate();
+		Server->PostSendUpdate();
+	}
+
+	void VerifyGroupStatus(FNetObjectGroupHandle GroupHandle, uint32 ConnectionId)
+	{
+		FNetObjectGroup* Group = ServerGroups->GetGroup(GroupHandle);
+
+		UE_NET_ASSERT_TRUE(Group != nullptr);
+
+		ENetFilterStatus ReplicationStatus = ENetFilterStatus::Disallow;
+		UE_NET_ASSERT_TRUE(ServerFiltering->GetGroupFilterStatus(GroupHandle, ConnectionId, ReplicationStatus));
+
+		const FNetBitArrayView ObjectGroupFilter = ServerFiltering->GetGroupFilteredOutObjects(ConnectionId);
+		for (uint32 InternalIndex : MakeArrayView(Group->Members.GetData(), Group->Members.Num()))
+		{
+			UE_NET_ASSERT_TRUE(ObjectGroupFilter.GetBit(InternalIndex) == (ReplicationStatus == ENetFilterStatus::Disallow));
+		}
+	}
+
+	void VerifyObjectFilterStatus(FNetHandle Handle, uint32 ConnectionId, bool bExpectedReplicationAllowed)
+	{
+		const uint32 InternalIndex = Server->GetReplicationSystem()->GetReplicationSystemInternal()->GetNetHandleManager().GetInternalIndex(Handle);
+		UE_NET_ASSERT_NE((uint32)FNetHandleManager::InvalidInternalIndex, InternalIndex);
+
+		// Verify expected filter status
+		const FNetBitArrayView ObjectGroupFilter = ServerFiltering->GetGroupFilteredOutObjects(ConnectionId);
+
+		// Filter is subtractive, so if the bit is set we the object is filtered out
+		UE_NET_ASSERT_TRUE(ObjectGroupFilter.GetBit(InternalIndex) == !bExpectedReplicationAllowed);		
+	}
+};
+
+UE_NET_TEST_FIXTURE(FTestGroupsFixture, PublicGroupAPI)
+{
+	// Add client
+	FReplicationSystemTestClient* Client = CreateClient();
+	UReplicationSystem* ReplicationSystem = Server->ReplicationSystem;
+
+	// Create
+	{
+		FNetObjectGroupHandle GroupHandle = ReplicationSystem->CreateGroup();
+
+		UE_NET_ASSERT_FALSE(ReplicationSystem->IsValidGroup(0U));
+		UE_NET_ASSERT_TRUE(ReplicationSystem->IsValidGroup(GroupHandle));
+
+		// Destroy
+		ReplicationSystem->DestroyGroup(GroupHandle);
+		UE_NET_ASSERT_FALSE(ReplicationSystem->IsValidGroup(GroupHandle));
+	}
+
+	// Add/Remove
+	{
+		// Spawn object on server
+		UReplicatedTestObject* ServerObject = Server->CreateObject(0, 0);
+
+		FNetObjectGroupHandle GroupHandle = ReplicationSystem->CreateGroup();
+
+		UE_NET_ASSERT_FALSE(ReplicationSystem->IsInGroup(GroupHandle, ServerObject->NetHandle));
+
+		ReplicationSystem->AddToGroup(GroupHandle, ServerObject->NetHandle);
+
+		UE_NET_ASSERT_TRUE(ReplicationSystem->IsInGroup(GroupHandle, ServerObject->NetHandle));
+
+		ReplicationSystem->RemoveFromGroup(GroupHandle, ServerObject->NetHandle);
+
+		UE_NET_ASSERT_FALSE(ReplicationSystem->IsInGroup(GroupHandle, ServerObject->NetHandle));
+
+		ReplicationSystem->DestroyGroup(GroupHandle);
+
+		UE_NET_ASSERT_FALSE(ReplicationSystem->IsInGroup(GroupHandle, ServerObject->NetHandle));
+	}
+}
+
+UE_NET_TEST_FIXTURE(FTestGroupsFixture, PublicGroupAPIMemberOfMultipleGroups)
+{
+	// Add client
+	FReplicationSystemTestClient* Client = CreateClient();
+	UReplicationSystem* ReplicationSystem = Server->ReplicationSystem;
+
+	// Create
+	FNetObjectGroupHandle GroupHandle0 = ReplicationSystem->CreateGroup();
+	FNetObjectGroupHandle GroupHandle1 = ReplicationSystem->CreateGroup();
+
+	UReplicatedTestObject* ServerObject = Server->CreateObject(0, 0);
+
+	UE_NET_ASSERT_FALSE(ReplicationSystem->IsInGroup(GroupHandle0, ServerObject->NetHandle));
+	UE_NET_ASSERT_FALSE(ReplicationSystem->IsInGroup(GroupHandle1, ServerObject->NetHandle));
+
+	ReplicationSystem->AddToGroup(GroupHandle0, ServerObject->NetHandle);
+
+	UE_NET_ASSERT_TRUE(ReplicationSystem->IsInGroup(GroupHandle0, ServerObject->NetHandle));
+	UE_NET_ASSERT_FALSE(ReplicationSystem->IsInGroup(GroupHandle1, ServerObject->NetHandle));
+
+	ReplicationSystem->AddToGroup(GroupHandle1, ServerObject->NetHandle);
+
+	UE_NET_ASSERT_TRUE(ReplicationSystem->IsInGroup(GroupHandle0, ServerObject->NetHandle));
+	UE_NET_ASSERT_TRUE(ReplicationSystem->IsInGroup(GroupHandle1, ServerObject->NetHandle));
+
+	ReplicationSystem->RemoveFromGroup(GroupHandle0, ServerObject->NetHandle);
+
+	UE_NET_ASSERT_FALSE(ReplicationSystem->IsInGroup(GroupHandle0, ServerObject->NetHandle));
+	UE_NET_ASSERT_TRUE(ReplicationSystem->IsInGroup(GroupHandle1, ServerObject->NetHandle));
+
+	ReplicationSystem->RemoveFromGroup(GroupHandle1, ServerObject->NetHandle);
+
+	UE_NET_ASSERT_FALSE(ReplicationSystem->IsInGroup(GroupHandle0, ServerObject->NetHandle));
+	UE_NET_ASSERT_FALSE(ReplicationSystem->IsInGroup(GroupHandle1, ServerObject->NetHandle));
+}
+
+// Group filtering tests
+UE_NET_TEST_FIXTURE(FTestGroupsFixture, GroupFilterAPI)
+{
+	// Spawn object on server
+	UReplicatedTestObject* ServerObject0 = Server->CreateObject(0, 0);
+	UReplicatedTestObject* ServerObject1 = Server->CreateObject(0, 0);
+	UReplicatedTestObject* ServerObject2 = Server->CreateObject(0, 0);
+
+	// Create group
+	FNetObjectGroupHandle GroupHandle = Server->ReplicationSystem->CreateGroup();
+	UE_NET_ASSERT_NE(0U, (uint32)GroupHandle);	
+
+	// Add Objects to group
+	Server->ReplicationSystem->AddToGroup(GroupHandle, ServerObject0->NetHandle);
+	Server->ReplicationSystem->AddToGroup(GroupHandle, ServerObject1->NetHandle);
+	Server->ReplicationSystem->AddToGroup(GroupHandle, ServerObject2->NetHandle);
+	
+	UReplicatedTestObject* ServerObject3 = Server->CreateObject(0, 0);
+	UReplicatedTestObject* ServerObject4 = Server->CreateObject(0, 0);
+	UReplicatedTestObject* ServerObject5 = Server->CreateObject(0, 0);
+
+	// Create another group
+	FNetObjectGroupHandle GroupHandle2 = Server->ReplicationSystem->CreateGroup();
+	UE_NET_ASSERT_NE(0U, (uint32)GroupHandle2);	
+
+	// Add Objects to group
+	Server->ReplicationSystem->AddToGroup(GroupHandle2, ServerObject3->NetHandle);
+	Server->ReplicationSystem->AddToGroup(GroupHandle2, ServerObject4->NetHandle);
+	Server->ReplicationSystem->AddToGroup(GroupHandle2, ServerObject5->NetHandle);
+													 
+	UReplicatedTestObject* ServerObject6 = Server->CreateObject(0, 0);
+
+	// Create another group
+	FNetObjectGroupHandle GroupHandle3 = Server->ReplicationSystem->CreateGroup();
+	UE_NET_ASSERT_NE(0U, (uint32)GroupHandle3);	
+
+	// Add client
+	FReplicationSystemTestClient* Client = CreateClient();
+
+	// Update filters
+	Filter();
+
+	const FNetBitArrayView ObjectGroupFilter = ServerFiltering->GetGroupFilteredOutObjects(Client->ConnectionIdOnServer);
+
+	// Add group to group filter
+	Server->ReplicationSystem->AddGroupFilter(GroupHandle);
+	Server->ReplicationSystem->AddGroupFilter(GroupHandle2);
+	Server->ReplicationSystem->AddGroupFilter(GroupHandle3);
+
+	Server->ReplicationSystem->SetGroupFilterStatus(GroupHandle, Client->ConnectionIdOnServer, ENetFilterStatus::Disallow);
+	Server->ReplicationSystem->SetGroupFilterStatus(GroupHandle2, Client->ConnectionIdOnServer, ENetFilterStatus::Disallow);
+	Server->ReplicationSystem->SetGroupFilterStatus(GroupHandle3, Client->ConnectionIdOnServer, ENetFilterStatus::Disallow);
+
+	// Update filters
+	Filter();
+
+	VerifyGroupStatus(GroupHandle, Client->ConnectionIdOnServer);
+	VerifyGroupStatus(GroupHandle2, Client->ConnectionIdOnServer);
+	VerifyGroupStatus(GroupHandle3, Client->ConnectionIdOnServer);
+
+	UE_NET_ASSERT_TRUE(ObjectGroupFilter.IsAnyBitSet());
+
+	Server->ReplicationSystem->SetGroupFilterStatus(GroupHandle, Client->ConnectionIdOnServer, ENetFilterStatus::Allow);
+
+	// Update filters
+	Filter();
+
+	VerifyGroupStatus(GroupHandle, Client->ConnectionIdOnServer);
+	VerifyGroupStatus(GroupHandle2, Client->ConnectionIdOnServer);
+	VerifyGroupStatus(GroupHandle3, Client->ConnectionIdOnServer);
+
+	UE_NET_ASSERT_TRUE(ObjectGroupFilter.IsAnyBitSet());
+
+	Server->ReplicationSystem->SetGroupFilterStatus(GroupHandle2, Client->ConnectionIdOnServer, ENetFilterStatus::Allow);
+
+	// Update filters
+	Filter();
+
+	VerifyGroupStatus(GroupHandle, Client->ConnectionIdOnServer);
+	VerifyGroupStatus(GroupHandle2, Client->ConnectionIdOnServer);
+	VerifyGroupStatus(GroupHandle3, Client->ConnectionIdOnServer);
+
+	UE_NET_ASSERT_TRUE(ObjectGroupFilter.IsNoBitSet());
+
+	// Add Object to group 3 which is set, which should filter out the object
+	Server->ReplicationSystem->AddToGroup(GroupHandle3, ServerObject6->NetHandle);
+
+	// Update filters
+	Filter();
+
+	VerifyGroupStatus(GroupHandle, Client->ConnectionIdOnServer);
+	VerifyGroupStatus(GroupHandle2, Client->ConnectionIdOnServer);
+	VerifyGroupStatus(GroupHandle3, Client->ConnectionIdOnServer);
+
+	UE_NET_ASSERT_TRUE(ObjectGroupFilter.IsAnyBitSet());
+
+	// Add Object to group 3 which is set, which should filter out the object
+	Server->ReplicationSystem->RemoveFromGroup(GroupHandle3, ServerObject6->NetHandle);
+
+	// Update filters
+	Filter();
+
+	VerifyGroupStatus(GroupHandle, Client->ConnectionIdOnServer);
+	VerifyGroupStatus(GroupHandle2, Client->ConnectionIdOnServer);
+	VerifyGroupStatus(GroupHandle3, Client->ConnectionIdOnServer);
+
+	UE_NET_ASSERT_TRUE(ObjectGroupFilter.IsNoBitSet());
+
+	// Delete group
+	Server->ReplicationSystem->DestroyGroup(GroupHandle);
+	Server->ReplicationSystem->DestroyGroup(GroupHandle2);
+	Server->ReplicationSystem->DestroyGroup(GroupHandle3);
+
+	Filter();
+
+	UE_NET_ASSERT_TRUE(ObjectGroupFilter.IsNoBitSet());
+}
+
+UE_NET_TEST_FIXTURE(FTestGroupsFixture, GroupFilterAPIManyFilters)
+{
+	// Spawn object on server
+	UReplicatedTestObject* ServerObject0 = Server->CreateObject(0, 0);
+
+	// Create groups and add them as filters
+	const uint32 GroupCount = 256;
+
+	FNetObjectGroupHandle GroupHandles[GroupCount];
+
+	for (uint32 It = 0; It < GroupCount; ++It)
+	{
+		FNetObjectGroupHandle GroupHandle = Server->ReplicationSystem->CreateGroup();
+		Server->ReplicationSystem->AddGroupFilter(GroupHandle);
+	
+		GroupHandles[It] = GroupHandle;	
+	}
+}
+
+UE_NET_TEST_FIXTURE(FTestGroupsFixture, GroupFilterAPIObjectMemberOfMoreThanOneGroup)
+{
+	// Spawn object on server
+	UReplicatedTestObject* ServerObject0 = Server->CreateObject(0, 0);
+	const FNetHandle ServerHandle = ServerObject0->NetHandle;
+
+	// Create groups
+	FNetObjectGroupHandle GroupHandle0 = Server->ReplicationSystem->CreateGroup();
+	FNetObjectGroupHandle GroupHandle1 = Server->ReplicationSystem->CreateGroup();
+	UE_NET_ASSERT_NE(0U, (uint32)GroupHandle0);	
+	UE_NET_ASSERT_NE(0U, (uint32)GroupHandle1);	
+
+	// Add Objects to groups
+	Server->ReplicationSystem->AddToGroup(GroupHandle0, ServerHandle);
+	Server->ReplicationSystem->AddToGroup(GroupHandle1, ServerHandle);
+
+	// Add client
+	FReplicationSystemTestClient* Client = CreateClient();
+
+	// Update filters
+	Filter();
+
+	const FNetBitArrayView ObjectGroupFilter = ServerFiltering->GetGroupFilteredOutObjects(Client->ConnectionIdOnServer);
+
+	// Verify that object is not filtered out
+	VerifyObjectFilterStatus(ServerHandle, Client->ConnectionIdOnServer, true);
+
+	// Add groups to group filter
+	Server->ReplicationSystem->AddGroupFilter(GroupHandle0);
+	Server->ReplicationSystem->AddGroupFilter(GroupHandle1);
+
+	// Filter out objects in the group
+	Server->ReplicationSystem->SetGroupFilterStatus(GroupHandle0, Client->ConnectionIdOnServer, ENetFilterStatus::Disallow);
+	Server->ReplicationSystem->SetGroupFilterStatus(GroupHandle1, Client->ConnectionIdOnServer, ENetFilterStatus::Disallow);
+
+	// Update filters
+	Filter();
+
+	// We now expect the object to be filtered out
+	VerifyObjectFilterStatus(ServerHandle, Client->ConnectionIdOnServer, false);
+
+	// Allow objects from group 0
+	Server->ReplicationSystem->SetGroupFilterStatus(GroupHandle0, Client->ConnectionIdOnServer, ENetFilterStatus::Allow);
+
+	// Update filters
+	Filter();
+
+	// We still  expect the object to be filtered out since the object is a member of two groups
+	VerifyObjectFilterStatus(ServerHandle, Client->ConnectionIdOnServer, false);
+
+	// Allow objects from group 1
+	Server->ReplicationSystem->SetGroupFilterStatus(GroupHandle1, Client->ConnectionIdOnServer, ENetFilterStatus::Allow);
+
+	// Update filters
+	Filter();
+
+	// We now expect the object to no longer be filtered out
+	VerifyObjectFilterStatus(ServerHandle, Client->ConnectionIdOnServer, true);
+}
+
+// Group filtering tests
+UE_NET_TEST_FIXTURE(FTestGroupsFixture, GroupFilterAPINotFilteredGroup)
+{
+	// Spawn object on server
+	UReplicatedTestObject* ServerObject0 = Server->CreateObject(0, 0);
+	UReplicatedTestObject* ServerObject1 = Server->CreateObject(0, 0);
+	UReplicatedTestObject* ServerObject2 = Server->CreateObject(0, 0);
+
+	// Add Objects to NotFilteredGroup
+	Server->ReplicationSystem->AddToGroup(NotReplicatedNetObjectGroupHandle, ServerObject0->NetHandle);
+
+	// Add client
+	FReplicationSystemTestClient* Client = CreateClient();
+
+	// Update filters
+	Filter();
+
+	const FNetBitArrayView ObjectGroupFilter = ServerFiltering->GetGroupFilteredOutObjects(Client->ConnectionIdOnServer);
+
+	// We expect the filter to filter out the object	
+	UE_NET_ASSERT_TRUE(ObjectGroupFilter.IsAnyBitSet());
+
+	// Remove object from NotReplicatedNetObjectGroup
+	Server->ReplicationSystem->RemoveFromGroup(NotReplicatedNetObjectGroupHandle, ServerObject0->NetHandle);
+
+	// Update filters
+	Filter();
+
+	// We now expect no filtering
+	UE_NET_ASSERT_TRUE(ObjectGroupFilter.IsNoBitSet());
+}
+
+// Verify that group filtering on subobjects works as expected
+UE_NET_TEST_FIXTURE(FTestGroupsFixture, GroupFilterAPISubObjects)
+{
+	// Spawn object on server
+	UReplicatedTestObject* ServerObject = Server->CreateObject(0, 0);
+	const FNetHandle ServerHandle = ServerObject->NetHandle;
+
+	UReplicatedTestObject* ServerSubObject = Server->CreateSubObject(ServerHandle, 0, 0);	
+	const FNetHandle ServerSubObjectHandle = ServerSubObject->NetHandle;
+	
+	// Create groups
+	FNetObjectGroupHandle GroupHandle = Server->ReplicationSystem->CreateGroup();
+	Server->ReplicationSystem->AddGroupFilter(GroupHandle);
+
+	// Add Object to group
+	Server->ReplicationSystem->AddToGroup(GroupHandle, ServerHandle);
+
+	// Add client
+	FReplicationSystemTestClient* Client = CreateClient();
+
+	// Update filters
+	Filter();
+
+	// Verify that both object and subobject are filtered out
+	const FNetBitArrayView ObjectGroupFilter = ServerFiltering->GetGroupFilteredOutObjects(Client->ConnectionIdOnServer);
+
+	const bool bExpectsReplicationToBeAllowed = false;
+	VerifyObjectFilterStatus(ServerHandle, Client->ConnectionIdOnServer, bExpectsReplicationToBeAllowed);
+	VerifyObjectFilterStatus(ServerSubObjectHandle, Client->ConnectionIdOnServer, bExpectsReplicationToBeAllowed);
+}
+
+// Verify that group filtering on subobjects works as expected and that it is restored if subobject is removed same frame
+UE_NET_TEST_FIXTURE(FTestGroupsFixture, GroupFilterAPIRemovingSubObjectsRestoresGroupFilteredOutObjects)
+{
+	// Create groups
+	FNetObjectGroupHandle GroupHandle = Server->ReplicationSystem->CreateGroup();
+	Server->ReplicationSystem->AddGroupFilter(GroupHandle);
+
+	// Add client
+	FReplicationSystemTestClient* Client = CreateClient();
+
+	// Update filters
+	Filter();
+
+	// Spawn object on server
+	UReplicatedTestObject* ServerObject = Server->CreateObject(0, 0);
+	const FNetHandle ServerHandle = ServerObject->NetHandle;
+
+	// Add Object to group
+	Server->ReplicationSystem->AddToGroup(GroupHandle, ServerHandle);
+
+
+	// Create sub object
+	UReplicatedTestObject* ServerSubObject = Server->CreateSubObject(ServerHandle, 0, 0);
+	const FNetHandle ServerSubObjectHandle = ServerSubObject->NetHandle;
+
+	// destroy subobject in the same frame
+	Server->DestroyObject(ServerSubObject);
+
+	// Update filters
+	Filter();
+
+	// Verify that both object and subobject are filtered out
+	const FNetBitArrayView ObjectGroupFilter = ServerFiltering->GetGroupFilteredOutObjects(Client->ConnectionIdOnServer);
+
+	const bool bExpectsReplicationToBeAllowed = false;
+	VerifyObjectFilterStatus(ServerHandle, Client->ConnectionIdOnServer, bExpectsReplicationToBeAllowed);
+}
+
+
+
+UE_NET_TEST_FIXTURE(FTestGroupsFixture, GroupFilterAPILateAddedSubObjectIsFilteredOut)
+{
+	// Spawn object on server
+	UReplicatedTestObject* ServerObject = Server->CreateObject(0, 0);
+	const FNetHandle ServerHandle = ServerObject->NetHandle;
+	
+	// Create groups
+	FNetObjectGroupHandle GroupHandle = Server->ReplicationSystem->CreateGroup();
+	Server->ReplicationSystem->AddGroupFilter(GroupHandle);
+
+	// Add Object to group
+	Server->ReplicationSystem->AddToGroup(GroupHandle, ServerHandle);
+
+	// Add client
+	FReplicationSystemTestClient* Client = CreateClient();
+
+	// Update filters
+	Filter();
+
+	// Add subobject
+	UReplicatedTestObject* ServerSubObject = Server->CreateSubObject(ServerHandle, 0, 0);	
+	const FNetHandle ServerSubObjectHandle = ServerSubObject->NetHandle;
+
+	// Update filters
+	Filter();
+
+	// Verify that both object and subobject are filtered out
+	const FNetBitArrayView ObjectGroupFilter = ServerFiltering->GetGroupFilteredOutObjects(Client->ConnectionIdOnServer);
+
+	const bool bExpectsReplicationToBeAllowed = false;
+	VerifyObjectFilterStatus(ServerHandle, Client->ConnectionIdOnServer, bExpectsReplicationToBeAllowed);
+	VerifyObjectFilterStatus(ServerSubObjectHandle, Client->ConnectionIdOnServer, bExpectsReplicationToBeAllowed);
+}
+
+UE_NET_TEST_FIXTURE(FTestGroupsFixture, GroupFilterAPISubObjectIsFilteredOutIfOwnerIsFilterIsChanged)
+{
+	// Spawn object on server
+	UReplicatedTestObject* ServerObject = Server->CreateObject(0, 0);
+	const FNetHandle ServerHandle = ServerObject->NetHandle;
+
+	// Add subobject
+	UReplicatedTestObject* ServerSubObject = Server->CreateSubObject(ServerHandle, 0, 0);	
+	const FNetHandle ServerSubObjectHandle = ServerSubObject->NetHandle;
+	
+	// Create groups
+	FNetObjectGroupHandle GroupHandle = Server->ReplicationSystem->CreateGroup();
+	FNetObjectGroupHandle SubObjectGroupHandle = Server->ReplicationSystem->CreateGroup();
+
+	Server->ReplicationSystem->AddGroupFilter(GroupHandle);
+	Server->ReplicationSystem->SetGroupFilterStatus(GroupHandle, ENetFilterStatus::Allow);
+
+	// Add Object to group
+	Server->ReplicationSystem->AddToGroup(GroupHandle, ServerHandle);
+
+	// Add client
+	FReplicationSystemTestClient* Client = CreateClient();
+
+	// Update filters
+	Filter();
+
+	// Verify that we allow replication for both object and subobject
+	{
+		const bool bExpectsReplicationToBeAllowed = true;
+		VerifyObjectFilterStatus(ServerHandle, Client->ConnectionIdOnServer, bExpectsReplicationToBeAllowed);
+		VerifyObjectFilterStatus(ServerSubObjectHandle, Client->ConnectionIdOnServer, bExpectsReplicationToBeAllowed);
+	}
+
+	// Disallow group filter
+	Server->ReplicationSystem->SetGroupFilterStatus(GroupHandle, ENetFilterStatus::Disallow);
+
+	// Update filters
+	Filter();
+
+	// Verify that both object and subobject are filtered out
+	{
+		const bool bExpectsReplicationToBeAllowed = false;
+		VerifyObjectFilterStatus(ServerHandle, Client->ConnectionIdOnServer, bExpectsReplicationToBeAllowed);
+		VerifyObjectFilterStatus(ServerSubObjectHandle, Client->ConnectionIdOnServer, bExpectsReplicationToBeAllowed);
+	}
+}
+
+UE_NET_TEST_FIXTURE(FTestGroupsFixture, GroupFilterAPISubObjectIsFilteredOutWithSeparateFilter)
+{
+	// Spawn object on server
+	UReplicatedTestObject* ServerObject = Server->CreateObject(0, 0);
+	const FNetHandle ServerHandle = ServerObject->NetHandle;
+
+	// Add subobject
+	UReplicatedTestObject* ServerSubObject = Server->CreateSubObject(ServerHandle, 0, 0);	
+	const FNetHandle ServerSubObjectHandle = ServerSubObject->NetHandle;
+	
+	// Create groups
+	FNetObjectGroupHandle GroupHandle = Server->ReplicationSystem->CreateGroup();
+	FNetObjectGroupHandle SubObjectGroupHandle = Server->ReplicationSystem->CreateGroup();
+
+	Server->ReplicationSystem->AddGroupFilter(GroupHandle);
+	Server->ReplicationSystem->AddGroupFilter(SubObjectGroupHandle);
+	Server->ReplicationSystem->SetGroupFilterStatus(GroupHandle, ENetFilterStatus::Allow);
+
+	// Add Object to group
+	Server->ReplicationSystem->AddToGroup(GroupHandle, ServerHandle);
+	Server->ReplicationSystem->AddToGroup(SubObjectGroupHandle, ServerSubObjectHandle);
+
+	// Add client
+	FReplicationSystemTestClient* Client = CreateClient();
+
+	// Update filters
+	Filter();
+
+	// Verify that we only allow replication for owner
+	{
+		VerifyObjectFilterStatus(ServerHandle, Client->ConnectionIdOnServer, true);
+		VerifyObjectFilterStatus(ServerSubObjectHandle, Client->ConnectionIdOnServer, false);
+	}
+
+	// Enable SubObjectGroupHandle
+	Server->ReplicationSystem->SetGroupFilterStatus(SubObjectGroupHandle, ENetFilterStatus::Allow);
+
+	// Update filters
+	Filter();
+
+	// Verify that both object and subobject are enabled
+	{
+		const bool bExpectsReplicationToBeAllowed = true;
+		VerifyObjectFilterStatus(ServerHandle, Client->ConnectionIdOnServer, bExpectsReplicationToBeAllowed);
+		VerifyObjectFilterStatus(ServerSubObjectHandle, Client->ConnectionIdOnServer, bExpectsReplicationToBeAllowed);
+	}
+}
+
+}
diff --git a/Engine/Plugins/Runtime/ReplicationSystemTestPlugin/Source/Private/Tests/ReplicationSystem/ReplicatedTestObject.cpp b/Engine/Plugins/Runtime/ReplicationSystemTestPlugin/Source/Private/Tests/ReplicationSystem/ReplicatedTestObject.cpp
new file mode 100644
--- /dev/null
+++ b/Engine/Plugins/Runtime/ReplicationSystemTestPlugin/Source/Private/Tests/ReplicationSystem/ReplicatedTestObject.cpp
@@ -0,0 +1,505 @@
+// Copyright Epic Games, Inc. All Rights Reserved.
+
+#include "ReplicatedTestObject.h"
+#include "NetworkAutomationTest.h"
+#include "NetworkAutomationTestMacros.h"
+#include "Iris/ReplicationState/PropertyReplicationState.h"
+#include "Iris/ReplicationSystem/NetHandle.h"
+#include "Iris/ReplicationSystem/RepTag.h"
+#include "Iris/ReplicationSystem/ReplicationProtocol.h"
+#include "Iris/ReplicationSystem/ReplicationBridge.h"
+#include "Iris/ReplicationSystem/ReplicationSystem.h"
+#include "Iris/ReplicationSystem/ReplicationSystemInternal.h"
+#include "Iris/ReplicationState/InternalPropertyReplicationState.h"
+#include "Iris/ReplicationSystem/ReplicationFragmentUtil.h"
+#include "Iris/ReplicationState/ReplicationStateDescriptorBuilder.h"
+#include "Iris/ReplicationState/ReplicationStateDescriptorImplementationMacros.h"
+#include "Iris/Serialization/NetBitStreamWriter.h"
+#include "Iris/Serialization/NetBitStreamReader.h"
+#include "Iris/Serialization/NetSerializers.h"
+#include "Iris/Core/MemoryLayoutUtil.h"
+#include "Templates/SharedPointer.h"
+#include "Containers/Map.h"
+#include "UObject/StrongObjectPtr.h"
+#include "Net/UnrealNetwork.h"
+#include "Iris/Serialization/NetBitStreamUtil.h"
+
+//////////////////////////////////////////////////////////////////////////
+// Implementation for UReplicatedTestObjectBridge
+//////////////////////////////////////////////////////////////////////////
+
+UReplicatedTestObjectBridge::UReplicatedTestObjectBridge()
+: UObjectReplicationBridge()
+, CreatedObjectsOnNode(nullptr)
+{
+}
+
+const UE::Net::FReplicationInstanceProtocol* UReplicatedTestObjectBridge::GetReplicationInstanceProtocol(FNetHandle Handle) const
+{
+	const UE::Net::Private::FNetHandleManager* LocalNetHandleManager = &GetReplicationSystem()->GetReplicationSystemInternal()->GetNetHandleManager();
+
+	if (uint32 InternalObjectIndex = LocalNetHandleManager->GetInternalIndex(Handle))
+	{
+		return LocalNetHandleManager->GetReplicatedObjectDataNoCheck(InternalObjectIndex).InstanceProtocol;
+	}
+
+	return nullptr;
+};
+
+UE::Net::FNetHandle UReplicatedTestObjectBridge::BeginReplication(UReplicatedTestObject* Instance)
+{
+	// Create NetHandle for the registered fragments
+	Super::FCreateNetHandleParams Params = Super::DefaultCreateNetHandleParams;
+	Params.bCanReceive = true;
+
+	FNetHandle Handle = Super::BeginReplication(Instance, Params);
+
+	// This is optional but typically we want to cache at least the NetHandle in the game instance to avoid doing map lookups to find it
+	if (Handle.IsValid())
+	{
+		Instance->NetHandle = Handle;
+	}
+
+	return Handle;
+}
+
+UE::Net::FNetHandle UReplicatedTestObjectBridge::BeginReplication(FNetHandle OwnerHandle, UReplicatedTestObject* Instance, FNetHandle InsertRelativeToSubObjectHandle, ESubObjectInsertionOrder InsertionOrder)
+{
+	// Create NetHandle for the registered fragments
+	Super::FCreateNetHandleParams Params = Super::DefaultCreateNetHandleParams;
+	Params.bCanReceive = true;
+
+	// Create NetHandle for the registered fragments
+	FNetHandle Handle = Super::BeginReplication(OwnerHandle, Instance, InsertRelativeToSubObjectHandle, InsertionOrder, Params);
+
+	if (Handle.IsValid())
+	{
+		// This is optional but typically we want to cache at least the NetHandle in the game instance to avoid doing map lookups to find it
+		Instance->NetHandle = Handle;
+	}
+	
+	return Handle;
+}
+
+
+bool UReplicatedTestObjectBridge::WriteCreationHeader(UE::Net::FNetSerializationContext& Context, FNetHandle Handle)
+{
+	UE::Net::FNetBitStreamWriter& Writer = *Context.GetBitStreamWriter();
+
+	uint16 NumComponentsToSpawn = 0U;
+	uint16 NumIrisComponentsToSpawn = 0U;
+	uint16 NumDynamicComponentsToSpawn = 0U;
+	uint16 NumConnectionFilteredComponentsToSpawn = 0U;
+
+	const UObject* Object = GetReplicatedObject(Handle);
+	
+	if (const UTestReplicatedIrisObject* TestReplicatedIrisObject = Cast<UTestReplicatedIrisObject>(Object))
+	{
+		NumComponentsToSpawn = TestReplicatedIrisObject->Components.Num();
+		NumIrisComponentsToSpawn = TestReplicatedIrisObject->IrisComponents.Num();
+		NumDynamicComponentsToSpawn = TestReplicatedIrisObject->DynamicStateComponents.Num();
+		NumConnectionFilteredComponentsToSpawn = TestReplicatedIrisObject->ConnectionFilteredComponents.Num();
+	}
+
+	UObject* Archetype = Object->GetArchetype();
+	check(Archetype);
+
+	WriteString(&Writer, Archetype->GetPathName());
+	Writer.WriteBits(NumComponentsToSpawn, 16);
+	Writer.WriteBits(NumIrisComponentsToSpawn, 16);
+	Writer.WriteBits(NumDynamicComponentsToSpawn, 16);
+	Writer.WriteBits(NumConnectionFilteredComponentsToSpawn, 16);
+
+	return !Writer.IsOverflown();
+}
+
+UObjectReplicationBridge::FCreationHeader* UReplicatedTestObjectBridge::ReadCreationHeader(UE::Net::FNetSerializationContext& Context)
+{
+	UE::Net::FNetBitStreamReader& Reader = *Context.GetBitStreamReader();
+	TUniquePtr<FReplicationTestObjectCreationHeader> Header(new FReplicationTestObjectCreationHeader);
+
+	ReadString(&Reader, Header->ArchetypeName);
+
+	//FNetObjectReference ArcheTypeRef;
+
+	Header->NumComponentsToSpawn = Reader.ReadBits(16);
+	Header->NumIrisComponentsToSpawn = Reader.ReadBits(16);
+	Header->NumDynamicComponentsToSpawn = Reader.ReadBits(16);
+	Header->NumConnectionFilteredComponentsToSpawn = Reader.ReadBits(16);
+
+	if (Reader.IsOverflown())
+	{
+		return nullptr;
+	}
+
+	return Header.Release();
+}
+
+UObject* UReplicatedTestObjectBridge::BeginInstantiateFromRemote(FNetHandle SubObjectOwnerHandle, const UE::Net::FNetObjectResolveContext& ResolveContext, const FCreationHeader* InHeader)
+{
+	const FReplicationTestObjectCreationHeader* Header = static_cast<const FReplicationTestObjectCreationHeader*>(InHeader);
+
+	UObject* ArcheType = StaticFindObject(UObject::StaticClass(), nullptr, *Header->ArchetypeName, false);
+
+	check(ArcheType);
+
+	FStaticConstructObjectParameters ConstructObjectParameters(ArcheType->GetClass());
+	UObject* CreatedObject = StaticConstructObject_Internal(ConstructObjectParameters);
+	
+	if (UTestReplicatedIrisObject* CreatedTestObject = Cast<UTestReplicatedIrisObject>(CreatedObject))
+	{
+		UTestReplicatedIrisObject::FComponents Components;
+		Components.PropertyComponentCount = Header->NumComponentsToSpawn;
+		Components.IrisComponentCount = Header->NumIrisComponentsToSpawn;
+		Components.DynamicStateComponentCount = Header->NumDynamicComponentsToSpawn;
+		Components.ConnectionFilteredComponentCount = Header->NumConnectionFilteredComponentsToSpawn;
+
+		CreatedTestObject->AddComponents(Components);
+	}
+
+	// Store the object so that we can find detached/torn off instances from tests
+	if (CreatedObjectsOnNode)
+	{
+		CreatedObjectsOnNode->Add(TStrongObjectPtr<UObject>(CreatedObject));
+	}
+	
+	return CreatedObject;
+}
+
+void UReplicatedTestObjectBridge::EndInstantiateFromRemote(FNetHandle Handle)
+{
+	UReplicatedTestObject* Instance = CastChecked<UReplicatedTestObject>(GetReplicatedObject(Handle));
+
+	Instance->NetHandle = Handle;
+}
+
+void UReplicatedTestObjectBridge::DestroyInstanceFromRemote(UObject* Instance, bool bTearOff)
+{
+	if (Instance && !bTearOff)
+	{
+		// Remove the object from the created objects on the node
+		if (CreatedObjectsOnNode)
+		{
+			CreatedObjectsOnNode->Remove(TStrongObjectPtr<UObject>(Instance));
+		}
+
+		Instance->PreDestroyFromReplication();
+		Instance->MarkAsGarbage();
+	}
+}
+
+void UReplicatedTestObjectBridge::SetPollFramePeriod(UReplicatedTestObject* Instance, uint8 FramePeriod)
+{
+	using namespace UE::Net;
+	using namespace UE::Net::Private;
+
+	if (Instance == nullptr)
+	{
+		return;
+	}
+	const FNetHandleManager* LocalNetHandleManager = &GetReplicationSystem()->GetReplicationSystemInternal()->GetNetHandleManager();
+	const uint32 InternalObjectIndex = LocalNetHandleManager->GetInternalIndex(Instance->NetHandle);
+	if (InternalObjectIndex != FNetHandleManager::InvalidInternalIndex)
+	{
+		UObjectReplicationBridge::SetPollFramePeriod(InternalObjectIndex, FramePeriod);
+	}
+}
+
+//////////////////////////////////////////////////////////////////////////
+// Implementation for UTestReplicatedIrisPropertyComponent
+//////////////////////////////////////////////////////////////////////////
+UTestReplicatedIrisPropertyComponent::UTestReplicatedIrisPropertyComponent() : UObject()
+{
+}
+
+void UTestReplicatedIrisPropertyComponent::GetLifetimeReplicatedProps( TArray< class FLifetimeProperty > & OutLifetimeProps ) const
+{
+	DOREPLIFETIME(UTestReplicatedIrisPropertyComponent, IntA);
+	DOREPLIFETIME(UTestReplicatedIrisPropertyComponent, StructWithStructWithTag);
+	DOREPLIFETIME_CONDITION(UTestReplicatedIrisPropertyComponent, IntB, COND_InitialOnly);
+}
+
+//////////////////////////////////////////////////////////////////////////
+// Implementation for UTestReplicatedIrisPushModelComponentWithObjectReference
+//////////////////////////////////////////////////////////////////////////
+UTestReplicatedIrisPushModelComponentWithObjectReference::UTestReplicatedIrisPushModelComponentWithObjectReference() : UObject()
+{
+}
+
+void UTestReplicatedIrisPushModelComponentWithObjectReference::GetLifetimeReplicatedProps( TArray< class FLifetimeProperty > & OutLifetimeProps ) const
+{
+	FDoRepLifetimeParams LifetimeParams;
+	LifetimeParams.bIsPushBased = true;
+
+	DOREPLIFETIME_WITH_PARAMS(ThisClass, IntA, LifetimeParams);
+	DOREPLIFETIME_WITH_PARAMS(ThisClass, RawObjectPtrRef, LifetimeParams);
+	DOREPLIFETIME_WITH_PARAMS(ThisClass, WeakObjectPtrObjectRef, LifetimeParams);
+}
+
+//////////////////////////////////////////////////////////////////////////
+// Implementation for UTestReplicatedIrisDynamicStatePropertyComponent
+//////////////////////////////////////////////////////////////////////////
+UTestReplicatedIrisDynamicStatePropertyComponent::UTestReplicatedIrisDynamicStatePropertyComponent() : UObject()
+{
+}
+
+void UTestReplicatedIrisDynamicStatePropertyComponent::GetLifetimeReplicatedProps( TArray< class FLifetimeProperty > & OutLifetimeProps ) const
+{
+	DOREPLIFETIME(UTestReplicatedIrisDynamicStatePropertyComponent, IntArray);
+	DOREPLIFETIME(UTestReplicatedIrisDynamicStatePropertyComponent, IntStaticArray);
+}
+
+//////////////////////////////////////////////////////////////////////////
+// Implementation for UTestReplicatedIrisLifetimeConditionalsPropertyState
+//////////////////////////////////////////////////////////////////////////
+UTestReplicatedIrisLifetimeConditionalsPropertyState::UTestReplicatedIrisLifetimeConditionalsPropertyState() : UObject()
+{
+}
+
+void UTestReplicatedIrisLifetimeConditionalsPropertyState::GetLifetimeReplicatedProps( TArray< class FLifetimeProperty > & OutLifetimeProps ) const
+{
+	DOREPLIFETIME_CONDITION(ThisClass, ToOwnerA, COND_OwnerOnly);
+	DOREPLIFETIME_CONDITION(ThisClass, ToOwnerB, COND_OwnerOnly);
+
+	DOREPLIFETIME_CONDITION(ThisClass, SkipOwnerA, COND_SkipOwner);
+	DOREPLIFETIME_CONDITION(ThisClass, SkipOwnerB, COND_SkipOwner);
+
+	DOREPLIFETIME_CONDITION(ThisClass, SimulatedOnlyInt, COND_SimulatedOnly);
+	DOREPLIFETIME_CONDITION(ThisClass, AutonomousOnlyInt, COND_AutonomousOnly);
+	DOREPLIFETIME_CONDITION(ThisClass, SimulatedOrPhysicsInt, COND_SimulatedOrPhysics);
+	DOREPLIFETIME_CONDITION(ThisClass, SimulatedOnlyNoReplayInt, COND_SimulatedOnlyNoReplay);
+	DOREPLIFETIME_CONDITION(ThisClass, SimulatedOrPhysicsNoReplayInt, COND_SimulatedOrPhysicsNoReplay);
+}
+
+//////////////////////////////////////////////////////////////////////////
+// Implementation for UTestReplicationSystem_TestIrisComponent
+//////////////////////////////////////////////////////////////////////////
+FTestReplicatedIrisComponent::FTestReplicatedIrisComponent()
+: ReplicationFragment(*this, ReplicationState)
+{
+}
+
+void FTestReplicatedIrisComponent::ApplyReplicationState(const FFakeGeneratedReplicationState& State, UE::Net::FReplicationStateApplyContext& Context)
+{
+	ReplicationState = State;
+}
+
+//////////////////////////////////////////////////////////////////////////
+// Implementation for UTestReplicationSystem_TestClass
+//////////////////////////////////////////////////////////////////////////
+void UTestReplicatedIrisObject::GetLifetimeReplicatedProps( TArray< class FLifetimeProperty > & OutLifetimeProps ) const
+{
+}
+
+UTestReplicatedIrisObject::UTestReplicatedIrisObject()
+: UReplicatedTestObject()
+{
+}
+
+void UTestReplicatedIrisObject::AddComponents(const UTestReplicatedIrisObject::FComponents& InComponents)
+{
+	check(!NetHandle.IsValid())
+	for (uint32 It = 0; It < InComponents.PropertyComponentCount; ++It)
+	{
+		Components.Emplace(NewObject<UTestReplicatedIrisPropertyComponent>());
+	}
+
+	for (uint32 It = 0; It < InComponents.IrisComponentCount; ++It)
+	{
+		IrisComponents.Emplace(TUniquePtr<FTestReplicatedIrisComponent>(new FTestReplicatedIrisComponent()));
+	}
+
+	for (uint32 It = 0; It < InComponents.DynamicStateComponentCount; ++It)
+	{
+		DynamicStateComponents.Emplace(NewObject<UTestReplicatedIrisDynamicStatePropertyComponent>());
+	}
+
+	for (uint32 It = 0; It < InComponents.ConnectionFilteredComponentCount; ++It)
+	{
+		ConnectionFilteredComponents.Emplace(NewObject<UTestReplicatedIrisLifetimeConditionalsPropertyState>());
+	}
+
+	for (uint32 It = 0; It < InComponents.ObjectReferenceComponentCount; ++It)
+	{
+		ObjectReferenceComponents.Emplace(NewObject<UTestReplicatedIrisPushModelComponentWithObjectReference>());
+	}
+}
+
+void UTestReplicatedIrisObject::AddComponents(uint32 PropertyComponentCount, uint32 IrisComponentCount)
+{
+	check(!NetHandle.IsValid())
+	// Setup a few components
+	for (uint32 It = 0; It < PropertyComponentCount; ++It)
+	{
+		Components.Emplace(NewObject<UTestReplicatedIrisPropertyComponent>());
+	}
+
+	// Setup a few components
+	for (uint32 It = 0; It < IrisComponentCount; ++It)
+	{
+		IrisComponents.Emplace(TUniquePtr<FTestReplicatedIrisComponent>(new FTestReplicatedIrisComponent()));
+	}
+}
+
+void UTestReplicatedIrisObject::AddDynamicStateComponents(uint32 DynamicStateComponentCount)
+{
+	check(!NetHandle.IsValid())
+	for (uint32 It = 0; It < DynamicStateComponentCount; ++It)
+	{
+		DynamicStateComponents.Emplace(NewObject<UTestReplicatedIrisDynamicStatePropertyComponent>());
+	}
+}
+
+void UTestReplicatedIrisObject::RegisterReplicationFragments(UE::Net::FFragmentRegistrationContext& Context, UE::Net::EFragmentRegistrationFlags RegistrationFlags)
+{
+	using namespace UE::Net;
+
+	Super::RegisterReplicationFragments(Context, RegistrationFlags);
+
+	// Base object owns the fragment in this case
+	{
+		this->ReplicationFragments.Reset();
+		FReplicationFragmentUtil::CreateAndRegisterFragmentsForObject(this, Context, RegistrationFlags, &this->ReplicationFragments);
+	}
+
+	// Register components from "Components" as well, in this case we let the replication system own the fragments
+	for (const auto& Component : Components)
+	{
+		Component->ReplicationFragments.Reset();
+		FReplicationFragmentUtil::CreateAndRegisterFragmentsForObject(Component.Get(), Context, RegistrationFlags, &Component->ReplicationFragments);
+	}	
+
+	// Register components from "IrisComponents" as well
+	for (const auto& Component : IrisComponents)
+	{
+		Component->ReplicationFragment.Register(Context);
+	}	
+
+	// Register components from "DynamicStateComponents"
+	for (const auto& Component : DynamicStateComponents)
+	{
+		Component->ReplicationFragments.Reset();
+		FReplicationFragmentUtil::CreateAndRegisterFragmentsForObject(Component.Get(), Context, RegistrationFlags, &Component->ReplicationFragments);
+	}	
+
+	// Register components from "ConnectionFilteredComponents"
+	for (const auto& Component : ConnectionFilteredComponents)
+	{
+		Component->ReplicationFragments.Reset();
+		FReplicationFragmentUtil::CreateAndRegisterFragmentsForObject(Component.Get(), Context, RegistrationFlags, &Component->ReplicationFragments);
+	}	
+
+	// Register components from "ObjectReferenceComponents"
+	for (const auto& Component : ObjectReferenceComponents)
+	{
+		Component->ReplicationFragments.Reset();
+		FReplicationFragmentUtil::CreateAndRegisterFragmentsForObject(Component.Get(), Context, RegistrationFlags, &Component->ReplicationFragments);
+	}	
+}
+
+//////////////////////////////////////////////////////////////////////////
+// Implementation for UTestReplicationSystem_TestClass
+//////////////////////////////////////////////////////////////////////////
+
+void UTestReplicatedIrisObjectWithObjectReference::GetLifetimeReplicatedProps( TArray< class FLifetimeProperty > & OutLifetimeProps ) const
+{
+}
+
+UTestReplicatedIrisObjectWithObjectReference::UTestReplicatedIrisObjectWithObjectReference()
+: UReplicatedTestObject()
+{
+}
+
+void UTestReplicatedIrisObjectWithObjectReference::RegisterReplicationFragments(UE::Net::FFragmentRegistrationContext& Context, UE::Net::EFragmentRegistrationFlags RegistrationFlags)
+{
+	Super::RegisterReplicationFragments(Context, RegistrationFlags);
+
+	// Base object owns the fragment in this case
+	{
+		this->ReplicationFragments.Reset();
+		UE::Net::FReplicationFragmentUtil::CreateAndRegisterFragmentsForObject(this, Context, RegistrationFlags, &this->ReplicationFragments);
+	}
+}
+
+uint32 UReplicatedSubObjectOrderObject::RepOrderCounter = 0U;
+
+UTestReplicatedIrisObjectWithNoReplicatedMembers::UTestReplicatedIrisObjectWithNoReplicatedMembers()
+: UReplicatedTestObject()
+{
+}
+
+void UReplicatedSubObjectOrderObject::GetLifetimeReplicatedProps(TArray< class FLifetimeProperty >& OutLifetimeProps) const
+{
+	DOREPLIFETIME(UReplicatedSubObjectOrderObject, IntA);
+	DOREPLIFETIME(UReplicatedSubObjectOrderObject, OtherSubObject);
+}
+
+UReplicatedSubObjectOrderObject::UReplicatedSubObjectOrderObject()
+	: UReplicatedTestObject()
+{
+}
+
+void UReplicatedSubObjectOrderObject::RegisterReplicationFragments(UE::Net::FFragmentRegistrationContext& Context, UE::Net::EFragmentRegistrationFlags RegistrationFlags)
+{
+	Super::RegisterReplicationFragments(Context, RegistrationFlags);
+
+	// Base object owns the fragment in this case
+	{
+		this->ReplicationFragments.Reset();
+		UE::Net::FReplicationFragmentUtil::CreateAndRegisterFragmentsForObject(this, Context, RegistrationFlags, &this->ReplicationFragments);
+	}
+}
+
+//////////////////////////////////////////////////////////////////////////
+// THIS SECTION WILL BE GENERATED FROM UHT
+//////////////////////////////////////////////////////////////////////////
+
+// FFakeGeneratedReplicationState
+const UE::Net::FRepTag RepTag_FakeGeneratedReplicationState_IntB = UE::Net::MakeRepTag("FakeGeneratedReplicationState_IntB");
+
+constexpr UE::Net::FReplicationStateMemberChangeMaskDescriptor FFakeGeneratedReplicationState::sReplicationStateChangeMaskDescriptors[3];
+
+IRIS_BEGIN_SERIALIZER_DESCRIPTOR(FFakeGeneratedReplicationState)
+IRIS_SERIALIZER_DESCRIPTOR(UE::Net::FInt32NetSerializer, nullptr)
+IRIS_SERIALIZER_DESCRIPTOR(UE::Net::FInt32NetSerializer, nullptr)
+IRIS_SERIALIZER_DESCRIPTOR(UE::Net::FInt32NetSerializer, nullptr)
+IRIS_END_SERIALIZER_DESCRIPTOR()
+
+// Member traits
+IRIS_BEGIN_TRAITS_DESCRIPTOR(FFakeGeneratedReplicationState)
+IRIS_TRAITS_DESCRIPTOR(UE::Net::EReplicationStateMemberTraits::None)
+IRIS_TRAITS_DESCRIPTOR(UE::Net::EReplicationStateMemberTraits::None)
+IRIS_TRAITS_DESCRIPTOR(UE::Net::EReplicationStateMemberTraits::None)
+IRIS_END_TRAITS_DESCRIPTOR()
+
+// This is required to work around issues with static initialization order
+IRIS_BEGIN_INTERNAL_TYPE_INFO(FFakeGeneratedReplicationState)
+IRIS_INTERNAL_TYPE_INFO(UE::Net::FInt32NetSerializer)
+IRIS_INTERNAL_TYPE_INFO(UE::Net::FInt32NetSerializer)
+IRIS_INTERNAL_TYPE_INFO(UE::Net::FInt32NetSerializer)
+IRIS_END_INTERNAL_TYPE_INFO()
+
+IRIS_BEGIN_MEMBER_DESCRIPTOR(FFakeGeneratedReplicationState)
+IRIS_MEMBER_DESCRIPTOR(FFakeGeneratedReplicationState, IntA, 0)
+IRIS_MEMBER_DESCRIPTOR(FFakeGeneratedReplicationState, IntB, 1)
+IRIS_MEMBER_DESCRIPTOR(FFakeGeneratedReplicationState, IntC, 2)
+IRIS_END_MEMBER_DESCRIPTOR()
+
+IRIS_BEGIN_MEMBER_DEBUG_DESCRIPTOR(FFakeGeneratedReplicationState)
+IRIS_MEMBER_DEBUG_DESCRIPTOR(FFakeGeneratedReplicationState, IntA)
+IRIS_MEMBER_DEBUG_DESCRIPTOR(FFakeGeneratedReplicationState, IntB)
+IRIS_MEMBER_DEBUG_DESCRIPTOR(FFakeGeneratedReplicationState, IntC)
+IRIS_END_MEMBER_DEBUG_DESCRIPTOR()
+
+IRIS_BEGIN_TAG_DESCRIPTOR(FFakeGeneratedReplicationState)
+IRIS_TAG_DESCRIPTOR(RepTag_FakeGeneratedReplicationState_IntB, 1)
+IRIS_END_TAG_DESCRIPTOR()
+
+IRIS_BEGIN_FUNCTION_DESCRIPTOR(FFakeGeneratedReplicationState)
+IRIS_END_FUNCTION_DESCRIPTOR()
+
+IRIS_BEGIN_REFERENCE_DESCRIPTOR(FFakeGeneratedReplicationState)
+IRIS_END_REFERENCE_DESCRIPTOR()
+
+IRIS_IMPLEMENT_CONSTRUCT_AND_DESTRUCT(FFakeGeneratedReplicationState)
+IRIS_IMPLEMENT_REPLICATIONSTATEDESCRIPTOR(FFakeGeneratedReplicationState)
diff --git a/Engine/Plugins/Runtime/ReplicationSystemTestPlugin/Source/Private/Tests/ReplicationSystem/TestCancelPendingDestroy.cpp b/Engine/Plugins/Runtime/ReplicationSystemTestPlugin/Source/Private/Tests/ReplicationSystem/TestCancelPendingDestroy.cpp
new file mode 100644
--- /dev/null
+++ b/Engine/Plugins/Runtime/ReplicationSystemTestPlugin/Source/Private/Tests/ReplicationSystem/TestCancelPendingDestroy.cpp
@@ -0,0 +1,744 @@
+// Copyright Epic Games, Inc. All Rights Reserved.
+
+#include "ReplicationSystemServerClientTestFixture.h"
+#include "Iris/ReplicationSystem/ReplicationSystem.h"
+#include "Iris/ReplicationSystem/Filtering/NetObjectFilter.h"
+
+
+namespace UE::Net::Private
+{
+
+class FTestCancelPendingDestroyFixture : public FReplicationSystemServerClientTestFixture
+{
+};
+
+UE_NET_TEST_FIXTURE(FTestCancelPendingDestroyFixture, TestCancelPendingDestroyDuringWaitOnCreateConfirmationWithoutPacketLoss)
+{
+	// Add a client
+	FReplicationSystemTestClient* Client = CreateClient();
+
+	// Spawn object on server
+	UTestReplicatedIrisObject* ServerObject = Server->CreateObject(UTestReplicatedIrisObject::FComponents());
+
+	// Write packet
+	Server->PreSendUpdate();
+	Server->SendUpdate(Client->ConnectionIdOnServer);
+	Server->PostSendUpdate();
+
+	// Filter out object to cause a PendingDestroy
+	Server->GetReplicationSystem()->AddToGroup(NotReplicatedNetObjectGroupHandle, ServerObject->NetHandle);
+	Server->PreSendUpdate();
+	Server->PostSendUpdate();
+
+	// Remove object from filter to cause object to end up in CancelPendingDestroy
+	Server->GetReplicationSystem()->RemoveFromGroup(NotReplicatedNetObjectGroupHandle, ServerObject->NetHandle);
+	Server->PreSendUpdate();
+	Server->PostSendUpdate();
+
+	// Deliver original packet
+	Server->DeliverTo(Client, DeliverPacket);
+
+	// Verify that the object now exists on client
+	UE_NET_ASSERT_NE(Client->GetReplicationBridge()->GetReplicatedObject(ServerObject->NetHandle), nullptr);
+
+	// Modify a property on the object and make sure it's replicated as the object should now be confirmed created
+	ServerObject->IntA ^= 1;
+
+	Server->PreSendUpdate();
+	Server->SendAndDeliverTo(Client, DeliverPacket);
+	Server->PostSendUpdate();
+
+	UTestReplicatedIrisObject* ClientObject = Cast<UTestReplicatedIrisObject>(Client->GetReplicationBridge()->GetReplicatedObject(ServerObject->NetHandle));
+	UE_NET_ASSERT_NE(ClientObject, nullptr);
+	UE_NET_ASSERT_EQ(ClientObject->IntA, ServerObject->IntA);
+}
+
+UE_NET_TEST_FIXTURE(FTestCancelPendingDestroyFixture, TestCancelPendingDestroyDuringWaitOnCreateConfirmationWithPacketLoss)
+{
+	// Add a client
+	FReplicationSystemTestClient* Client = CreateClient();
+
+	// Spawn object on server
+	UTestReplicatedIrisObject* ServerObject = Server->CreateObject(UTestReplicatedIrisObject::FComponents());
+
+	// Write packet
+	Server->PreSendUpdate();
+	Server->SendUpdate(Client->ConnectionIdOnServer);
+	Server->PostSendUpdate();
+
+	// Filter out object to cause a PendingDestroy
+	Server->GetReplicationSystem()->AddToGroup(NotReplicatedNetObjectGroupHandle, ServerObject->NetHandle);
+	Server->PreSendUpdate();
+	Server->PostSendUpdate();
+
+	// Remove object from filter to cause object to end up in CancelPendingDestroy
+	Server->GetReplicationSystem()->RemoveFromGroup(NotReplicatedNetObjectGroupHandle, ServerObject->NetHandle);
+	Server->PreSendUpdate();
+	Server->PostSendUpdate();
+
+	// Cause packet loss on object creation
+	Server->DeliverTo(Client, DoNotDeliverPacket);
+
+	// Write and send packet and verify object is created
+	Server->PreSendUpdate();
+	Server->SendAndDeliverTo(Client, DeliverPacket);
+	Server->PostSendUpdate();
+
+	UE_NET_ASSERT_NE(Client->GetReplicationBridge()->GetReplicatedObject(ServerObject->NetHandle), nullptr);
+}
+
+UE_NET_TEST_FIXTURE(FTestCancelPendingDestroyFixture, TestCancelPendingDestroyDuringCreated)
+{
+	// Add a client
+	FReplicationSystemTestClient* Client = CreateClient();
+
+	// Spawn object on server
+	UTestReplicatedIrisObject* ServerObject = Server->CreateObject(UTestReplicatedIrisObject::FComponents());
+
+	// Write and send packet
+	Server->PreSendUpdate();
+	Server->SendAndDeliverTo(Client, DeliverPacket);
+	Server->PostSendUpdate();
+
+	// Verify that the object now also exists on the client
+	UE_NET_ASSERT_NE(Client->GetReplicationBridge()->GetReplicatedObject(ServerObject->NetHandle), nullptr);
+
+	// Filter out object to cause a PendingDestroy
+	Server->GetReplicationSystem()->AddToGroup(NotReplicatedNetObjectGroupHandle, ServerObject->NetHandle);
+	Server->PreSendUpdate();
+	Server->PostSendUpdate();
+
+	// Remove object from filter to cause object to end up in CancelPendingDestroy
+	Server->GetReplicationSystem()->RemoveFromGroup(NotReplicatedNetObjectGroupHandle, ServerObject->NetHandle);
+	Server->PreSendUpdate();
+	Server->PostSendUpdate();
+
+	// Modify a property on the object and make sure it's replicated as the object should still be created
+	ServerObject->IntA ^= 1;
+
+	Server->PreSendUpdate();
+	Server->SendAndDeliverTo(Client, DeliverPacket);
+	Server->PostSendUpdate();
+
+	UTestReplicatedIrisObject* ClientObject = Cast<UTestReplicatedIrisObject>(Client->GetReplicationBridge()->GetReplicatedObject(ServerObject->NetHandle));
+	UE_NET_ASSERT_NE(ClientObject, nullptr);
+	UE_NET_ASSERT_EQ(ClientObject->IntA, ServerObject->IntA);
+}
+
+UE_NET_TEST_FIXTURE(FTestCancelPendingDestroyFixture, TestCancelPendingDestroyDuringWaitOnDestroyConfirmationWithPacketLoss)
+{
+	// Add a client
+	FReplicationSystemTestClient* Client = CreateClient();
+
+	// Spawn object on server
+	UTestReplicatedIrisObject* ServerObject = Server->CreateObject(UTestReplicatedIrisObject::FComponents());
+
+	// Write and send packet
+	Server->PreSendUpdate();
+	Server->SendAndDeliverTo(Client, DeliverPacket);
+	Server->PostSendUpdate();
+
+	// Filter out object and write packet to cause object to end up in WaitOnDestroyConfirmation
+	Server->GetReplicationSystem()->AddToGroup(NotReplicatedNetObjectGroupHandle, ServerObject->NetHandle);
+	Server->PreSendUpdate();
+	Server->SendUpdate(Client->ConnectionIdOnServer);
+	Server->PostSendUpdate();
+
+	// Remove object from filter to cause object to end up in CancelPendingDestroy
+	Server->GetReplicationSystem()->RemoveFromGroup(NotReplicatedNetObjectGroupHandle, ServerObject->NetHandle);
+	Server->PreSendUpdate();
+	Server->PostSendUpdate();
+
+	// Cause packet loss on the object destroy packet
+	Server->DeliverTo(Client, DoNotDeliverPacket);
+
+	// Verify that object still exists on client
+	UE_NET_ASSERT_NE(Client->GetReplicationBridge()->GetReplicatedObject(ServerObject->NetHandle), nullptr);
+
+	// Modify a property on the object and make sure it's replicated as the object should still be created
+	ServerObject->IntA ^= 1;
+
+	Server->PreSendUpdate();
+	Server->SendAndDeliverTo(Client, DeliverPacket);
+	Server->PostSendUpdate();
+
+	UTestReplicatedIrisObject* ClientObject = Cast<UTestReplicatedIrisObject>(Client->GetReplicationBridge()->GetReplicatedObject(ServerObject->NetHandle));
+	UE_NET_ASSERT_NE(ClientObject, nullptr);
+	UE_NET_ASSERT_EQ(ClientObject->IntA, ServerObject->IntA);
+}
+
+UE_NET_TEST_FIXTURE(FTestCancelPendingDestroyFixture, TestCancelPendingDestroyDuringWaitOnDestroyConfirmationWithoutPacketLoss)
+{
+	// Add a client
+	FReplicationSystemTestClient* Client = CreateClient();
+
+	// Spawn object on server
+	UTestReplicatedIrisObject* ServerObject = Server->CreateObject(UTestReplicatedIrisObject::FComponents());
+
+	// Write and send packet
+	Server->PreSendUpdate();
+	Server->SendAndDeliverTo(Client, DeliverPacket);
+	Server->PostSendUpdate();
+
+	// Filter out object and write packet to cause object to end up in WaitOnDestroyConfirmation
+	Server->GetReplicationSystem()->AddToGroup(NotReplicatedNetObjectGroupHandle, ServerObject->NetHandle);
+	Server->PreSendUpdate();
+	Server->SendUpdate(Client->ConnectionIdOnServer);
+	Server->PostSendUpdate();
+
+	// Remove object from filter to cause object to end up in CancelPendingDestroy, also modify a property 
+	Server->GetReplicationSystem()->RemoveFromGroup(NotReplicatedNetObjectGroupHandle, ServerObject->NetHandle);
+	ServerObject->IntA ^= 1;
+	Server->PreSendUpdate();
+	Server->PostSendUpdate();
+
+	// Deliver object destroy packet
+	Server->DeliverTo(Client, DeliverPacket);
+
+	// The destroy was delivered so the client object should not exist
+	UE_NET_ASSERT_EQ(Client->GetReplicationBridge()->GetReplicatedObject(ServerObject->NetHandle), nullptr);
+
+	// Write and send packet
+	Server->PreSendUpdate();
+	Server->SendAndDeliverTo(Client, DeliverPacket);
+	Server->PostSendUpdate();
+
+	UTestReplicatedIrisObject* ClientObject = Cast<UTestReplicatedIrisObject>(Client->GetReplicationBridge()->GetReplicatedObject(ServerObject->NetHandle));
+	UE_NET_ASSERT_NE(ClientObject, nullptr);
+	UE_NET_ASSERT_EQ(ClientObject->IntA, ServerObject->IntA);
+}
+
+// Tests for canceling a canceled pending destroy
+UE_NET_TEST_FIXTURE(FTestCancelPendingDestroyFixture, TestCancelingCancelPendingDestroyDuringWaitOnCreateConfirmationWithoutPacketLoss)
+{
+	// Add a client
+	FReplicationSystemTestClient* Client = CreateClient();
+
+	// Spawn object on server
+	UTestReplicatedIrisObject* ServerObject = Server->CreateObject(UTestReplicatedIrisObject::FComponents());
+
+	// Write packet
+	Server->PreSendUpdate();
+	Server->SendUpdate(Client->ConnectionIdOnServer);
+	Server->PostSendUpdate();
+
+	// Filter out object to cause a PendingDestroy
+	Server->GetReplicationSystem()->AddToGroup(NotReplicatedNetObjectGroupHandle, ServerObject->NetHandle);
+	Server->PreSendUpdate();
+	Server->PostSendUpdate();
+
+	// Remove object from filter to cause object to end up in CancelPendingDestroy...
+	Server->GetReplicationSystem()->RemoveFromGroup(NotReplicatedNetObjectGroupHandle, ServerObject->NetHandle);
+	Server->PreSendUpdate();
+	Server->PostSendUpdate();
+
+	// ... and cancel that thought! We want the object to be destroyed after all.
+	Server->GetReplicationSystem()->AddToGroup(NotReplicatedNetObjectGroupHandle, ServerObject->NetHandle);
+	Server->PreSendUpdate();
+	Server->PostSendUpdate();
+
+	Server->DeliverTo(Client, DeliverPacket);
+
+	// Verify that the object now exists on client
+	UE_NET_ASSERT_NE(Client->GetReplicationBridge()->GetReplicatedObject(ServerObject->NetHandle), nullptr);
+
+	// Write and send a packet
+	Server->PreSendUpdate();
+	Server->SendAndDeliverTo(Client, DeliverPacket);
+	Server->PostSendUpdate();
+
+	// Verify that the object was destroyed on the client
+	UE_NET_ASSERT_EQ(Client->GetReplicationBridge()->GetReplicatedObject(ServerObject->NetHandle), nullptr);
+}
+
+UE_NET_TEST_FIXTURE(FTestCancelPendingDestroyFixture, TestCancelingCancelPendingDestroyDuringWaitOnCreateConfirmationWithPacketLoss)
+{
+	// Add a client
+	FReplicationSystemTestClient* Client = CreateClient();
+
+	// Spawn object on server
+	UTestReplicatedIrisObject* ServerObject = Server->CreateObject(UTestReplicatedIrisObject::FComponents());
+
+	// Write packet
+	Server->PreSendUpdate();
+	Server->SendUpdate(Client->ConnectionIdOnServer);
+	Server->PostSendUpdate();
+
+	// Filter out object to cause a PendingDestroy
+	Server->GetReplicationSystem()->AddToGroup(NotReplicatedNetObjectGroupHandle, ServerObject->NetHandle);
+	Server->PreSendUpdate();
+	Server->PostSendUpdate();
+
+	// Remove object from filter to cause object to end up in CancelPendingDestroy...
+	Server->GetReplicationSystem()->RemoveFromGroup(NotReplicatedNetObjectGroupHandle, ServerObject->NetHandle);
+	Server->PreSendUpdate();
+	Server->PostSendUpdate();
+
+	// ... and cancel that thought! We want the object to be destroyed after all.
+	Server->GetReplicationSystem()->AddToGroup(NotReplicatedNetObjectGroupHandle, ServerObject->NetHandle);
+	Server->PreSendUpdate();
+	Server->PostSendUpdate();
+
+	// Cause packet loss on object creation
+	Server->DeliverTo(Client, DoNotDeliverPacket);
+
+	// Verify that the object doesn't exist on the client
+	UE_NET_ASSERT_EQ(Client->GetReplicationBridge()->GetReplicatedObject(ServerObject->NetHandle), nullptr);
+
+	// Write and send a packet
+	Server->PreSendUpdate();
+	Server->SendAndDeliverTo(Client, DeliverPacket);
+	Server->PostSendUpdate();
+
+	// Verify that the object still doesn't exist on the client
+	UE_NET_ASSERT_EQ(Client->GetReplicationBridge()->GetReplicatedObject(ServerObject->NetHandle), nullptr);
+}
+
+UE_NET_TEST_FIXTURE(FTestCancelPendingDestroyFixture, TestCancelingCancelPendingDestroyDuringCreated)
+{
+	// Add a client
+	FReplicationSystemTestClient* Client = CreateClient();
+
+	// Spawn object on server
+	UTestReplicatedIrisObject* ServerObject = Server->CreateObject(UTestReplicatedIrisObject::FComponents());
+
+	// Write and send packet
+	Server->PreSendUpdate();
+	Server->SendAndDeliverTo(Client, DeliverPacket);
+	Server->PostSendUpdate();
+
+	// Verify that the object now also exists on the client
+	UE_NET_ASSERT_NE(Client->GetReplicationBridge()->GetReplicatedObject(ServerObject->NetHandle), nullptr);
+
+	// Filter out object to cause a PendingDestroy
+	Server->GetReplicationSystem()->AddToGroup(NotReplicatedNetObjectGroupHandle, ServerObject->NetHandle);
+	Server->PreSendUpdate();
+	Server->PostSendUpdate();
+
+	// Remove object from filter to cause object to end up in CancelPendingDestroy...
+	Server->GetReplicationSystem()->RemoveFromGroup(NotReplicatedNetObjectGroupHandle, ServerObject->NetHandle);
+	Server->PreSendUpdate();
+	Server->PostSendUpdate();
+
+	// ... and cancel that thought! We want the object to be destroyed after all.
+	Server->GetReplicationSystem()->AddToGroup(NotReplicatedNetObjectGroupHandle, ServerObject->NetHandle);
+	Server->PreSendUpdate();
+	Server->PostSendUpdate();
+
+	// Write and send a packet
+	Server->PreSendUpdate();
+	Server->SendAndDeliverTo(Client, DeliverPacket);
+	Server->PostSendUpdate();
+
+	// Verify that the object was destroyed on the client
+	UE_NET_ASSERT_EQ(Client->GetReplicationBridge()->GetReplicatedObject(ServerObject->NetHandle), nullptr);
+}
+
+UE_NET_TEST_FIXTURE(FTestCancelPendingDestroyFixture, TestCancelingCancelPendingDestroyDuringWaitOnDestroyConfirmationWithPacketLoss)
+{
+	// Add a client
+	FReplicationSystemTestClient* Client = CreateClient();
+
+	// Spawn object on server
+	UTestReplicatedIrisObject* ServerObject = Server->CreateObject(UTestReplicatedIrisObject::FComponents());
+
+	// Write and send packet
+	Server->PreSendUpdate();
+	Server->SendAndDeliverTo(Client, DeliverPacket);
+	Server->PostSendUpdate();
+
+	// Filter out object and write packet to cause object to end up in WaitOnDestroyConfirmation
+	Server->GetReplicationSystem()->AddToGroup(NotReplicatedNetObjectGroupHandle, ServerObject->NetHandle);
+	Server->PreSendUpdate();
+	Server->SendUpdate(Client->ConnectionIdOnServer);
+	Server->PostSendUpdate();
+
+	// Remove object from filter to cause object to end up in CancelPendingDestroy...
+	Server->GetReplicationSystem()->RemoveFromGroup(NotReplicatedNetObjectGroupHandle, ServerObject->NetHandle);
+	Server->PreSendUpdate();
+	Server->PostSendUpdate();
+
+	// ... and cancel that thought! We want the object to be destroyed after all.
+	Server->GetReplicationSystem()->AddToGroup(NotReplicatedNetObjectGroupHandle, ServerObject->NetHandle);
+	Server->PreSendUpdate();
+	Server->PostSendUpdate();
+
+	// Cause packet loss on the object destroy packet
+	Server->DeliverTo(Client, DoNotDeliverPacket);
+
+	// Verify that the object still exists on client
+	UE_NET_ASSERT_NE(Client->GetReplicationBridge()->GetReplicatedObject(ServerObject->NetHandle), nullptr);
+
+	// Write and send a packet
+	Server->PreSendUpdate();
+	Server->SendAndDeliverTo(Client, DeliverPacket);
+	Server->PostSendUpdate();
+
+	// Verify that the object was destroyed on the client
+	UE_NET_ASSERT_EQ(Client->GetReplicationBridge()->GetReplicatedObject(ServerObject->NetHandle), nullptr);
+}
+
+UE_NET_TEST_FIXTURE(FTestCancelPendingDestroyFixture, TestCancelingCancelPendingDestroyDuringWaitOnDestroyConfirmationWithoutPacketLoss)
+{
+	// Add a client
+	FReplicationSystemTestClient* Client = CreateClient();
+
+	// Spawn object on server
+	UTestReplicatedIrisObject* ServerObject = Server->CreateObject(UTestReplicatedIrisObject::FComponents());
+
+	// Write and send packet
+	Server->PreSendUpdate();
+	Server->SendAndDeliverTo(Client, DeliverPacket);
+	Server->PostSendUpdate();
+
+	// Filter out object and write packet to cause object to end up in WaitOnDestroyConfirmation
+	Server->GetReplicationSystem()->AddToGroup(NotReplicatedNetObjectGroupHandle, ServerObject->NetHandle);
+	Server->PreSendUpdate();
+	Server->SendUpdate(Client->ConnectionIdOnServer);
+	Server->PostSendUpdate();
+
+	// Remove object from filter to cause object to end up in CancelPendingDestroy...
+	Server->GetReplicationSystem()->RemoveFromGroup(NotReplicatedNetObjectGroupHandle, ServerObject->NetHandle);
+	Server->PreSendUpdate();
+	Server->PostSendUpdate();
+
+	// ... and cancel that thought! We want the object to be destroyed after all.
+	Server->GetReplicationSystem()->AddToGroup(NotReplicatedNetObjectGroupHandle, ServerObject->NetHandle);
+	Server->PreSendUpdate();
+	Server->PostSendUpdate();
+
+	// Deliver destroy packet
+	Server->DeliverTo(Client, DeliverPacket);
+
+	// Verify that the object doesn't exist on the client
+	UE_NET_ASSERT_EQ(Client->GetReplicationBridge()->GetReplicatedObject(ServerObject->NetHandle), nullptr);
+
+	// Write and send a packet
+	Server->PreSendUpdate();
+	Server->SendAndDeliverTo(Client, DeliverPacket);
+	Server->PostSendUpdate();
+
+	// Verify that the object still doesn't exist on the client
+	UE_NET_ASSERT_EQ(Client->GetReplicationBridge()->GetReplicatedObject(ServerObject->NetHandle), nullptr);
+}
+
+
+// Subobject tests
+UE_NET_TEST_FIXTURE(FTestCancelPendingDestroyFixture, TestCancelPendingSubObjectDestroyDuringCreated)
+{
+	// Add a client
+	FReplicationSystemTestClient* Client = CreateClient();
+
+	// Spawn object with subobject on server
+	UTestReplicatedIrisObject* ServerObject = Server->CreateObject(UTestReplicatedIrisObject::FComponents());
+	UTestReplicatedIrisObject* ServerSubObject = Server->CreateSubObject(ServerObject->NetHandle, UTestReplicatedIrisObject::FComponents());
+
+	// Write and send packet
+	Server->PreSendUpdate();
+	Server->SendAndDeliverTo(Client, DeliverPacket);
+	Server->PostSendUpdate();
+
+	// Verify that the subobject exists on the client
+	UE_NET_ASSERT_NE(Client->GetReplicationBridge()->GetReplicatedObject(ServerSubObject->NetHandle), nullptr);
+
+	// Filter out subobject to cause a SubObjectPendingDestroy
+	Server->GetReplicationSystem()->AddToGroup(NotReplicatedNetObjectGroupHandle, ServerSubObject->NetHandle);
+	Server->PreSendUpdate();
+	Server->PostSendUpdate();
+
+	// Remove subobject from filter to cause object to end up in CancelPendingDestroy
+	Server->GetReplicationSystem()->RemoveFromGroup(NotReplicatedNetObjectGroupHandle, ServerSubObject->NetHandle);
+	Server->PreSendUpdate();
+	Server->PostSendUpdate();
+
+	// Modify a property on the subobject and make sure it's replicated as the subobject should still exist
+	ServerSubObject->IntA ^= 1;
+
+	Server->PreSendUpdate();
+	Server->SendAndDeliverTo(Client, DeliverPacket);
+	Server->PostSendUpdate();
+
+	UTestReplicatedIrisObject* ClientSubObject = Cast<UTestReplicatedIrisObject>(Client->GetReplicationBridge()->GetReplicatedObject(ServerSubObject->NetHandle));
+	UE_NET_ASSERT_NE(ClientSubObject, nullptr);
+	UE_NET_ASSERT_EQ(ClientSubObject->IntA, ServerSubObject->IntA);
+}
+
+UE_NET_TEST_FIXTURE(FTestCancelPendingDestroyFixture, TestCancelPendingSubObjectDestroyDuringWaitOnDestroyConfirmationWithPacketLoss)
+{
+	// Add a client
+	FReplicationSystemTestClient* Client = CreateClient();
+
+	// Spawn object on server
+	UTestReplicatedIrisObject* ServerObject = Server->CreateObject(UTestReplicatedIrisObject::FComponents());
+	UTestReplicatedIrisObject* ServerSubObject = Server->CreateSubObject(ServerObject->NetHandle, UTestReplicatedIrisObject::FComponents());
+
+	// Write and send packet
+	Server->PreSendUpdate();
+	Server->SendAndDeliverTo(Client, DeliverPacket);
+	Server->PostSendUpdate();
+
+	// Filter out object and write packet to cause subobject to end up in WaitOnDestroyConfirmation
+	Server->GetReplicationSystem()->AddToGroup(NotReplicatedNetObjectGroupHandle, ServerSubObject->NetHandle);
+	Server->PreSendUpdate();
+	Server->SendUpdate(Client->ConnectionIdOnServer);
+	Server->PostSendUpdate();
+
+	// Remove subobject from filter to cause object to end up in CancelPendingDestroy
+	Server->GetReplicationSystem()->RemoveFromGroup(NotReplicatedNetObjectGroupHandle, ServerSubObject->NetHandle);
+	Server->PreSendUpdate();
+	Server->PostSendUpdate();
+
+	// Cause packet loss on the subobject destroy packet
+	Server->DeliverTo(Client, DoNotDeliverPacket);
+
+	// Verify that object still exists on client
+	UE_NET_ASSERT_NE(Client->GetReplicationBridge()->GetReplicatedObject(ServerSubObject->NetHandle), nullptr);
+
+	// Modify a property on the subobject and make sure it's replicated as the subobject should still exist
+	ServerSubObject->IntA ^= 1;
+
+	Server->PreSendUpdate();
+	Server->SendAndDeliverTo(Client, DeliverPacket);
+	Server->PostSendUpdate();
+
+	UTestReplicatedIrisObject* ClientSubObject = Cast<UTestReplicatedIrisObject>(Client->GetReplicationBridge()->GetReplicatedObject(ServerSubObject->NetHandle));
+	UE_NET_ASSERT_NE(ClientSubObject, nullptr);
+	UE_NET_ASSERT_EQ(ClientSubObject->IntA, ServerSubObject->IntA);
+}
+
+UE_NET_TEST_FIXTURE(FTestCancelPendingDestroyFixture, TestCancelPendingSubObjectDestroyDuringWaitOnDestroyConfirmationWithoutPacketLoss)
+{
+	// Add a client
+	FReplicationSystemTestClient* Client = CreateClient();
+
+	// Spawn object on server
+	UTestReplicatedIrisObject* ServerObject = Server->CreateObject(UTestReplicatedIrisObject::FComponents());
+	UTestReplicatedIrisObject* ServerSubObject = Server->CreateSubObject(ServerObject->NetHandle, UTestReplicatedIrisObject::FComponents());
+
+	// Write and send packet
+	Server->PreSendUpdate();
+	Server->SendAndDeliverTo(Client, DeliverPacket);
+	Server->PostSendUpdate();
+
+	// Filter out object and write packet to cause subobject to end up in WaitOnDestroyConfirmation
+	Server->GetReplicationSystem()->AddToGroup(NotReplicatedNetObjectGroupHandle, ServerSubObject->NetHandle);
+	Server->PreSendUpdate();
+	Server->SendUpdate(Client->ConnectionIdOnServer);
+	Server->PostSendUpdate();
+
+	// Remove object from filter to cause object to end up in CancelPendingDestroy, also modify a property 
+	Server->GetReplicationSystem()->RemoveFromGroup(NotReplicatedNetObjectGroupHandle, ServerSubObject->NetHandle);
+	ServerSubObject->IntA ^= 1;
+	Server->PreSendUpdate();
+	Server->PostSendUpdate();
+
+	// Deliver subobject destroy packet
+	Server->DeliverTo(Client, DeliverPacket);
+
+	// The destroy was delivered so the client subobject should not exist
+	UE_NET_ASSERT_EQ(Client->GetReplicationBridge()->GetReplicatedObject(ServerSubObject->NetHandle), nullptr);
+
+	// Write and send a packet
+	Server->PreSendUpdate();
+	Server->SendAndDeliverTo(Client, DeliverPacket);
+	Server->PostSendUpdate();
+
+	// Client subobject should have been created with the latest state
+	UTestReplicatedIrisObject* ClientSubObject = Cast<UTestReplicatedIrisObject>(Client->GetReplicationBridge()->GetReplicatedObject(ServerSubObject->NetHandle));
+	UE_NET_ASSERT_NE(ClientSubObject, nullptr);
+	UE_NET_ASSERT_EQ(ClientSubObject->IntA, ServerSubObject->IntA);
+}
+
+UE_NET_TEST_FIXTURE(FTestCancelPendingDestroyFixture, TestCancelingCancelPendingSubObjectDestroyDuringCreated)
+{
+	// Add a client
+	FReplicationSystemTestClient* Client = CreateClient();
+
+	// Spawn object on server
+	UTestReplicatedIrisObject* ServerObject = Server->CreateObject(UTestReplicatedIrisObject::FComponents());
+	UTestReplicatedIrisObject* ServerSubObject = Server->CreateSubObject(ServerObject->NetHandle, UTestReplicatedIrisObject::FComponents());
+
+	// Write and send packet
+	Server->PreSendUpdate();
+	Server->SendAndDeliverTo(Client, DeliverPacket);
+	Server->PostSendUpdate();
+
+	// Verify that the subobject now also exists on the client
+	UE_NET_ASSERT_NE(Client->GetReplicationBridge()->GetReplicatedObject(ServerSubObject->NetHandle), nullptr);
+
+	// Filter out object to cause a SubObjectPendingDestroy
+	Server->GetReplicationSystem()->AddToGroup(NotReplicatedNetObjectGroupHandle, ServerSubObject->NetHandle);
+	Server->PreSendUpdate();
+	Server->PostSendUpdate();
+
+	// Remove object from filter to cause object to end up in CancelPendingDestroy...
+	Server->GetReplicationSystem()->RemoveFromGroup(NotReplicatedNetObjectGroupHandle, ServerSubObject->NetHandle);
+	Server->PreSendUpdate();
+	Server->PostSendUpdate();
+
+	// ... and cancel that thought! We want the subobject to be destroyed after all.
+	Server->GetReplicationSystem()->AddToGroup(NotReplicatedNetObjectGroupHandle, ServerSubObject->NetHandle);
+	Server->PreSendUpdate();
+	Server->PostSendUpdate();
+
+	// Write and send a packet
+	Server->PreSendUpdate();
+	Server->SendAndDeliverTo(Client, DeliverPacket);
+	Server->PostSendUpdate();
+
+	// Verify that the subobject was destroyed on the client
+	UE_NET_ASSERT_EQ(Client->GetReplicationBridge()->GetReplicatedObject(ServerSubObject->NetHandle), nullptr);
+}
+
+UE_NET_TEST_FIXTURE(FTestCancelPendingDestroyFixture, TestCancelingCancelPendingSubObjectDestroyDuringWaitOnDestroyConfirmationWithPacketLoss)
+{
+	// Add a client
+	FReplicationSystemTestClient* Client = CreateClient();
+
+	// Spawn object on server
+	UTestReplicatedIrisObject* ServerObject = Server->CreateObject(UTestReplicatedIrisObject::FComponents());
+	UTestReplicatedIrisObject* ServerSubObject = Server->CreateSubObject(ServerObject->NetHandle, UTestReplicatedIrisObject::FComponents());
+
+	// Write and send packet
+	Server->PreSendUpdate();
+	Server->SendAndDeliverTo(Client, DeliverPacket);
+	Server->PostSendUpdate();
+
+	// Filter out object and write packet to cause subobject to end up in WaitOnDestroyConfirmation
+	Server->GetReplicationSystem()->AddToGroup(NotReplicatedNetObjectGroupHandle, ServerSubObject->NetHandle);
+	Server->PreSendUpdate();
+	Server->SendUpdate(Client->ConnectionIdOnServer);
+	Server->PostSendUpdate();
+
+	// Remove object from filter to cause subobject to end up in CancelPendingDestroy...
+	Server->GetReplicationSystem()->RemoveFromGroup(NotReplicatedNetObjectGroupHandle, ServerSubObject->NetHandle);
+	Server->PreSendUpdate();
+	Server->PostSendUpdate();
+
+	// ... and cancel that thought! We want the subobject to be destroyed after all.
+	Server->GetReplicationSystem()->AddToGroup(NotReplicatedNetObjectGroupHandle, ServerSubObject->NetHandle);
+	Server->PreSendUpdate();
+	Server->PostSendUpdate();
+
+	// Cause packet loss on the subobject destroy packet
+	Server->DeliverTo(Client, DoNotDeliverPacket);
+
+	// Verify that the object still exists on client
+	UE_NET_ASSERT_NE(Client->GetReplicationBridge()->GetReplicatedObject(ServerSubObject->NetHandle), nullptr);
+
+	// Write and send a packet
+	Server->PreSendUpdate();
+	Server->SendAndDeliverTo(Client, DeliverPacket);
+	Server->PostSendUpdate();
+
+	// Verify that the subobject was destroyed on the client
+	UE_NET_ASSERT_EQ(Client->GetReplicationBridge()->GetReplicatedObject(ServerSubObject->NetHandle), nullptr);
+}
+
+UE_NET_TEST_FIXTURE(FTestCancelPendingDestroyFixture, TestCancelingCancelPendingSubObjectDestroyDuringWaitOnDestroyConfirmationWithoutPacketLoss)
+{
+	// Add a client
+	FReplicationSystemTestClient* Client = CreateClient();
+
+	// Spawn object on server
+	UTestReplicatedIrisObject* ServerObject = Server->CreateObject(UTestReplicatedIrisObject::FComponents());
+	UTestReplicatedIrisObject* ServerSubObject = Server->CreateSubObject(ServerObject->NetHandle, UTestReplicatedIrisObject::FComponents());
+
+	// Write and send packet
+	Server->PreSendUpdate();
+	Server->SendAndDeliverTo(Client, DeliverPacket);
+	Server->PostSendUpdate();
+
+	// Filter out subobject and write packet to cause subobject to end up in WaitOnDestroyConfirmation
+	Server->GetReplicationSystem()->AddToGroup(NotReplicatedNetObjectGroupHandle, ServerSubObject->NetHandle);
+	Server->PreSendUpdate();
+	Server->SendUpdate(Client->ConnectionIdOnServer);
+	Server->PostSendUpdate();
+
+	// Remove subobject from filter to cause subobject to end up in CancelPendingDestroy...
+	Server->GetReplicationSystem()->RemoveFromGroup(NotReplicatedNetObjectGroupHandle, ServerSubObject->NetHandle);
+	Server->PreSendUpdate();
+	Server->PostSendUpdate();
+
+	// ... and cancel that thought! We want the subobject to be destroyed after all.
+	Server->GetReplicationSystem()->AddToGroup(NotReplicatedNetObjectGroupHandle, ServerSubObject->NetHandle);
+	Server->PreSendUpdate();
+	Server->PostSendUpdate();
+
+	// Deliver destroy packet
+	Server->DeliverTo(Client, DeliverPacket);
+
+	// Verify that the subobject doesn't exist on the client
+	UE_NET_ASSERT_EQ(Client->GetReplicationBridge()->GetReplicatedObject(ServerSubObject->NetHandle), nullptr);
+
+	// Write and send a packet
+	Server->PreSendUpdate();
+	Server->SendAndDeliverTo(Client, DeliverPacket);
+	Server->PostSendUpdate();
+
+	// Verify that the subobject still doesn't exist on the client
+	UE_NET_ASSERT_EQ(Client->GetReplicationBridge()->GetReplicatedObject(ServerSubObject->NetHandle), nullptr);
+}
+
+UE_NET_TEST_FIXTURE(FTestCancelPendingDestroyFixture, TestCancelingSubObjectDestroyAfterParentIsMarkedForDestroy)
+{
+	// Add a client
+	FReplicationSystemTestClient* Client = CreateClient();
+
+	// Spawn object on server
+	UTestReplicatedIrisObject* ServerObject = Server->CreateObject(UTestReplicatedIrisObject::FComponents());
+	UTestReplicatedIrisObject* ServerSubObject1 = Server->CreateSubObject(ServerObject->NetHandle, UTestReplicatedIrisObject::FComponents());
+	UTestReplicatedIrisObject* ServerSubObject2 = Server->CreateSubObject(ServerObject->NetHandle, UTestReplicatedIrisObject::FComponents());
+
+	// Write and send packet
+	Server->PreSendUpdate();
+	Server->SendAndDeliverTo(Client, DeliverPacket);
+	Server->PostSendUpdate();
+
+	// Filter out subobject and write packet to cause subobject to end up in WaitOnDestroyConfirmation
+	Server->GetReplicationSystem()->AddToGroup(NotReplicatedNetObjectGroupHandle, ServerSubObject2->NetHandle);
+	Server->PreSendUpdate();
+	Server->SendUpdate(Client->ConnectionIdOnServer);
+	Server->PostSendUpdate();
+
+	// Filter out parent and write packet
+	Server->GetReplicationSystem()->AddToGroup(NotReplicatedNetObjectGroupHandle, ServerObject->NetHandle);
+	Server->PreSendUpdate();
+	Server->SendUpdate(Client->ConnectionIdOnServer);
+	Server->PostSendUpdate();
+
+	// Remove subobject from filter to cause subobject to end up in CancelPendingDestroy
+	Server->GetReplicationSystem()->RemoveFromGroup(NotReplicatedNetObjectGroupHandle, ServerSubObject2->NetHandle);
+	Server->PreSendUpdate();
+	Server->PostSendUpdate();
+
+	// Cause subobject destroy packet to be lost
+	Server->DeliverTo(Client, DoNotDeliverPacket);
+
+	// Verify all objects and subobjects exist on the client
+	UE_NET_ASSERT_NE(Client->GetReplicationBridge()->GetReplicatedObject(ServerObject->NetHandle), nullptr);
+	UE_NET_ASSERT_NE(Client->GetReplicationBridge()->GetReplicatedObject(ServerSubObject1->NetHandle), nullptr);
+	UE_NET_ASSERT_NE(Client->GetReplicationBridge()->GetReplicatedObject(ServerSubObject2->NetHandle), nullptr);
+
+	// Deliver owner destroy packet
+	Server->DeliverTo(Client, DeliverPacket);
+
+	// Verify that no objects exist on the client
+	UE_NET_ASSERT_EQ(Client->GetReplicationBridge()->GetReplicatedObject(ServerObject->NetHandle), nullptr);
+	UE_NET_ASSERT_EQ(Client->GetReplicationBridge()->GetReplicatedObject(ServerSubObject1->NetHandle), nullptr);
+	UE_NET_ASSERT_EQ(Client->GetReplicationBridge()->GetReplicatedObject(ServerSubObject2->NetHandle), nullptr);
+
+	// Write and send a packet
+	Server->PreSendUpdate();
+	Server->SendAndDeliverTo(Client, DeliverPacket);
+	Server->PostSendUpdate();
+
+	// Verify that all objects are still gone on the client
+	UE_NET_ASSERT_EQ(Client->GetReplicationBridge()->GetReplicatedObject(ServerObject->NetHandle), nullptr);
+	UE_NET_ASSERT_EQ(Client->GetReplicationBridge()->GetReplicatedObject(ServerSubObject1->NetHandle), nullptr);
+	UE_NET_ASSERT_EQ(Client->GetReplicationBridge()->GetReplicatedObject(ServerSubObject2->NetHandle), nullptr);
+}
+
+// See TestObjectSplitting.cpp for cancel pending destroy on huge object
+
+}
diff --git a/Engine/Plugins/Runtime/ReplicationSystemTestPlugin/Source/Private/Tests/ReplicationSystem/TestDependentObjects.cpp b/Engine/Plugins/Runtime/ReplicationSystemTestPlugin/Source/Private/Tests/ReplicationSystem/TestDependentObjects.cpp
new file mode 100644
--- /dev/null
+++ b/Engine/Plugins/Runtime/ReplicationSystemTestPlugin/Source/Private/Tests/ReplicationSystem/TestDependentObjects.cpp
@@ -0,0 +1,379 @@
+// Copyright Epic Games, Inc. All Rights Reserved.
+
+#include "ReplicationSystemServerClientTestFixture.h"
+#include "Iris/Serialization/NetBitStreamReader.h"
+#include "Iris/Serialization/NetBitStreamWriter.h"
+#include "Iris/ReplicationSystem/ReplicationSystem.h"
+
+#include "Iris/ReplicationSystem/ReplicationSystemInternal.h"
+#include "Iris/ReplicationSystem/NetTokenStore.h"
+
+namespace UE::Net::Private
+{
+
+UE_NET_TEST_FIXTURE(FReplicationSystemServerClientTestFixture, TestDependentObjectDroppedDataIsRetransmitted)
+{
+	UReplicationSystem* ReplicationSystem = Server->ReplicationSystem;
+
+	// Add a client
+	FReplicationSystemTestClient* Client = CreateClient();
+
+	// Spawn object on server
+	UTestReplicatedIrisObject* ServerObject = Server->CreateObject(0, 0);
+
+	// Spawn second object on server add as a dependent object
+	UTestReplicatedIrisObject* ServerDependentObject = Server->CreateObject(0, 0);
+	
+	UReplicatedTestObjectBridge* Bridge = Server->GetReplicationBridge();
+
+	// Trigger replication
+	ServerObject->IntA = 1;
+
+	// Send and deliver packet
+	Server->PreSendUpdate();
+	Server->SendAndDeliverTo(Client, true);
+	Server->PostSendUpdate();
+
+	Bridge->AddDependentObject(ServerObject->NetHandle, ServerDependentObject->NetHandle);
+	ServerDependentObject->IntA = 1;
+
+	// Send and deliver packet
+	Server->PreSendUpdate();
+	Server->SendAndDeliverTo(Client, true);
+	Server->PostSendUpdate();
+
+	// Store Pointer to objects
+	UTestReplicatedIrisObject* ClientObject = Cast<UTestReplicatedIrisObject>(Client->GetReplicationBridge()->GetReplicatedObject(ServerObject->NetHandle));
+	UTestReplicatedIrisObject* ClientDependentObject = Cast<UTestReplicatedIrisObject>(Client->GetReplicationBridge()->GetReplicatedObject(ServerDependentObject->NetHandle));
+
+	UE_NET_ASSERT_NE(ClientDependentObject, nullptr);
+	UE_NET_ASSERT_EQ(ServerDependentObject->IntA, ClientDependentObject->IntA);
+
+	// Modify the value of dependent object only
+	ServerDependentObject->IntA = 2;
+
+	// Send and do not deliver packet
+	Server->PreSendUpdate();
+	Server->SendAndDeliverTo(Client, false);
+	Server->PostSendUpdate();
+
+	// Verify that the final state was applied to dependent object 
+	UE_NET_ASSERT_NE(ServerDependentObject->IntA, ClientDependentObject->IntA);
+
+	// Send and deliver packet
+	Server->PreSendUpdate();
+	Server->SendAndDeliverTo(Client, true);
+	Server->PostSendUpdate();
+
+	// Verify that the final state was applied to dependent object 
+	UE_NET_ASSERT_EQ(ServerDependentObject->IntA, ClientDependentObject->IntA);
+}
+
+// Dependent objects
+UE_NET_TEST_FIXTURE(FReplicationSystemServerClientTestFixture, TestDependentObjectWithZeroPrioOnlyReplicatesWithParent)
+{
+	UReplicationSystem* ReplicationSystem = Server->ReplicationSystem;
+	UReplicatedTestObjectBridge* Bridge = Server->GetReplicationBridge();
+
+	// Add a client
+	FReplicationSystemTestClient* Client = CreateClient();
+
+	// Spawn object on server
+	UTestReplicatedIrisObject* ServerObject = Server->CreateObject(0, 0);
+
+	// Filter out Server object to start with
+	FNetObjectGroupHandle FilterGroup = ReplicationSystem->CreateGroup();
+	ReplicationSystem->AddGroupFilter(FilterGroup);
+	ReplicationSystem->AddToGroup(FilterGroup, ServerObject->NetHandle);
+
+	// Setup dependent object to only replicate with ServerObject
+	UTestReplicatedIrisObject* ServerDependentObject = Server->CreateObject(0, 0);
+	ReplicationSystem->SetStaticPriority(ServerDependentObject->NetHandle, 0.f);
+	Bridge->AddDependentObject(ServerObject->NetHandle, ServerDependentObject->NetHandle);
+
+	// Send and deliver packet
+	Server->PreSendUpdate();
+	Server->SendAndDeliverTo(Client, true);
+	Server->PostSendUpdate();
+
+	// Objects should not be created on client
+	UTestReplicatedIrisObject* ClientObject = Cast<UTestReplicatedIrisObject>(Client->GetReplicationBridge()->GetReplicatedObject(ServerObject->NetHandle));
+	UTestReplicatedIrisObject* ClientDependentObject = Cast<UTestReplicatedIrisObject>(Client->GetReplicationBridge()->GetReplicatedObject(ServerDependentObject->NetHandle));
+	UE_NET_ASSERT_EQ(ClientObject, nullptr);
+	UE_NET_ASSERT_EQ(ClientDependentObject, nullptr);
+
+	// Enable 
+	ReplicationSystem->SetGroupFilterStatus(FilterGroup, ENetFilterStatus::Allow);
+
+	// Send and deliver packet
+	Server->PreSendUpdate();
+	Server->SendAndDeliverTo(Client, true);
+	Server->PostSendUpdate();
+
+	// Objects should now exist on client
+	ClientObject = Cast<UTestReplicatedIrisObject>(Client->GetReplicationBridge()->GetReplicatedObject(ServerObject->NetHandle));
+	ClientDependentObject = Cast<UTestReplicatedIrisObject>(Client->GetReplicationBridge()->GetReplicatedObject(ServerDependentObject->NetHandle));
+	UE_NET_ASSERT_NE(ClientObject, nullptr);
+	UE_NET_ASSERT_NE(ClientDependentObject, nullptr);
+}
+
+// Chained dependent objects
+UE_NET_TEST_FIXTURE(FReplicationSystemServerClientTestFixture, TestChainedDependentObjectWithZeroPrio)
+{
+	UReplicationSystem* ReplicationSystem = Server->ReplicationSystem;
+	UReplicatedTestObjectBridge* Bridge = Server->GetReplicationBridge();
+
+	// Add a client
+	FReplicationSystemTestClient* Client = CreateClient();
+
+	// Spawn object on server
+	UTestReplicatedIrisObject* ServerObject = Server->CreateObject(0, 0);
+
+	// Filter out Server object to start with
+	FNetObjectGroupHandle FilterGroup = ReplicationSystem->CreateGroup();
+	ReplicationSystem->AddGroupFilter(FilterGroup);
+	ReplicationSystem->AddToGroup(FilterGroup, ServerObject->NetHandle);
+
+	// Setup dependent object to only replicate with ServerObject
+	UTestReplicatedIrisObject* ServerDependentObject0 = Server->CreateObject(0, 0);
+	UTestReplicatedIrisObject* ServerDependentObject1 = Server->CreateObject(0, 0);
+	
+	ReplicationSystem->SetStaticPriority(ServerDependentObject0->NetHandle, 0.f);
+	ReplicationSystem->SetStaticPriority(ServerDependentObject1->NetHandle, 0.f);
+	
+	// Setup dependency chain
+	Bridge->AddDependentObject(ServerObject->NetHandle, ServerDependentObject0->NetHandle);
+	Bridge->AddDependentObject(ServerDependentObject0->NetHandle, ServerDependentObject1->NetHandle);
+
+	// Send and deliver packet
+	Server->PreSendUpdate();
+	Server->SendAndDeliverTo(Client, true);
+	Server->PostSendUpdate();
+
+	// Objects should not be created on client
+	UTestReplicatedIrisObject* ClientObject = Cast<UTestReplicatedIrisObject>(Client->GetReplicationBridge()->GetReplicatedObject(ServerObject->NetHandle));
+	UTestReplicatedIrisObject* ClientDependentObject0 = Cast<UTestReplicatedIrisObject>(Client->GetReplicationBridge()->GetReplicatedObject(ServerDependentObject0->NetHandle));
+	UTestReplicatedIrisObject* ClientDependentObject1 = Cast<UTestReplicatedIrisObject>(Client->GetReplicationBridge()->GetReplicatedObject(ServerDependentObject1->NetHandle));
+
+	UE_NET_ASSERT_EQ(ClientObject, nullptr);
+	UE_NET_ASSERT_EQ(ClientDependentObject0, nullptr);
+	UE_NET_ASSERT_EQ(ClientDependentObject1, nullptr);
+
+	// Enable the parent
+	ReplicationSystem->SetGroupFilterStatus(FilterGroup, ENetFilterStatus::Allow);
+
+	// Send and deliver packet
+	Server->PreSendUpdate();
+	Server->SendAndDeliverTo(Client, true);
+	Server->PostSendUpdate();
+
+	// Verify that dependent object now is created
+	ClientObject = Cast<UTestReplicatedIrisObject>(Client->GetReplicationBridge()->GetReplicatedObject(ServerObject->NetHandle));
+	ClientDependentObject0 = Cast<UTestReplicatedIrisObject>(Client->GetReplicationBridge()->GetReplicatedObject(ServerDependentObject0->NetHandle));
+	ClientDependentObject1 = Cast<UTestReplicatedIrisObject>(Client->GetReplicationBridge()->GetReplicatedObject(ServerDependentObject1->NetHandle));
+	UE_NET_ASSERT_NE(ClientObject, nullptr);
+	UE_NET_ASSERT_NE(ClientDependentObject0, nullptr);
+	UE_NET_ASSERT_NE(ClientDependentObject1, nullptr);
+}
+
+UE_NET_TEST_FIXTURE(FReplicationSystemServerClientTestFixture, TestDependentObjectPollFrequency)
+{
+	UReplicationSystem* ReplicationSystem = Server->ReplicationSystem;
+
+	// Add a client
+	FReplicationSystemTestClient* Client = CreateClient();
+
+	// Spawn object on server
+	UTestReplicatedIrisObject* ServerObject = Server->CreateObject(0, 0);
+
+	// Spawn second object on server that later will be added as a dependent object
+	// With high PollFramePeriod so that it will not replicate in a while unless it is a dependent
+	UTestReplicatedIrisObject* ServerDependentObject = Server->CreateObject(0, 0);
+	Server->ReplicationBridge->SetPollFramePeriod(ServerDependentObject, 255U);
+
+	// Send and deliver packet
+	Server->PreSendUpdate();
+	Server->SendAndDeliverTo(Client, true);
+	Server->PostSendUpdate();
+
+	// Store Pointer to objects and verify state after initial replication
+	UTestReplicatedIrisObject* ClientObject = Cast<UTestReplicatedIrisObject>(Client->GetReplicationBridge()->GetReplicatedObject(ServerObject->NetHandle));
+	UTestReplicatedIrisObject* ClientDependentObject = Cast<UTestReplicatedIrisObject>(Client->GetReplicationBridge()->GetReplicatedObject(ServerDependentObject->NetHandle));
+	
+	UE_NET_ASSERT_NE(ClientDependentObject, nullptr);
+	UE_NET_ASSERT_NE(ClientObject, nullptr);
+	UE_NET_ASSERT_EQ(ClientObject->IntA, ServerObject->IntA);
+	UE_NET_ASSERT_EQ(ClientDependentObject->IntA, ServerDependentObject->IntA);
+
+	// Trigger replication
+	ServerObject->IntA = 1;
+	ServerDependentObject->IntA = 1;
+
+	// Send and deliver packet
+	Server->PreSendUpdate();
+	Server->SendAndDeliverTo(Client, true);
+	Server->PostSendUpdate();
+
+	// Verify that only the server object has been updated
+	UE_NET_ASSERT_EQ(ClientObject->IntA, ServerObject->IntA);
+	UE_NET_ASSERT_NE(ClientDependentObject->IntA, ServerDependentObject->IntA);
+
+	// Add dependency
+	Server->ReplicationBridge->AddDependentObject(ServerObject->NetHandle, ServerDependentObject->NetHandle);
+
+	// Change a value on owner
+	ServerObject->IntA = 2;
+
+	// Send and deliver packet
+	Server->PreSendUpdate();
+	Server->SendAndDeliverTo(Client, true);
+	Server->PostSendUpdate();
+
+	// We now expect both objects to be in sync
+	UE_NET_ASSERT_EQ(ClientObject->IntA, ServerObject->IntA);
+	UE_NET_ASSERT_EQ(ClientDependentObject->IntA, ServerDependentObject->IntA);
+}
+
+UE_NET_TEST_FIXTURE(FReplicationSystemServerClientTestFixture, TestDependentObjectPollWithDirtyParent)
+{
+	UReplicationSystem* ReplicationSystem = Server->ReplicationSystem;
+
+	// Add a client
+	FReplicationSystemTestClient* Client = CreateClient();
+
+	// Spawn object on server
+	UTestReplicatedIrisObject* ServerObject0 = Server->CreateObject(0, 0);
+	UTestReplicatedIrisObject* ServerObject1 = Server->CreateObject(0, 0);
+
+	// Spawn second object on server that later will be added as a dependent object
+	UTestReplicatedIrisObject* ServerDependentObject = Server->CreateObject(0, 0);
+
+	// Setup different pollframeperiods for the objects
+	Server->ReplicationBridge->SetPollFramePeriod(ServerObject0, 10);
+	Server->ReplicationBridge->SetPollFramePeriod(ServerObject1, 20);
+	Server->ReplicationBridge->SetPollFramePeriod(ServerDependentObject, 40U);
+
+	// Add dependent object to both server objects
+	Server->ReplicationBridge->AddDependentObject(ServerObject0->NetHandle, ServerDependentObject->NetHandle);
+	Server->ReplicationBridge->AddDependentObject(ServerObject1->NetHandle, ServerDependentObject->NetHandle);
+
+	// Send and deliver packet, All objects are polled and replicated
+	Server->PreSendUpdate();
+	Server->SendAndDeliverTo(Client, true);
+	Server->PostSendUpdate();
+
+	// Verify state after initial replication
+	UTestReplicatedIrisObject* ClientDependentObject = Cast<UTestReplicatedIrisObject>(Client->GetReplicationBridge()->GetReplicatedObject(ServerDependentObject->NetHandle));
+	UTestReplicatedIrisObject* ClientObject0 = Cast<UTestReplicatedIrisObject>(Client->GetReplicationBridge()->GetReplicatedObject(ServerObject0->NetHandle));
+	UTestReplicatedIrisObject* ClientObject1 = Cast<UTestReplicatedIrisObject>(Client->GetReplicationBridge()->GetReplicatedObject(ServerObject1->NetHandle));
+
+	UE_NET_ASSERT_NE(ClientDependentObject, nullptr);
+	UE_NET_ASSERT_NE(ClientObject0, nullptr);
+	UE_NET_ASSERT_NE(ClientObject1, nullptr);
+	UE_NET_ASSERT_EQ(ClientDependentObject->IntA, 0);
+
+	const uint32 TickCount = 256;
+	int PrevRcvdValue = 0;
+	uint32 RepCount0 = 0U;
+	uint32 RepCount1 = 0U;
+	uint32 RepCountDependent = 0U;
+
+	for (uint32 CurrentFrame = 1U; CurrentFrame < TickCount; ++CurrentFrame)
+	{
+		ServerObject0->IntA = CurrentFrame;
+		ServerObject1->IntA = CurrentFrame;
+		ServerDependentObject->IntA = CurrentFrame;
+
+		// Send and deliver packet
+		Server->PreSendUpdate();
+		Server->SendAndDeliverTo(Client, true);
+		Server->PostSendUpdate();
+
+		const int RcvdValue = ClientDependentObject->IntA;
+		if (RcvdValue != PrevRcvdValue)
+		{
+			if (RcvdValue == ClientObject0->IntA)
+			{
+				++RepCount0;
+			}
+			if (RcvdValue == ClientObject1->IntA)
+			{
+				++RepCount1;
+			}
+			++RepCountDependent;
+		}
+		PrevRcvdValue = RcvdValue;		
+	}
+	// We expect the dependent object to have replicated more often than the parents
+	UE_NET_ASSERT_GE(RepCountDependent, RepCount0);
+	UE_NET_ASSERT_GE(RepCountDependent, RepCount1);
+}
+
+UE_NET_TEST_FIXTURE(FReplicationSystemServerClientTestFixture, TestDependentObjectIsPolledIfParentIsMarkedDirty)
+{
+	UReplicationSystem* ReplicationSystem = Server->ReplicationSystem;
+
+	// Add a client
+	FReplicationSystemTestClient* Client = CreateClient();
+
+	// Spawn object on server
+	UTestReplicatedIrisObject* ServerObject = Server->CreateObject(0, 0);
+	Server->ReplicationBridge->SetPollFramePeriod(ServerObject, 4U);
+
+	// Spawn second object add it as a dependency and bump poll period
+	UTestReplicatedIrisObject* ServerDependentObject = Server->CreateObject(0, 0);
+	Server->ReplicationBridge->SetPollFramePeriod(ServerDependentObject, 10U);
+	Server->ReplicationBridge->AddDependentObject(ServerObject->NetHandle, ServerDependentObject->NetHandle);
+
+	// Send and deliver packet
+	Server->PreSendUpdate();
+	Server->SendAndDeliverTo(Client, true);
+	Server->PostSendUpdate();
+
+	// Store Pointer to objects and verify state after initial replication
+	UTestReplicatedIrisObject* ClientObject = Cast<UTestReplicatedIrisObject>(Client->GetReplicationBridge()->GetReplicatedObject(ServerObject->NetHandle));
+	UTestReplicatedIrisObject* ClientDependentObject = Cast<UTestReplicatedIrisObject>(Client->GetReplicationBridge()->GetReplicatedObject(ServerDependentObject->NetHandle));
+
+	// Modify data
+	ServerObject->IntA = 1;
+	ServerDependentObject->IntA = 1;
+
+	// Send and deliver packet, we expect nothing to replicate
+	Server->PreSendUpdate();
+	Server->SendAndDeliverTo(Client, true);
+	Server->PostSendUpdate();
+
+	// Verify that nothing replicated
+	UE_NET_ASSERT_NE(ClientObject->IntA, ServerObject->IntA);
+	UE_NET_ASSERT_NE(ClientDependentObject->IntA, ServerDependentObject->IntA);
+
+	// Mark dependent dirty
+	ReplicationSystem->MarkDirty(ServerDependentObject->NetHandle);
+
+	// Send and deliver packet, we expect dependent object to have replicated 
+	Server->PreSendUpdate();
+	Server->SendAndDeliverTo(Client, true);
+	Server->PostSendUpdate();
+
+	// Verify that dependent object has replicated
+	UE_NET_ASSERT_NE(ClientObject->IntA, ServerObject->IntA);
+	UE_NET_ASSERT_EQ(ClientDependentObject->IntA, ServerDependentObject->IntA);
+
+	// Modify data
+	ServerObject->IntA = 2;
+	ServerDependentObject->IntA = 2;
+
+	// Mark parent dirty
+	ReplicationSystem->MarkDirty(ServerObject->NetHandle);
+
+	// Send and deliver packet, we expect both objects to have replicated
+	Server->PreSendUpdate();
+	Server->SendAndDeliverTo(Client, true);
+	Server->PostSendUpdate();
+
+	// Verify that both objects have replicated
+	UE_NET_ASSERT_EQ(ClientObject->IntA, ServerObject->IntA);
+	UE_NET_ASSERT_EQ(ClientDependentObject->IntA, ServerDependentObject->IntA);
+}
+
+}
diff --git a/Engine/Plugins/Runtime/ReplicationSystemTestPlugin/Source/Private/Tests/ReplicationSystem/TestDirtyNetObjectTracker.cpp b/Engine/Plugins/Runtime/ReplicationSystemTestPlugin/Source/Private/Tests/ReplicationSystem/TestDirtyNetObjectTracker.cpp
new file mode 100644
--- /dev/null
+++ b/Engine/Plugins/Runtime/ReplicationSystemTestPlugin/Source/Private/Tests/ReplicationSystem/TestDirtyNetObjectTracker.cpp
@@ -0,0 +1,120 @@
+// Copyright Epic Games, Inc. All Rights Reserved.
+
+#include "NetworkAutomationTest.h"
+#include "NetworkAutomationTestMacros.h"
+#include "Iris/ReplicationSystem/DirtyNetObjectTracker.h"
+#include "Iris/ReplicationSystem/NetHandle.h"
+#include "Iris/ReplicationSystem/ReplicationSystem.h"
+
+namespace UE::Net::Private
+{
+
+class FDirtyNetObjectTrackerTestFixture : public FNetworkAutomationTestSuiteFixture
+{
+public:
+	FDirtyNetObjectTrackerTestFixture()
+	: FNetworkAutomationTestSuiteFixture()
+	, DirtyNetObjectTracker(nullptr)
+	{
+	}
+
+	virtual void SetUp() override
+	{
+		DirtyNetObjectTracker = new FDirtyNetObjectTracker();
+	}
+
+	virtual void TearDown() override
+	{
+		delete DirtyNetObjectTracker;
+	}
+
+protected:
+	static constexpr uint32 ReplicationSystemId = 0;
+	static constexpr uint32 NetObjectIndexRangeStart = 1;
+	static constexpr uint32 NetObjectIndexRangeEnd = 1000;
+
+	FDirtyNetObjectTracker* DirtyNetObjectTracker;
+
+};
+
+class FDirtyNetObjectTrackerWithInitTestFixture : public FDirtyNetObjectTrackerTestFixture
+{
+protected:
+	virtual void SetUp() override
+	{
+		FDirtyNetObjectTrackerTestFixture::SetUp();
+		Init();
+	}
+
+	void Init();
+};
+
+UE_NET_TEST_FIXTURE(FDirtyNetObjectTrackerTestFixture, TestInit)
+{
+	const FDirtyNetObjectTrackerInitParams InitParams 
+	{
+		ReplicationSystemId,
+		NetObjectIndexRangeStart,
+		NetObjectIndexRangeEnd,
+	};
+
+	DirtyNetObjectTracker->Init(InitParams);
+	const FNetBitArrayView& DirtyObjects = DirtyNetObjectTracker->GetDirtyNetObjects();
+	UE_NET_ASSERT_EQ(DirtyObjects.GetNumBits(), NetObjectIndexRangeEnd + 1U);
+
+	UE_NET_ASSERT_FALSE(DirtyObjects.IsAnyBitSet()) << "Objects are marked as dirty before marking any object state as dirty";
+}
+
+UE_NET_TEST_FIXTURE(FDirtyNetObjectTrackerWithInitTestFixture, CannotMarkInvalidObjectAsDirty)
+{
+	MarkNetObjectStateDirty(ReplicationSystemId, FNetHandle().GetId());
+	MarkNetObjectStateDirty(ReplicationSystemId, NetObjectIndexRangeEnd + 1);
+
+	const FNetBitArrayView& DirtyObjects = DirtyNetObjectTracker->GetDirtyNetObjects();
+	UE_NET_ASSERT_FALSE(DirtyObjects.IsAnyBitSet());
+}
+
+UE_NET_TEST_FIXTURE(FDirtyNetObjectTrackerWithInitTestFixture, CanMarkValidObjectAsDirty)
+{
+	constexpr uint32 IndexInRange = NetObjectIndexRangeStart + (NetObjectIndexRangeEnd - NetObjectIndexRangeStart)/2;
+	MarkNetObjectStateDirty(ReplicationSystemId, IndexInRange);
+	MarkNetObjectStateDirty(ReplicationSystemId, NetObjectIndexRangeStart);
+	MarkNetObjectStateDirty(ReplicationSystemId, NetObjectIndexRangeEnd);
+
+	const FNetBitArrayView& DirtyObjects = DirtyNetObjectTracker->GetDirtyNetObjects();
+	UE_NET_ASSERT_TRUE(DirtyObjects.GetBit(IndexInRange));
+	UE_NET_ASSERT_TRUE(DirtyObjects.GetBit(NetObjectIndexRangeStart));
+	UE_NET_ASSERT_TRUE(DirtyObjects.GetBit(NetObjectIndexRangeEnd));
+}
+
+UE_NET_TEST_FIXTURE(FDirtyNetObjectTrackerWithInitTestFixture, CanClearDirtyObjects)
+{
+	MarkNetObjectStateDirty(ReplicationSystemId, NetObjectIndexRangeStart);
+	MarkNetObjectStateDirty(ReplicationSystemId, NetObjectIndexRangeEnd);
+
+	DirtyNetObjectTracker->ClearDirtyNetObjects();
+	const FNetBitArrayView& DirtyObjects = DirtyNetObjectTracker->GetDirtyNetObjects();
+	UE_NET_ASSERT_FALSE(DirtyObjects.IsAnyBitSet());
+}
+
+UE_NET_TEST(DirtyNetObjectTracker, MarkingObjectAsDirtyInNonExistingSystemDoesNotCrash)
+{
+	constexpr uint32 NonExistingReplicationSystemId = 4711;
+	constexpr uint32 ArbitraryNetObjectIndex = 1;
+	MarkNetObjectStateDirty(NonExistingReplicationSystemId, ArbitraryNetObjectIndex);
+}
+
+//
+void FDirtyNetObjectTrackerWithInitTestFixture::Init()
+{
+	const FDirtyNetObjectTrackerInitParams InitParams 
+	{
+		ReplicationSystemId,
+		NetObjectIndexRangeStart,
+		NetObjectIndexRangeEnd,
+	};
+
+	DirtyNetObjectTracker->Init(InitParams);
+}
+
+}
