Iris experimental replication system.

diff --git a/Engine/Plugins/Runtime/ReplicationSystemTestPlugin/Source/Private/Tests/ReplicationSystem/TestReplicationBridge.cpp b/Engine/Plugins/Runtime/ReplicationSystemTestPlugin/Source/Private/Tests/ReplicationSystem/TestReplicationBridge.cpp
new file mode 100644
--- /dev/null
+++ b/Engine/Plugins/Runtime/ReplicationSystemTestPlugin/Source/Private/Tests/ReplicationSystem/TestReplicationBridge.cpp
@@ -0,0 +1,190 @@
+// Copyright Epic Games, Inc. All Rights Reserved.
+
+#include "ReplicationSystemTestFixture.h"
+
+namespace UE::Net::Private
+{
+
+class FReplicationBridgeTestFixture : public FReplicationSystemTestFixture
+{
+protected:
+	enum : uint32
+	{
+		RegularStatePropertyCountInPropertyComponent = 2,
+	};
+
+	static constexpr uint32 PropertyComponentSize = 40U;
+};
+
+UE_NET_TEST_FIXTURE(FReplicationBridgeTestFixture, CreateObjectWithSingleFragment)
+{
+	UTestReplicatedIrisObject* TestObject = CreateObject(0, 0);
+
+	// Create NetHandle for the CreatedHandle,
+	FNetHandle CreatedHandle = ReplicationBridge->BeginReplication(TestObject);
+
+	// Verify that we managed to create a valid NetHandle	
+	UE_NET_ASSERT_TRUE(CreatedHandle.IsValid());
+	UE_NET_ASSERT_TRUE(CreatedHandle == TestObject->NetHandle);
+	UE_NET_ASSERT_TRUE(CreatedHandle == ReplicationBridge->GetReplicatedHandle(TestObject));
+
+	// Verify ReplicationProtocol
+	const FReplicationProtocol* Protocol = ReplicationSystem->GetReplicationProtocol(TestObject->NetHandle);
+	UE_NET_ASSERT_TRUE(Protocol != nullptr);
+	UE_NET_ASSERT_TRUE(Protocol->ProtocolIdentifier != 0);
+	UE_NET_ASSERT_EQ(1u, Protocol->ReplicationStateCount);
+	UE_NET_ASSERT_EQ(4u, Protocol->InternalTotalAlignment);
+	UE_NET_ASSERT_EQ(16u, Protocol->InternalTotalSize);
+	UE_NET_ASSERT_EQ(28u, Protocol->MaxExternalStateSize);
+	UE_NET_ASSERT_EQ(4u, Protocol->ChangeMaskBitCount);
+	UE_NET_ASSERT_EQ(EReplicationProtocolTraits::SupportsDeltaCompression, Protocol->ProtocolTraits);
+
+	// Destroy the handle
+	ReplicationBridge->EndReplication(TestObject);
+}
+
+UE_NET_TEST_FIXTURE(FReplicationBridgeTestFixture, CreateObjectWithMultipleFragments)
+{
+	constexpr uint32 PropertyComponentCount = 12;
+	constexpr uint32 IrisComponentCount = 0;
+	UTestReplicatedIrisObject* TestObject = CreateObject(PropertyComponentCount, IrisComponentCount);
+
+	// Create NetHandle for the CreatedHandle,
+	FNetHandle CreatedHandle = ReplicationBridge->BeginReplication(TestObject);
+
+	// Verify that we managed to create a valid NetHandle	
+	UE_NET_ASSERT_TRUE(CreatedHandle.IsValid());
+	UE_NET_ASSERT_TRUE(CreatedHandle == TestObject->NetHandle);
+	UE_NET_ASSERT_TRUE(CreatedHandle == ReplicationBridge->GetReplicatedHandle(TestObject));
+
+	// Verify ReplicationProtocol
+	const FReplicationProtocol* Protocol = ReplicationSystem->GetReplicationProtocol(TestObject->NetHandle);
+	UE_NET_ASSERT_TRUE(Protocol != nullptr);
+	UE_NET_ASSERT_EQ(25u, Protocol->ReplicationStateCount);
+	UE_NET_ASSERT_EQ(8u, Protocol->InternalTotalAlignment);
+	UE_NET_ASSERT_EQ(16u + PropertyComponentCount*PropertyComponentSize, Protocol->InternalTotalSize);
+	UE_NET_ASSERT_EQ(40u, Protocol->MaxExternalStateSize);
+	UE_NET_ASSERT_EQ(4U + PropertyComponentCount*RegularStatePropertyCountInPropertyComponent, Protocol->ChangeMaskBitCount);
+	UE_NET_ASSERT_EQ(EReplicationProtocolTraits::SupportsDeltaCompression, Protocol->ProtocolTraits);
+
+	// Destroy the handle
+	ReplicationBridge->EndReplication(TestObject);
+}
+
+UE_NET_TEST_FIXTURE(FReplicationBridgeTestFixture, CreateObjectWithMixedFragments)
+{
+	constexpr uint32 PropertyComponentCount = 12;
+	constexpr uint32 IrisComponentCount = 3;
+	UTestReplicatedIrisObject* TestObject = CreateObject(PropertyComponentCount, IrisComponentCount);
+
+	// Create NetHandle for the CreatedHandle,
+	FNetHandle CreatedHandle = ReplicationBridge->BeginReplication(TestObject);
+
+	// Verify that we managed to create a valid NetHandle	
+	UE_NET_ASSERT_TRUE(CreatedHandle.IsValid());
+	UE_NET_ASSERT_TRUE(CreatedHandle == TestObject->NetHandle);
+	UE_NET_ASSERT_TRUE(CreatedHandle == ReplicationBridge->GetReplicatedHandle(TestObject));
+
+	// Verify ReplicationProtocol
+	const FReplicationProtocol* Protocol = ReplicationSystem->GetReplicationProtocol(TestObject->NetHandle);;
+	UE_NET_ASSERT_TRUE(Protocol != nullptr);
+	UE_NET_ASSERT_EQ(28u, Protocol->ReplicationStateCount);
+	UE_NET_ASSERT_EQ(8u, Protocol->InternalTotalAlignment);
+	UE_NET_ASSERT_EQ(16u + PropertyComponentCount*PropertyComponentSize + IrisComponentCount*12U, Protocol->InternalTotalSize);
+	UE_NET_ASSERT_EQ(40u, Protocol->MaxExternalStateSize);
+	UE_NET_ASSERT_EQ(4U + PropertyComponentCount*RegularStatePropertyCountInPropertyComponent + IrisComponentCount*3U, Protocol->ChangeMaskBitCount);
+	UE_NET_ASSERT_EQ(EReplicationProtocolTraits::SupportsDeltaCompression, Protocol->ProtocolTraits);
+
+	// Destroy the handle
+	ReplicationBridge->EndReplication(TestObject);
+}
+
+UE_NET_TEST_FIXTURE(FReplicationBridgeTestFixture, CreateObjectWithDynamicStateComponent)
+{
+	constexpr uint32 DynamicStateComponentCount = 1U;
+	UTestReplicatedIrisObject* TestObject = CreateObjectWithDynamicState(0, 0, DynamicStateComponentCount);
+
+	// Create NetHandle for the CreatedHandle,
+	FNetHandle CreatedHandle = ReplicationBridge->BeginReplication(TestObject);
+
+	// Verify that we managed to create a valid NetHandle	
+	UE_NET_ASSERT_TRUE(CreatedHandle.IsValid());
+	UE_NET_ASSERT_TRUE(CreatedHandle == TestObject->NetHandle);
+	UE_NET_ASSERT_TRUE(CreatedHandle == ReplicationBridge->GetReplicatedHandle(TestObject));
+
+	// Verify ReplicationProtocol
+	const FReplicationProtocol* Protocol = ReplicationSystem->GetReplicationProtocol(TestObject->NetHandle);
+	UE_NET_ASSERT_NE(Protocol, nullptr);
+	UE_NET_ASSERT_NE(Protocol->ProtocolIdentifier, FReplicationProtocolIdentifier(0));
+	UE_NET_ASSERT_EQ(Protocol->ReplicationStateCount, 1U + DynamicStateComponentCount);
+	UE_NET_ASSERT_EQ(Protocol->ProtocolTraits, EReplicationProtocolTraits::HasDynamicState | EReplicationProtocolTraits::SupportsDeltaCompression);
+
+	// Destroy the handle
+	ReplicationBridge->EndReplication(TestObject);
+}
+
+}
+
+namespace UE::Net
+{
+
+FTestMessage& operator<<(FTestMessage& TestMessage, const EReplicationProtocolTraits Traits)
+{
+	if (Traits == EReplicationProtocolTraits::None)
+	{
+		return TestMessage << TEXT("EReplicationProtocolTraits::None");
+	}
+
+	// Horribly slow
+	uint32 Flags = uint32(Traits);
+	const uint32 FlagCount = FMath::CountBits(Flags);
+	if (FlagCount > 1U)
+	{
+		TestMessage << TEXT('(');
+	}
+
+	for (uint32 FlagIt = 0, FlagEndIt = FlagCount; FlagIt != FlagEndIt; ++FlagIt)
+	{
+		const uint32 Flag = Flags & uint32(-int32(Flags));
+		Flags ^= Flag;
+		if (FlagIt > 0)
+		{
+			TestMessage << TEXT(" | ");
+		}
+
+		switch (EReplicationProtocolTraits(Flag))
+		{
+		case EReplicationProtocolTraits::HasDynamicState:
+			TestMessage << TEXT("EReplicationProtocolTraits::HasDynamicState");
+			break;
+		case EReplicationProtocolTraits::HasLifetimeConditionals:
+			TestMessage << TEXT("EReplicationProtocolTraits::HasLifetimeConditionals");
+			break;
+		case EReplicationProtocolTraits::HasConditionalChangeMask:
+			TestMessage << TEXT("EReplicationProtocolTraits::HasConditionalChangeMask");
+			break;
+		case EReplicationProtocolTraits::HasConnectionSpecificSerialization:
+			TestMessage << TEXT("EReplicationProtocolTraits::HasConnectionSpecificSerialization");
+			break;
+		case EReplicationProtocolTraits::HasObjectReference:
+			TestMessage << TEXT("EReplicationProtocolTraits::HasObjectReference");
+			break;
+		case EReplicationProtocolTraits::SupportsDeltaCompression:
+			TestMessage << TEXT("EReplicationProtocolTraits::SupportsDeltaCompression");
+			break;
+		default:
+			TestMessage << TEXT("EReplicationProtocolTraits::UnknownFlagWithValue_") << Flag;
+			break;
+		}
+	}
+
+	if (FlagCount > 1U)
+	{
+		TestMessage << TEXT(')');
+	}
+
+	return TestMessage;
+}
+
+}
+
diff --git a/Engine/Plugins/Runtime/ReplicationSystemTestPlugin/Source/Private/Tests/ReplicationSystem/TestReplicationSystem.cpp b/Engine/Plugins/Runtime/ReplicationSystemTestPlugin/Source/Private/Tests/ReplicationSystem/TestReplicationSystem.cpp
new file mode 100644
--- /dev/null
+++ b/Engine/Plugins/Runtime/ReplicationSystemTestPlugin/Source/Private/Tests/ReplicationSystem/TestReplicationSystem.cpp
@@ -0,0 +1,1534 @@
+// Copyright Epic Games, Inc. All Rights Reserved.
+
+#include "ReplicationSystemServerClientTestFixture.h"
+#include "Iris/Serialization/NetBitStreamReader.h"
+#include "Iris/Serialization/NetBitStreamWriter.h"
+#include "Iris/ReplicationSystem/ReplicationSystem.h"
+
+#include "Iris/ReplicationSystem/ReplicationSystemInternal.h"
+#include "Iris/ReplicationSystem/NetTokenStore.h"
+
+namespace UE::Net::Private
+{
+
+UE_NET_TEST_FIXTURE(FReplicationSystemServerClientTestFixture, ReplicateAndDestroySingleObject)
+{
+	// Add a client
+	FReplicationSystemTestClient* Client = CreateClient();
+
+	// Spawn object on server
+	UReplicatedTestObject* ServerObject = Server->CreateObject(0,0);
+
+	// Send and deliver packet
+	Server->PreSendUpdate();
+	Server->SendAndDeliverTo(Client, true);
+	Server->PostSendUpdate();
+
+	// Verify that created server handle now also exists on client
+	UE_NET_ASSERT_TRUE(Client->GetReplicationBridge()->GetReplicatedObject(ServerObject->NetHandle) != nullptr);
+
+	// Destroy the spawned object on server
+	Server->DestroyObject(ServerObject);
+
+	// Send and deliver packet
+	Server->PreSendUpdate();
+	Server->SendAndDeliverTo(Client, true);
+	Server->PostSendUpdate();
+
+	// Verify that object now is destroyed on client as well
+	UE_NET_ASSERT_FALSE(Client->GetReplicationBridge()->GetReplicatedObject(ServerObject->NetHandle) != nullptr);
+}
+
+UE_NET_TEST_FIXTURE(FReplicationSystemServerClientTestFixture, TestDroppedDestroyed)
+{
+	// Add a client
+	FReplicationSystemTestClient* Client = CreateClient();
+
+	// Spawn object on server
+	UReplicatedTestObject* ServerObject = Server->CreateObject(0,0);
+
+	// Send and deliver packet
+	Server->PreSendUpdate();
+	Server->SendAndDeliverTo(Client, true);
+	Server->PostSendUpdate();
+
+	// Verify that created server handle now also exists on client
+	UE_NET_ASSERT_TRUE(Client->GetReplicationBridge()->GetReplicatedObject(ServerObject->NetHandle) != nullptr);
+
+	// Destroy the spawned object on server
+	Server->DestroyObject(ServerObject);
+
+	// Send and drop packet
+	Server->PreSendUpdate();
+	Server->SendAndDeliverTo(Client, false);
+	Server->PostSendUpdate();
+
+	// Send and deliver packet
+	Server->PreSendUpdate();
+	Server->SendAndDeliverTo(Client, true);
+	Server->PostSendUpdate();
+
+	// Verify that object now is destroyed on client
+	UE_NET_ASSERT_FALSE(Client->GetReplicationBridge()->GetReplicatedObject(ServerObject->NetHandle) != nullptr);
+}
+
+UE_NET_TEST_FIXTURE(FReplicationSystemServerClientTestFixture, TestDestroyWhilePendingCreate)
+{
+	// Add a client
+	FReplicationSystemTestClient* Client = CreateClient();
+
+	// Spawn object on server
+	UReplicatedTestObject* ServerObject = Server->CreateObject(0,0);
+
+	// Send and deliver packet
+	Server->PreSendUpdate();
+	Server->SendAndDeliverTo(Client, false);
+
+	Server->PostSendUpdate();
+
+	// Destroy
+	Server->DestroyObject(ServerObject);
+
+	// Verify that created server handle does not exists on client
+	UE_NET_ASSERT_FALSE(Client->GetReplicationBridge()->GetReplicatedObject(ServerObject->NetHandle) != nullptr);
+
+	// Send and drop packet
+	Server->PreSendUpdate();
+	Server->SendAndDeliverTo(Client, false);
+	Server->PostSendUpdate();
+
+	// Send and deliver packet
+	Server->PreSendUpdate();
+	Server->SendAndDeliverTo(Client, true);
+	Server->PostSendUpdate();
+
+	// Verify that object now is destroyed on client
+	UE_NET_ASSERT_FALSE(Client->GetReplicationBridge()->GetReplicatedObject(ServerObject->NetHandle) != nullptr);
+}
+
+UE_NET_TEST_FIXTURE(FReplicationSystemServerClientTestFixture, TestDestroyWhileWaitingOnCreate)
+{
+	// Add a client
+	FReplicationSystemTestClient* Client = CreateClient();
+
+	// Spawn object on server
+	UReplicatedTestObject* ServerObject = Server->CreateObject(0,0);
+
+	// Send packet with create
+	Server->PreSendUpdate();
+	Server->SendTo(Client);
+	Server->PostSendUpdate();
+
+	// Destroy while we are waiting for confirmation
+	Server->DestroyObject(ServerObject);
+
+	// Send packet with destroy
+	Server->PreSendUpdate();
+	Server->SendTo(Client);
+	Server->PostSendUpdate();
+
+	// Drop and report packet with create as lost
+	Server->DeliverTo(Client, false);
+
+	// deliver packet with destroy
+	Server->DeliverTo(Client, true);
+}
+
+UE_NET_TEST_FIXTURE(FReplicationSystemServerClientTestFixture, TestDestroyWithDataInFlight)
+{
+	// Add a client
+	FReplicationSystemTestClient* Client = CreateClient();
+
+	// Spawn object on server
+	UTestReplicatedIrisObject* ServerObject = Server->CreateObject(0,0);
+	Server->GetReplicationSystem()->SetStaticPriority(ServerObject->NetHandle, 1.f);
+
+	// Send packet with create
+	Server->PreSendUpdate();
+	Server->SendAndDeliverTo(Client, true);
+	Server->PostSendUpdate();
+
+	// Modify some data to mark object dirty
+	ServerObject->IntA = 13;
+
+	// Trigger send 
+	Server->PreSendUpdate();
+	UE_NET_ASSERT_TRUE(Server->SendTo(Client));
+	Server->PostSendUpdate();
+
+	// Destroy while we are waiting for confirmation
+	Server->DestroyObject(ServerObject);
+
+	// Send packet with destroy
+	Server->PreSendUpdate();
+	UE_NET_ASSERT_TRUE(Server->SendTo(Client));
+	Server->PostSendUpdate();
+
+	// Drop and report packet with create as lost
+	Server->DeliverTo(Client, false);
+
+	// deliver packet with destroy
+	Server->DeliverTo(Client, true);
+}
+
+UE_NET_TEST_FIXTURE(FReplicationSystemServerClientTestFixture, ReplicateAndDestroySubObject)
+{
+	// Add a client
+	FReplicationSystemTestClient* Client = CreateClient();
+
+	// Spawn object on server
+	UReplicatedTestObject* ServerObject = Server->CreateObject(0,0);
+
+	// Spawn second object on server as a subobject
+	UReplicatedTestObject* ServerSubObject = Server->CreateSubObject(ServerObject->NetHandle, 0, 0);
+
+	// Send and deliver packet
+	Server->PreSendUpdate();
+	Server->SendAndDeliverTo(Client, true);
+	Server->PostSendUpdate();
+
+	// Verify that created server handles now also exists on client
+	UE_NET_ASSERT_TRUE(Client->GetReplicationBridge()->GetReplicatedObject(ServerObject->NetHandle) != nullptr);
+	UE_NET_ASSERT_TRUE(Client->GetReplicationBridge()->GetReplicatedObject(ServerSubObject->NetHandle) != nullptr);
+
+	// Destroy the spawned subobject on server
+	Server->DestroyObject(ServerSubObject);
+
+	// Send and deliver packet
+	Server->PreSendUpdate();
+	Server->SendAndDeliverTo(Client, true);
+	Server->PostSendUpdate();
+
+	// Verify that object now is destroyed on client as well
+	UE_NET_ASSERT_FALSE(Client->GetReplicationBridge()->GetReplicatedObject(ServerSubObject->NetHandle) != nullptr);
+}
+
+UE_NET_TEST_FIXTURE(FReplicationSystemServerClientTestFixture, ReplicateAndDestroyMultipleSubObjects)
+{
+	// Add a client
+	FReplicationSystemTestClient* Client = CreateClient();
+
+	// Spawn object on server
+	UReplicatedTestObject* ServerObject = Server->CreateObject(0, 0);
+
+	// Spawn second object on server as a subobject
+	UReplicatedTestObject* ServerSubObject = Server->CreateSubObject(ServerObject->NetHandle, 0, 0);
+
+	// Send and deliver packet
+	Server->PreSendUpdate();
+	Server->SendAndDeliverTo(Client, true);
+	Server->PostSendUpdate();
+
+	// Verify that created server handles now also exists on client
+	UE_NET_ASSERT_NE(Client->GetReplicationBridge()->GetReplicatedObject(ServerObject->NetHandle), nullptr);
+	UE_NET_ASSERT_NE(Client->GetReplicationBridge()->GetReplicatedObject(ServerSubObject->NetHandle), nullptr);
+
+	// Destroy the spawned subobject on server
+	Server->DestroyObject(ServerSubObject);
+
+	// Send and deliver packet
+	Server->PreSendUpdate();
+	Server->SendAndDeliverTo(Client, true);
+	Server->PostSendUpdate();
+
+	// Verify that subobject now is destroyed on client as well
+	UE_NET_ASSERT_EQ(Client->GetReplicationBridge()->GetReplicatedObject(ServerSubObject->NetHandle), nullptr);
+
+	// Spawn second object on server as a subobject
+	ServerSubObject = Server->CreateSubObject(ServerObject->NetHandle, 0, 0);
+
+	// Send and deliver packet
+	Server->PreSendUpdate();
+	Server->SendAndDeliverTo(Client, true);
+	Server->PostSendUpdate();
+
+	// Verify that second subobject replicated properly to server
+	UE_NET_ASSERT_NE(Client->GetReplicationBridge()->GetReplicatedObject(ServerSubObject->NetHandle), nullptr);
+
+	// Destroy the spawned object on server
+	Server->DestroyObject(ServerSubObject);
+
+	// Send and deliver packet
+	Server->PreSendUpdate();
+	Server->SendAndDeliverTo(Client, true);
+	Server->PostSendUpdate();
+
+	// Verify that object now is destroyed on client as well
+	UE_NET_ASSERT_EQ(Client->GetReplicationBridge()->GetReplicatedObject(ServerSubObject->NetHandle), nullptr);
+}
+
+UE_NET_TEST_FIXTURE(FReplicationSystemServerClientTestFixture, ReplicateAndDestroySubObjectAndDestroyOwner)
+{
+	// Add a client
+	FReplicationSystemTestClient* Client = CreateClient();
+
+	// Spawn object on server
+	UReplicatedTestObject* ServerObject = Server->CreateObject(0, 0);
+
+	// Spawn second object on server as a subobject
+	UReplicatedTestObject* ServerSubObject = Server->CreateSubObject(ServerObject->NetHandle, 0, 0);
+
+	// Send and deliver packet
+	Server->PreSendUpdate();
+	Server->SendAndDeliverTo(Client, true);
+	Server->PostSendUpdate();
+
+	// Verify that created server handles now also exists on client
+	UE_NET_ASSERT_NE(Client->GetReplicationBridge()->GetReplicatedObject(ServerObject->NetHandle), nullptr);
+	UE_NET_ASSERT_NE(Client->GetReplicationBridge()->GetReplicatedObject(ServerSubObject->NetHandle), nullptr);
+
+	// Destroy the spawned subobject on server
+	Server->DestroyObject(ServerSubObject);
+
+	// Destroy owner after spawned subobject on server
+	Server->DestroyObject(ServerObject);
+
+	// Send and deliver packet
+	Server->PreSendUpdate();
+	Server->SendAndDeliverTo(Client, true);
+	Server->PostSendUpdate();
+
+	// Verify that subobject now is destroyed on client as well
+	UE_NET_ASSERT_EQ(Client->GetReplicationBridge()->GetReplicatedObject(ServerSubObject->NetHandle), nullptr);
+}
+
+UE_NET_TEST_FIXTURE(FReplicationSystemServerClientTestFixture, ReplicateAndDestroySubObjectAndDestroyOwnerWithDataInFlight)
+{
+	// Add a client
+	FReplicationSystemTestClient* Client = CreateClient();
+
+	// Spawn object on server
+	UReplicatedTestObject* ServerObject = Server->CreateObject(0, 0);
+
+	// Spawn second object on server as a subobject
+	UReplicatedTestObject* ServerSubObject = Server->CreateSubObject(ServerObject->NetHandle, 0, 0);
+
+	// Send and deliver packet
+	Server->PreSendUpdate();
+	Server->SendAndDeliverTo(Client, true);
+	Server->PostSendUpdate();
+
+	// Verify that created server handles now also exists on client
+	UE_NET_ASSERT_NE(Client->GetReplicationBridge()->GetReplicatedObject(ServerObject->NetHandle), nullptr);
+	UE_NET_ASSERT_NE(Client->GetReplicationBridge()->GetReplicatedObject(ServerSubObject->NetHandle), nullptr);
+
+	// Destroy the spawned subobject on server
+	Server->DestroyObject(ServerSubObject);
+
+	// Send and drop packet
+	Server->PreSendUpdate();
+	Server->SendAndDeliverTo(Client, false);
+	Server->PostSendUpdate();
+
+	// Destroy owner after we spawned subobject on server
+	Server->DestroyObject(ServerObject);
+
+	// Send and deliver packet
+	Server->PreSendUpdate();
+	Server->SendAndDeliverTo(Client, true);
+	Server->PostSendUpdate();
+
+	// Verify that subobject now is destroyed on client as well
+	UE_NET_ASSERT_EQ(Client->GetReplicationBridge()->GetReplicatedObject(ServerSubObject->NetHandle), nullptr);
+}
+
+UE_NET_TEST_FIXTURE(FReplicationSystemServerClientTestFixture, ReplicateAndDestroySubObjectWithLostData)
+{
+	// Add a client
+	FReplicationSystemTestClient* Client = CreateClient();
+
+	// Spawn object on server
+	UReplicatedTestObject* ServerObject = Server->CreateObject(0, 0);
+
+	// Spawn second object on server as a subobject
+	UReplicatedTestObject* ServerSubObject = Server->CreateSubObject(ServerObject->NetHandle, 0, 0);
+
+	// Send and deliver packet
+	Server->PreSendUpdate();
+	Server->SendAndDeliverTo(Client, true);
+	Server->PostSendUpdate();
+
+	// Verify that created server handles now also exists on client
+	UE_NET_ASSERT_NE(Client->GetReplicationBridge()->GetReplicatedObject(ServerObject->NetHandle), nullptr);
+	UE_NET_ASSERT_NE(Client->GetReplicationBridge()->GetReplicatedObject(ServerSubObject->NetHandle), nullptr);
+
+	// Destroy the spawned subobject on server
+	Server->DestroyObject(ServerSubObject);
+
+	// Send and drop packet
+	Server->PreSendUpdate();
+	Server->SendAndDeliverTo(Client, false);
+	Server->PostSendUpdate();
+
+	// Send and deliver packet
+	Server->PreSendUpdate();
+	Server->SendAndDeliverTo(Client, true);
+	Server->PostSendUpdate();
+
+	// Verify that subobject now is destroyed on client as well
+	UE_NET_ASSERT_EQ(Client->GetReplicationBridge()->GetReplicatedObject(ServerSubObject->NetHandle), nullptr);
+}
+
+UE_NET_TEST_FIXTURE(FReplicationSystemServerClientTestFixture, ReplicateAndDestroySubObjectPendingCreateConfirmation)
+{
+	// Add a client
+	FReplicationSystemTestClient* Client = CreateClient();
+
+	// Spawn object on server
+	UReplicatedTestObject* ServerObject = Server->CreateObject(0, 0);
+
+	// Spawn second object on server as a subobject
+	UReplicatedTestObject* ServerSubObject = Server->CreateSubObject(ServerObject->NetHandle, 0, 0);
+
+	uint32 NumUnAcknowledgedPackets = 0;
+	// Send and deliver packet
+	Server->PreSendUpdate();
+	Server->SendTo(Client);
+	Server->PostSendUpdate();
+
+	// Destroy the spawned subobject on server
+	Server->DestroyObject(ServerSubObject);
+
+	// We have no data to send but we want to tick ReplicationSystem to capture state change
+	Server->PreSendUpdate();
+	UE_NET_ASSERT_FALSE(Server->SendTo(Client));
+	Server->PostSendUpdate();
+
+	Server->DeliverTo(Client, false);
+	// As we did not send any data we do not have anything to deliver
+	//Server->DeliverTo(Client, true);
+
+	Server->PreSendUpdate();
+	Server->SendAndDeliverTo(Client, true);
+	Server->PostSendUpdate();
+
+	// Send and deliver packet
+	Server->PreSendUpdate();
+	Server->SendAndDeliverTo(Client, true);
+	Server->PostSendUpdate();
+
+	// Verify that subobject now is destroyed on client as well
+	UE_NET_ASSERT_EQ(Client->GetReplicationBridge()->GetReplicatedObject(ServerSubObject->NetHandle), nullptr);
+}
+
+UE_NET_TEST_FIXTURE(FReplicationSystemServerClientTestFixture, TestSubObjectDefaultReplicationOrder)
+{
+	// Add a client
+	FReplicationSystemTestClient* Client = CreateClient();
+
+	// Spawn object on server
+	UReplicatedTestObject* ServerObject = Server->CreateObject(0, 0);
+
+	// Spawn some subobjects
+	UReplicatedSubObjectOrderObject* ServerSubObject0 = Server->CreateSubObject<UReplicatedSubObjectOrderObject>(ServerObject->NetHandle);
+	UReplicatedSubObjectOrderObject* ServerSubObject1 = Server->CreateSubObject<UReplicatedSubObjectOrderObject>(ServerObject->NetHandle);
+	UReplicatedSubObjectOrderObject* ServerSubObject2 = Server->CreateSubObject<UReplicatedSubObjectOrderObject>(ServerObject->NetHandle);
+
+	UReplicatedSubObjectOrderObject::RepOrderCounter = 0U;
+
+	// Send and deliver packet
+	Server->PreSendUpdate();
+	Server->SendAndDeliverTo(Client, true);
+	Server->PostSendUpdate();
+
+	// Verify that objects have replicated
+	UReplicatedSubObjectOrderObject* ClientSubObject0 = Cast<UReplicatedSubObjectOrderObject>(Client->GetReplicationBridge()->GetReplicatedObject(ServerSubObject0->NetHandle));
+	UReplicatedSubObjectOrderObject* ClientSubObject1 = Cast<UReplicatedSubObjectOrderObject>(Client->GetReplicationBridge()->GetReplicatedObject(ServerSubObject1->NetHandle));
+	UReplicatedSubObjectOrderObject* ClientSubObject2 = Cast<UReplicatedSubObjectOrderObject>(Client->GetReplicationBridge()->GetReplicatedObject(ServerSubObject2->NetHandle));
+
+	UE_NET_ASSERT_NE(ClientSubObject0, nullptr);
+	UE_NET_ASSERT_NE(ClientSubObject1, nullptr);
+	UE_NET_ASSERT_NE(ClientSubObject2, nullptr);
+
+	// Verify that they have replicated in expected order
+	UE_NET_ASSERT_EQ(ClientSubObject0->LastRepOrderCounter, 1U);
+	UE_NET_ASSERT_GT(ClientSubObject1->LastRepOrderCounter, ClientSubObject0->LastRepOrderCounter);
+	UE_NET_ASSERT_GT(ClientSubObject2->LastRepOrderCounter, ClientSubObject1->LastRepOrderCounter);
+}
+
+UE_NET_TEST_FIXTURE(FReplicationSystemServerClientTestFixture, TestSubObjectSpecfiedReplicationOrder)
+{
+	// Add a client
+	FReplicationSystemTestClient* Client = CreateClient();
+
+	// Spawn object on server
+	UReplicatedTestObject* ServerObject = Server->CreateObject(0, 0);
+
+	// Spawn some subobjects
+	UReplicatedSubObjectOrderObject* ServerSubObject0 = Server->CreateSubObject<UReplicatedSubObjectOrderObject>(ServerObject->NetHandle);
+	// Specify Subobject1 to replicate with SubObject0, which means that it will replicate before Subobjet0 is replicated
+	UReplicatedSubObjectOrderObject* ServerSubObject1 = Server->CreateSubObject<UReplicatedSubObjectOrderObject>(ServerObject->NetHandle, ServerSubObject0->NetHandle, UReplicationBridge::ESubObjectInsertionOrder::ReplicateWith);
+	// Specify SubObect 2 to replicate with no specific order (it will be added to the owner and thus replicate last)
+	UReplicatedSubObjectOrderObject* ServerSubObject2 = Server->CreateSubObject<UReplicatedSubObjectOrderObject>(ServerObject->NetHandle, ServerSubObject1->NetHandle, UReplicationBridge::ESubObjectInsertionOrder::None);
+
+	UReplicatedSubObjectOrderObject::RepOrderCounter = 0U;
+
+	// Send and deliver packet
+	Server->PreSendUpdate();
+	Server->SendAndDeliverTo(Client, true);
+	Server->PostSendUpdate();
+
+	// Verify that objects have replicated
+	UReplicatedSubObjectOrderObject* ClientSubObject0 = Cast<UReplicatedSubObjectOrderObject>(Client->GetReplicationBridge()->GetReplicatedObject(ServerSubObject0->NetHandle));
+	UReplicatedSubObjectOrderObject* ClientSubObject1 = Cast<UReplicatedSubObjectOrderObject>(Client->GetReplicationBridge()->GetReplicatedObject(ServerSubObject1->NetHandle));
+	UReplicatedSubObjectOrderObject* ClientSubObject2 = Cast<UReplicatedSubObjectOrderObject>(Client->GetReplicationBridge()->GetReplicatedObject(ServerSubObject2->NetHandle));
+
+	UE_NET_ASSERT_NE(ClientSubObject0, nullptr);
+	UE_NET_ASSERT_NE(ClientSubObject1, nullptr);
+	UE_NET_ASSERT_NE(ClientSubObject2, nullptr);
+
+	// Verify that they have replicated in expected order setup earlier
+	UE_NET_ASSERT_EQ(ClientSubObject1->LastRepOrderCounter, 1U);
+	UE_NET_ASSERT_EQ(ClientSubObject0->LastRepOrderCounter, 2U);
+	UE_NET_ASSERT_EQ(ClientSubObject2->LastRepOrderCounter, 3U);
+}
+
+class FTestNetTokensFixture : public FReplicationSystemServerClientTestFixture
+{
+public:
+	FStringTokenStore* ServerStringTokenStore = nullptr;
+	FStringTokenStore* ClientStringTokenStore = nullptr;
+	FReplicationSystemTestClient* Client = nullptr;
+	const FNetTokenStoreState* ClientRemoteNetTokenState;
+
+	FNetToken CreateAndExportNetToken(const FString& TokenString)
+	{
+		FNetToken Token = ServerStringTokenStore->GetOrCreateToken(TokenString);
+		Server->GetConnectionInfo(Client->ConnectionIdOnServer).NetTokenDataStream->AddNetTokenForExplicitExport(Token);
+
+		return Token;
+	}
+
+	virtual void SetUp() override
+	{
+		FReplicationSystemServerClientTestFixture::SetUp();
+
+		Client = CreateClient();
+
+		ServerStringTokenStore = Server->GetReplicationSystem()->GetStringTokenStore();
+		ClientStringTokenStore = Client->GetReplicationSystem()->GetStringTokenStore();
+		ClientRemoteNetTokenState = Client->GetConnectionInfo(Client->LocalConnectionId).NetTokenDataStream->GetRemoteNetTokenStoreState();
+	}
+};
+
+
+
+UE_NET_TEST_FIXTURE(FTestNetTokensFixture, NetToken)
+{
+	// Create token
+	FString TokenStringA(TEXT("MyStringToken"));
+	FNetToken StringTokenA = CreateAndExportNetToken(TokenStringA);
+
+	// Send and drop packet
+	Server->PreSendUpdate();
+	Server->SendAndDeliverTo(Client, false);
+	Server->PostSendUpdate();
+
+	// Verify that we cannot resolve the token on the client
+	UE_NET_ASSERT_NE(TokenStringA, FString(ClientStringTokenStore->ResolveRemoteToken(StringTokenA, *ClientRemoteNetTokenState)));
+
+	// Send and deliver packet
+	Server->PreSendUpdate();
+	Server->SendAndDeliverTo(Client, true);
+	Server->PostSendUpdate();
+
+	// Verify that we can resolve the token on the client
+	UE_NET_ASSERT_EQ(TokenStringA, FString(ClientStringTokenStore->ResolveRemoteToken(StringTokenA, *ClientRemoteNetTokenState)));
+}
+
+UE_NET_TEST_FIXTURE(FTestNetTokensFixture, NetTokenResendWithFullPacket)
+{
+	// Create token
+	FString TokenStringA(TEXT("MyStringToken"));
+	FNetToken StringTokenA = CreateAndExportNetToken(TokenStringA);
+
+	// Send and drop packet
+	Server->PreSendUpdate();
+	Server->SendAndDeliverTo(Client, false);
+	Server->PostSendUpdate();
+
+	// Limit packet size
+	Server->SetMaxSendPacketSize(128U);
+
+	// Create a new token that will not fit in the packet and only fit the resend data
+	FString TokenStringB(TEXT("MyLongStringMyLongStringMyLongStringMyLongStringMyLongStringMyLongStringMyLongStringMyLongStringMyLongStringMyLongStringMyLongStringMyLongStringMyLongStringMyLongStringMyLongStringMyLongStringMyLongStringMyLongStringMyLongStringMyLongStringMyLongStringMyLongStringMyLongStringMyLongString"));
+	FNetToken StringTokenB = CreateAndExportNetToken(TokenStringB);
+
+	// Send and deliver packet
+	Server->PreSendUpdate();
+	Server->SendAndDeliverTo(Client, true);
+	Server->PostSendUpdate();
+
+	// Verify that we can resolve the token first token on the client even though second one should not fit
+	UE_NET_ASSERT_EQ(TokenStringA, FString(ClientStringTokenStore->ResolveRemoteToken(StringTokenA, *ClientRemoteNetTokenState)));
+	UE_NET_ASSERT_NE(TokenStringB, FString(ClientStringTokenStore->ResolveRemoteToken(StringTokenB, *ClientRemoteNetTokenState)));
+
+	// Restore packet size and make sure that we get the second token through
+	Server->SetMaxSendPacketSize(1024U);
+
+	// Send and deliver packet
+	Server->PreSendUpdate();
+	Server->SendAndDeliverTo(Client, true);
+	Server->PostSendUpdate();
+
+	UE_NET_ASSERT_EQ(TokenStringB, FString(ClientStringTokenStore->ResolveRemoteToken(StringTokenB, *ClientRemoteNetTokenState)));
+}
+
+UE_NET_TEST_FIXTURE(FTestNetTokensFixture, NetTokenResendWithFullPacketAfterFirstResend)
+{
+	// Create token
+	FString TestStringA(TEXT("MyLongStringMyLongStringMyLongStringMyLongStringMyLongStringMyLongStringMyLongStringMyLongStringMyLongStringMyLongStringMyLongStringMyLongStringMyLongStringMyLongStringMyLongStringMyLongStringMyLongStringMyLongStringMyLongStringMyLongStringMyLongStringMyLongStringMyLongStringMyLongString"));
+	FNetToken StringTokenA = CreateAndExportNetToken(TestStringA);	
+
+	// Send and delay delivery
+	Server->PreSendUpdate();
+	Server->SendTo(Client);
+	Server->PostSendUpdate();
+
+	// Create a new token that will not fit in the packet and only fit the resend data
+	FString TestStringB(TEXT("MyOtherLongStringMyLongStringMyLongStringMyLongStringMyLongStringMyLongStringMyLongStringMyLongStringMyLongStringMyLongStringMyLongStringMyLongStringMyLongStringMyLongStringMyLongStringMyLongStringMyLongStringMyLongStringMyLongStringMyLongStringMyLongStringMyLongStringMyLongStringMyLongString"));
+	FNetToken StringTokenB = CreateAndExportNetToken(TestStringB);
+
+	// Send and delay delivery
+	Server->PreSendUpdate();
+	Server->SendTo(Client);
+	Server->PostSendUpdate();
+
+	Server->DeliverTo(Client, false);
+	Server->DeliverTo(Client, false);
+
+	// Verify that tokens has not been received
+	UE_NET_ASSERT_NE(TestStringA, FString(ClientStringTokenStore->ResolveRemoteToken(StringTokenA, *ClientRemoteNetTokenState)));
+	UE_NET_ASSERT_NE(TestStringB, FString(ClientStringTokenStore->ResolveRemoteToken(StringTokenB, *ClientRemoteNetTokenState)));
+
+	// Send and deliver packet which now should contain two entries in the resend queue
+	Server->SetMaxSendPacketSize(1024);
+
+	Server->PreSendUpdate();
+	Server->SendAndDeliverTo(Client, true);
+	Server->PostSendUpdate();
+
+	// Verify that we can resolve the token
+	UE_NET_ASSERT_EQ(TestStringA, FString(ClientStringTokenStore->ResolveRemoteToken(StringTokenA, *ClientRemoteNetTokenState)));
+	UE_NET_ASSERT_EQ(TestStringB, FString(ClientStringTokenStore->ResolveRemoteToken(StringTokenB, *ClientRemoteNetTokenState)));
+}
+
+UE_NET_TEST_FIXTURE(FTestNetTokensFixture, NetTokenSequenceTest)
+{
+	const FString TestStrings[] = {
+		FString(TEXT("TokenA")),
+		FString(TEXT("TokenB")),
+		FString(TEXT("TokenC")),
+		FString(TEXT("TokenD")),
+		FString(TEXT("TokenE")),
+		FString(TEXT("TokenF")),
+	};
+
+	const uint32 TokenCount = UE_ARRAY_COUNT(TestStrings);
+
+	// Create token
+	FNetToken StringTokenA = CreateAndExportNetToken(TestStrings[0]);
+	FNetToken StringTokenB = CreateAndExportNetToken(TestStrings[1]);
+
+	// Send packet
+	Server->PreSendUpdate();
+	Server->SendTo(Client);
+	Server->PostSendUpdate();
+
+	// Create token
+	FNetToken StringTokenC = CreateAndExportNetToken(TestStrings[2]);
+
+	// Create token
+	FNetToken StringTokenD = CreateAndExportNetToken(TestStrings[3]);
+
+	// Send packet
+	Server->PreSendUpdate();
+	Server->SendTo(Client);
+	Server->PostSendUpdate();
+
+	// Drop packet 
+	Server->DeliverTo(Client, false);
+
+	// Deliver packet 
+	Server->DeliverTo(Client, true);
+
+	// Create local tokens
+	ClientStringTokenStore->GetOrCreateToken(TEXT("LocalTokenA"));
+	ClientStringTokenStore->GetOrCreateToken(TEXT("LocalTokenB"));
+
+	// Send packet with resend data
+	Server->PreSendUpdate();
+	Server->SendAndDeliverTo(Client, true);
+	Server->PostSendUpdate();
+
+	UE_NET_ASSERT_EQ(TestStrings[0], FString(ClientStringTokenStore->ResolveRemoteToken(StringTokenA, *ClientRemoteNetTokenState)));
+	UE_NET_ASSERT_EQ(TestStrings[1], FString(ClientStringTokenStore->ResolveRemoteToken(StringTokenB, *ClientRemoteNetTokenState)));
+	UE_NET_ASSERT_EQ(TestStrings[2], FString(ClientStringTokenStore->ResolveRemoteToken(StringTokenC, *ClientRemoteNetTokenState)));
+	UE_NET_ASSERT_EQ(TestStrings[3], FString(ClientStringTokenStore->ResolveRemoteToken(StringTokenD, *ClientRemoteNetTokenState)));
+}
+
+UE_NET_TEST_FIXTURE(FTestNetTokensFixture, NetTokenResendAndDataInSamePacketTest)
+{
+	const FString TestStrings[] = {
+		FString(TEXT("TokenA")),
+		FString(TEXT("TokenB")),
+	};
+
+	const uint32 TokenCount = UE_ARRAY_COUNT(TestStrings);
+
+
+	// Create token
+	FNetToken StringTokenA = CreateAndExportNetToken(TestStrings[0]);
+
+	// Send packet
+	Server->PreSendUpdate();
+	Server->SendTo(Client);
+	Server->PostSendUpdate();
+
+	// drop data
+	Server->DeliverTo(Client, false);
+
+	// Create token
+	FNetToken StringTokenB = CreateAndExportNetToken(TestStrings[1]);
+
+	// Send and deliver packet
+	Server->PreSendUpdate();
+	Server->SendAndDeliverTo(Client, true);
+	Server->PostSendUpdate();
+
+	UE_NET_ASSERT_EQ(TestStrings[0], FString(ClientStringTokenStore->ResolveRemoteToken(StringTokenA, *ClientRemoteNetTokenState)));
+	UE_NET_ASSERT_EQ(TestStrings[1], FString(ClientStringTokenStore->ResolveRemoteToken(StringTokenB, *ClientRemoteNetTokenState)));
+}
+
+
+UE_NET_TEST_FIXTURE(FReplicationSystemServerClientTestFixture, AddRemoveFromConnectionScopeTest)
+{
+	UReplicationSystem* ReplicationSystem = Server->ReplicationSystem;
+
+	// Add a client
+	FReplicationSystemTestClient* Client = CreateClient();
+
+	// Spawn object on server
+	UTestReplicatedIrisObject* ServerObject = Server->CreateObject(0,0);
+
+	// Add to group
+	FNetObjectGroupHandle Group = ReplicationSystem->CreateGroup();
+	ReplicationSystem->AddToGroup(Group, ServerObject->NetHandle);
+
+	ReplicationSystem->AddGroupFilter(Group);
+	ReplicationSystem->SetGroupFilterStatus(Group, ENetFilterStatus::Allow);
+
+	// Start replicating object
+	
+	// Send packet
+	// Expected state to be WaitOnCreateConfirmation
+	Server->PreSendUpdate();
+	Server->SendTo(Client);
+	Server->PostSendUpdate();
+
+	// Make sure we have data in flight
+	++ServerObject->IntA;
+
+	// Disallow group to trigger state change from PendingCreateConfirmation->PendingDestroy
+	ReplicationSystem->SetGroupFilterStatus(Group, ENetFilterStatus::Disallow);	
+
+	// Send packet
+	// Expected state to be WaitOnDestroyConfirmation
+	Server->PreSendUpdate();
+	Server->SendTo(Client);
+	Server->PostSendUpdate();
+
+	// Allow group to trigger state to ensure that we restart replication since we have not actually created the object on the client
+	ReplicationSystem->SetGroupFilterStatus(Group, ENetFilterStatus::Allow);
+
+	// Expect client to create object
+	Server->DeliverTo(Client, true);
+	UE_NET_ASSERT_TRUE(Client->GetReplicationBridge()->GetReplicatedObject(ServerObject->NetHandle) != nullptr);
+
+	// Expect client to destroy object and server to move to Destroyed state
+	Server->DeliverTo(Client, true);
+	UE_NET_ASSERT_TRUE(Client->GetReplicationBridge()->GetReplicatedObject(ServerObject->NetHandle) == nullptr);
+
+	// Trigger replication
+	++ServerObject->IntA;
+
+	// Send packet
+	// Invalid -> PendingCreate
+	// PendingCreate->WaitOnCreateConfirmation
+	Server->PreSendUpdate();
+	Server->SendAndDeliverTo(Client, true);
+	Server->PostSendUpdate();
+
+	// Verify that the object got created again
+	UE_NET_ASSERT_TRUE(Client->GetReplicationBridge()->GetReplicatedObject(ServerObject->NetHandle) != nullptr);
+}
+
+UE_NET_TEST_FIXTURE(FReplicationSystemServerClientTestFixture, TestNetTemporary)
+{
+	UReplicationSystem* ReplicationSystem = Server->ReplicationSystem;
+
+	// Add a client
+	FReplicationSystemTestClient* Client = CreateClient();
+
+	// Spawn object on server
+	UTestReplicatedIrisObject* ServerObject = Server->CreateObject(0,0);
+	UTestReplicatedIrisObject* ServerObject1 = Server->CreateObject(0,0);
+
+	// Trigger replication
+	ServerObject->IntA = 1;
+	ServerObject1->IntA = 1;
+
+	// Send and deliver packet
+	Server->PreSendUpdate();
+	Server->SendAndDeliverTo(Client, true);
+	Server->PostSendUpdate();
+
+	// Verify that client has received the data
+	{
+		UTestReplicatedIrisObject* ClientObject = Cast<UTestReplicatedIrisObject>(Client->GetReplicationBridge()->GetReplicatedObject(ServerObject->NetHandle));
+
+		UE_NET_ASSERT_TRUE(ClientObject != nullptr);
+		UE_NET_ASSERT_EQ(ServerObject->IntA, ClientObject->IntA);
+	}
+
+	{
+		UTestReplicatedIrisObject* ClientObject = Cast<UTestReplicatedIrisObject>(Client->GetReplicationBridge()->GetReplicatedObject(ServerObject1->NetHandle));
+
+		UE_NET_ASSERT_TRUE(ClientObject != nullptr);
+		UE_NET_ASSERT_EQ(ServerObject1->IntA, ClientObject->IntA);
+
+	}
+
+	// Mark the object as a net temporary
+	ReplicationSystem->SetIsNetTemporary(ServerObject->NetHandle);
+
+	// Modify the value
+	ServerObject->IntA = 2;
+	ServerObject1->IntA = 2;
+
+	// Send and deliver packet
+	Server->PreSendUpdate();
+	Server->SendAndDeliverTo(Client, true);
+	Server->PostSendUpdate();
+
+	// Verify that client has not received the data for changed temporary
+	{
+		UTestReplicatedIrisObject* ClientObject = Cast<UTestReplicatedIrisObject>(Client->GetReplicationBridge()->GetReplicatedObject(ServerObject->NetHandle));
+
+		UE_NET_ASSERT_TRUE(ClientObject != nullptr);
+		UE_NET_ASSERT_NE(ServerObject->IntA, ClientObject->IntA);
+	}
+
+	// Verify that client has received the data for normal object
+	{
+		UTestReplicatedIrisObject* ClientObject = Cast<UTestReplicatedIrisObject>(Client->GetReplicationBridge()->GetReplicatedObject(ServerObject1->NetHandle));
+
+		UE_NET_ASSERT_TRUE(ClientObject != nullptr);
+		UE_NET_ASSERT_EQ(ServerObject1->IntA, ClientObject->IntA);
+	}
+
+	// Test Late join
+	// Add a client
+	FReplicationSystemTestClient* Client2 = CreateClient();
+
+	// Send and deliver packet
+	Server->PreSendUpdate();
+	Server->SendAndDeliverTo(Client2, true);
+	Server->PostSendUpdate();
+
+	// We should now have the latest state for both objects
+	{
+		UTestReplicatedIrisObject* ClientObject = Cast<UTestReplicatedIrisObject>(Client2->GetReplicationBridge()->GetReplicatedObject(ServerObject->NetHandle));
+
+		UE_NET_ASSERT_TRUE(ClientObject != nullptr);
+		UE_NET_ASSERT_EQ(ServerObject->IntA, ClientObject->IntA);
+	}
+
+	// Verify that client has received the data for normal object
+	{
+		UTestReplicatedIrisObject* ClientObject = Cast<UTestReplicatedIrisObject>(Client2->GetReplicationBridge()->GetReplicatedObject(ServerObject1->NetHandle));
+
+		UE_NET_ASSERT_TRUE(ClientObject != nullptr);
+		UE_NET_ASSERT_EQ(ServerObject1->IntA, ClientObject->IntA);
+	}
+}
+
+// Tests for TearOff
+
+// Test TearOff for existing confirmed object
+UE_NET_TEST_FIXTURE(FReplicationSystemServerClientTestFixture, TestTearOffExistingObject)
+{
+	UReplicationSystem* ReplicationSystem = Server->ReplicationSystem;
+
+	// Add a client
+	FReplicationSystemTestClient* Client = CreateClient();
+
+	// Spawn object on server
+	UTestReplicatedIrisObject* ServerObject = Server->CreateObject(0,0);
+
+	// Trigger replication
+	ServerObject->IntA = 1;
+
+	// Send and deliver packet
+	Server->PreSendUpdate();
+	Server->SendAndDeliverTo(Client, true);
+	Server->PostSendUpdate();
+
+	// Store Pointer to object 
+	UTestReplicatedIrisObject* ClientObjectThatWillBeTornOff = Cast<UTestReplicatedIrisObject>(Client->GetReplicationBridge()->GetReplicatedObject(ServerObject->NetHandle));
+
+	UE_NET_ASSERT_TRUE(ClientObjectThatWillBeTornOff != nullptr);
+	UE_NET_ASSERT_EQ(ServerObject->IntA, ClientObjectThatWillBeTornOff->IntA);
+
+	// Modify the value
+	ServerObject->IntA = 2;
+
+	// TearOff the object
+	Server->ReplicationBridge->EndReplication(ServerObject, EEndReplicationFlags::TearOff);
+
+	// Send and deliver packet
+	Server->PreSendUpdate();
+	Server->SendAndDeliverTo(Client, true);
+	Server->PostSendUpdate();
+
+	// Verify that ClientObject is torn-off and that the final state was applied
+	UE_NET_ASSERT_EQ(ServerObject->IntA, ClientObjectThatWillBeTornOff->IntA);
+	UE_NET_ASSERT_TRUE(Cast<UTestReplicatedIrisObject>(Client->GetReplicationBridge()->GetReplicatedObject(ServerObject->NetHandle)) == nullptr);
+}
+
+// Test TearOff for new object
+UE_NET_TEST_FIXTURE(FReplicationSystemServerClientTestFixture, TestTearOffOnNewlyCreatedObject)
+{
+	UReplicationSystem* ReplicationSystem = Server->ReplicationSystem;
+
+	// Add a client
+	FReplicationSystemTestClient* Client = CreateClient();
+
+	// We should not have any created objects
+	const int32 NumObjectsCreatedOnClientBeforeReplication = Client->CreatedObjects.Num();
+
+	// Spawn object on server
+	UTestReplicatedIrisObject* ServerObject = Server->CreateObject(0,0);
+
+	// Set state
+	ServerObject->IntA = 1;
+
+	// TearOff the object before first replication
+	Server->ReplicationBridge->EndReplication(ServerObject, EEndReplicationFlags::TearOff);
+
+	// Send and deliver packet
+	Server->PreSendUpdate();
+	Server->SendAndDeliverTo(Client, true);
+	Server->PostSendUpdate();
+
+	// Client should have created a object
+	UE_NET_ASSERT_EQ(NumObjectsCreatedOnClientBeforeReplication + 1, Client->CreatedObjects.Num());
+
+	UE_NET_ASSERT_TRUE(Cast<UTestReplicatedIrisObject>(Client->GetReplicationBridge()->GetReplicatedObject(ServerObject->NetHandle)) == nullptr);
+
+	// We should be able to get the object from the created objects array to validate the state
+	UTestReplicatedIrisObject* ClientObjectThatWillBeTornOff = Cast<UTestReplicatedIrisObject>(Client->CreatedObjects[NumObjectsCreatedOnClientBeforeReplication].Get());
+
+	// Verify that we replicated the expected state
+	UE_NET_ASSERT_EQ(ServerObject->IntA, ClientObjectThatWillBeTornOff->IntA);
+}
+
+// Test TearOff resend for existing confirmed object with no state changes
+UE_NET_TEST_FIXTURE(FReplicationSystemServerClientTestFixture, TestTearOffResendForExistingObjectWithoutDirtyState)
+{
+	UReplicationSystem* ReplicationSystem = Server->ReplicationSystem;
+
+	// Add a client
+	FReplicationSystemTestClient* Client = CreateClient();
+
+	// Spawn object on server
+	UTestReplicatedIrisObject* ServerObject = Server->CreateObject(0,0);
+
+	// Trigger replication
+	ServerObject->IntA = 1;
+
+	// Send and deliver packet
+	Server->PreSendUpdate();
+	Server->SendAndDeliverTo(Client, DeliverPacket);
+	Server->PostSendUpdate();
+
+	// Store Pointer to object 
+	UTestReplicatedIrisObject* ClientObjectThatWillBeTornOff = Cast<UTestReplicatedIrisObject>(Client->GetReplicationBridge()->GetReplicatedObject(ServerObject->NetHandle));
+
+	UE_NET_ASSERT_NE(ClientObjectThatWillBeTornOff, nullptr);
+	UE_NET_ASSERT_EQ(ServerObject->IntA, ClientObjectThatWillBeTornOff->IntA);
+
+	// TearOff the object
+	Server->ReplicationBridge->EndReplication(ServerObject, EEndReplicationFlags::TearOff);
+
+	// Send and do not deliver packet
+	Server->PreSendUpdate();
+	Server->SendAndDeliverTo(Client, DoNotDeliverPacket);
+	Server->PostSendUpdate();
+
+	// The ClientObject should still be found using the NetHandle
+	UE_NET_ASSERT_NE(Cast<UTestReplicatedIrisObject>(Client->GetReplicationBridge()->GetReplicatedObject(ServerObject->NetHandle)), nullptr);
+
+	Server->PreSendUpdate();
+	Server->SendAndDeliverTo(Client, DeliverPacket);
+	Server->PostSendUpdate();
+
+	// Verify that ClientObject is torn-off
+	UE_NET_ASSERT_EQ(Cast<UTestReplicatedIrisObject>(Client->GetReplicationBridge()->GetReplicatedObject(ServerObject->NetHandle)), nullptr);
+}
+
+// Test TearOff for new object and resend (should not work or is this what we want?)
+#if 0 // Until we either keep object around but out of scope, or cache creation info + deps
+UE_NET_TEST_FIXTURE(FReplicationSystemServerClientTestFixture, TestTearOffImmediateOnNewlyCreatedObjectResend)
+{
+	UReplicationSystem* ReplicationSystem = Server->ReplicationSystem;
+
+	// Add a client
+	FReplicationSystemTestClient* Client = CreateClient();
+
+	// We should not have any created objects
+	const int32 NumObjectsCreatedOnClientBeforeReplication = Client->CreatedObjects.Num();
+
+	// Spawn object on server
+	UTestReplicatedIrisObject* ServerObject = Server->CreateObject(0,0);
+
+	// Set state
+	ServerObject->IntA = 1;
+
+	// TearOff the object before first replication
+	Server->ReplicationBridge->EndReplication(ServerObject, EEndReplicationFlags::TearOff);
+
+	// Send and deliver packet
+	Server->PreSendUpdate();
+	Server->SendAndDeliverTo(Client, false);
+	Server->PostSendUpdate();
+
+	// Send and deliver packet
+	Server->PreSendUpdate();
+	Server->SendAndDeliverTo(Client, true);
+	Server->PostSendUpdate();
+
+	// Client should have created a object
+	UE_NET_ASSERT_EQ(NumObjectsCreatedOnClientBeforeReplication + 1, Client->CreatedObjects.Num());
+
+	UE_NET_ASSERT_TRUE(Cast<UTestReplicatedIrisObject>(Client->GetReplicationBridge()->GetReplicatedObject(ServerObject->NetHandle)) == nullptr);
+
+	// We should be able to get the object from the created objects array to validate the state
+	UTestReplicatedIrisObject* ClientObjectThatWillBeTornOff = Cast<UTestReplicatedIrisObject>(Client->CreatedObjects[NumObjectsCreatedOnClientBeforeReplication].Get());
+
+	// Verify that we replicated the expected state
+	UE_NET_ASSERT_EQ(ServerObject->IntA, ClientObjectThatWillBeTornOff->IntA);
+}
+#endif
+
+
+UE_NET_TEST_FIXTURE(FReplicationSystemServerClientTestFixture, TestDefferedTearOffOnNewlyCreatedObjectResend)
+{
+	UReplicationSystem* ReplicationSystem = Server->ReplicationSystem;
+
+	// Add a client
+	FReplicationSystemTestClient* Client = CreateClient();
+
+	// We should not have any created objects
+	const int32 NumObjectsCreatedOnClientBeforeReplication = Client->CreatedObjects.Num();
+
+	// Spawn object on server
+	UTestReplicatedIrisObject* ServerObject = Server->CreateObject(0,0);
+
+	// Set state
+	ServerObject->IntA = 1;
+
+	// TearOff the object before first replication
+	Server->ReplicationSystem->TearOffNextUpdate(ServerObject->NetHandle);
+
+	// Send and deliver packet
+	Server->PreSendUpdate();
+	Server->SendAndDeliverTo(Client, false);
+	Server->PostSendUpdate();
+
+	// Send and deliver packet
+	Server->PreSendUpdate();
+	Server->SendAndDeliverTo(Client, true);
+	Server->PostSendUpdate();
+	
+	// End replication and destroy object
+	Server->ReplicationBridge->EndReplication(ServerObject, EEndReplicationFlags::Destroy);
+
+	// Client should have created a object
+	UE_NET_ASSERT_EQ(NumObjectsCreatedOnClientBeforeReplication + 1, Client->CreatedObjects.Num());
+
+	UE_NET_ASSERT_TRUE(Cast<UTestReplicatedIrisObject>(Client->GetReplicationBridge()->GetReplicatedObject(ServerObject->NetHandle)) == nullptr);
+
+	// We should be able to get the object from the created objects array to validate the state
+	UTestReplicatedIrisObject* ClientObjectThatWillBeTornOff = Cast<UTestReplicatedIrisObject>(Client->CreatedObjects[NumObjectsCreatedOnClientBeforeReplication].Get());
+
+	// Verify that we replicated the expected state
+	UE_NET_ASSERT_EQ(ServerObject->IntA, ClientObjectThatWillBeTornOff->IntA);
+}
+
+// Test TearOff for existing not yet confirmed object
+UE_NET_TEST_FIXTURE(FReplicationSystemServerClientTestFixture, TestTearOffObjectPendingCreateConfirmation)
+{
+	UReplicationSystem* ReplicationSystem = Server->ReplicationSystem;
+
+	// Add a client
+	FReplicationSystemTestClient* Client = CreateClient();
+
+	// Spawn object on server
+	UTestReplicatedIrisObject* ServerObject = Server->CreateObject(0,0);
+
+	// Trigger replication
+	ServerObject->IntA = 1;
+
+	// Send packet to get put the object in flight
+	Server->PreSendUpdate();
+	Server->SendTo(Client);
+	Server->PostSendUpdate();
+
+	// TearOff the object
+	Server->ReplicationBridge->EndReplication(ServerObject, EEndReplicationFlags::TearOff);
+
+	// Deliver Object (should now be created)
+	Server->DeliverTo(Client, true);
+
+	// Store Pointer to object and verify initial state
+	UTestReplicatedIrisObject* ClientObjectThatWillBeTornOff = Cast<UTestReplicatedIrisObject>(Client->GetReplicationBridge()->GetReplicatedObject(ServerObject->NetHandle));
+
+	UE_NET_ASSERT_TRUE(ClientObjectThatWillBeTornOff != nullptr);
+	UE_NET_ASSERT_EQ(ServerObject->IntA, ClientObjectThatWillBeTornOff->IntA);
+
+	// Send and deliver packet
+	Server->PreSendUpdate();
+	Server->SendAndDeliverTo(Client, true);
+	Server->PostSendUpdate();
+
+	// Verify that ClientObject is TornOFf and that the final state was applied
+	UE_NET_ASSERT_EQ(ServerObject->IntA, ClientObjectThatWillBeTornOff->IntA);
+	UE_NET_ASSERT_TRUE(Cast<UTestReplicatedIrisObject>(Client->GetReplicationBridge()->GetReplicatedObject(ServerObject->NetHandle)) == nullptr);
+}
+
+// Test TearOff for existing object pending destroy (should do nothing)
+UE_NET_TEST_FIXTURE(FReplicationSystemServerClientTestFixture, TestTearOffExistingObjectPendingDestroy)
+{
+	UReplicationSystem* ReplicationSystem = Server->ReplicationSystem;
+
+	// Add a client
+	FReplicationSystemTestClient* Client = CreateClient();
+
+	// Spawn object on server
+	UTestReplicatedIrisObject* ServerObject = Server->CreateObject(0,0);
+
+	// Trigger replication
+	ServerObject->IntA = 1;
+
+	// Send and deliver packet
+	Server->PreSendUpdate();
+	Server->SendAndDeliverTo(Client, true);
+	Server->PostSendUpdate();
+
+	// Store Pointer to object 
+	UTestReplicatedIrisObject* ClientObjectThatWillBeTornOff = Cast<UTestReplicatedIrisObject>(Client->GetReplicationBridge()->GetReplicatedObject(ServerObject->NetHandle));
+
+	UE_NET_ASSERT_TRUE(ClientObjectThatWillBeTornOff != nullptr);
+	UE_NET_ASSERT_EQ(ServerObject->IntA, ClientObjectThatWillBeTornOff->IntA);
+
+	// Modify the value
+	ServerObject->IntA = 2;
+
+	// Mark the object for destroy
+	Server->ReplicationBridge->EndReplication(ServerObject);
+
+	// TearOff the object
+	Server->ReplicationBridge->EndReplication(ServerObject, EEndReplicationFlags::TearOff);
+
+	// Send and deliver packet
+	Server->PreSendUpdate();
+	Server->SendAndDeliverTo(Client, true);
+	Server->PostSendUpdate();
+
+	// Verify that ClientObject is not tornOff and that the final state was not applied as we issued tearoff after ending replication
+	UE_NET_ASSERT_NE(ServerObject->IntA, ClientObjectThatWillBeTornOff->IntA);
+	UE_NET_ASSERT_TRUE(Cast<UTestReplicatedIrisObject>(Client->GetReplicationBridge()->GetReplicatedObject(ServerObject->NetHandle)) == nullptr);
+}
+
+// Test TearOff resend 
+UE_NET_TEST_FIXTURE(FReplicationSystemServerClientTestFixture, TestTearOffResend)
+{
+	UReplicationSystem* ReplicationSystem = Server->ReplicationSystem;
+
+	// Add a client
+	FReplicationSystemTestClient* Client = CreateClient();
+
+	// Spawn object on server
+	UTestReplicatedIrisObject* ServerObject = Server->CreateObject(0,0);
+
+	// Trigger replication
+	ServerObject->IntA = 1;
+
+	// Send and deliver packet
+	Server->PreSendUpdate();
+	Server->SendAndDeliverTo(Client, true);
+	Server->PostSendUpdate();
+
+	// Store Pointer to object 
+	UTestReplicatedIrisObject* ClientObjectThatWillBeTornOff = Cast<UTestReplicatedIrisObject>(Client->GetReplicationBridge()->GetReplicatedObject(ServerObject->NetHandle));
+
+	UE_NET_ASSERT_TRUE(ClientObjectThatWillBeTornOff != nullptr);
+	UE_NET_ASSERT_EQ(ServerObject->IntA, ClientObjectThatWillBeTornOff->IntA);
+
+	// Modify the value
+	ServerObject->IntA = 2;
+
+	// TearOff the object
+	Server->ReplicationSystem->TearOffNextUpdate(ServerObject->NetHandle);
+
+	// Send and deliver packet, in this case the packet containing 2 was lost, but, we did not know that when we 
+	Server->PreSendUpdate();
+	Server->SendAndDeliverTo(Client, false);
+	Server->PostSendUpdate();
+
+	// Object should now be torn-off, so it should not copy the latest state
+	ServerObject->IntA = 3;
+
+	// Send and deliver packet
+	Server->PreSendUpdate();
+	Server->SendAndDeliverTo(Client, true);
+	Server->PostSendUpdate();
+
+	// Verify that ClientObject is torn-off and that the expected final state was applied
+	UE_NET_ASSERT_EQ(2, ClientObjectThatWillBeTornOff->IntA);
+	UE_NET_ASSERT_TRUE(Cast<UTestReplicatedIrisObject>(Client->GetReplicationBridge()->GetReplicatedObject(ServerObject->NetHandle)) == nullptr);
+}
+
+// Test TearOff does not pickup statechanges after tear off
+UE_NET_TEST_FIXTURE(FReplicationSystemServerClientTestFixture, TestTornedOffObjectDoesNotCopyStateChanges)
+{
+	UReplicationSystem* ReplicationSystem = Server->ReplicationSystem;
+
+	// Add a client
+	FReplicationSystemTestClient* Client = CreateClient();
+
+	// Spawn object on server
+	UTestReplicatedIrisObject* ServerObject = Server->CreateObject(0,0);
+
+	// Trigger replication
+	ServerObject->IntA = 1;
+
+	// Send and deliver packet
+	Server->PreSendUpdate();
+	Server->SendAndDeliverTo(Client, true);
+	Server->PostSendUpdate();
+
+	// Store Pointer to object 
+	UTestReplicatedIrisObject* ClientObjectThatWillBeTornOff = Cast<UTestReplicatedIrisObject>(Client->GetReplicationBridge()->GetReplicatedObject(ServerObject->NetHandle));
+
+	UE_NET_ASSERT_TRUE(ClientObjectThatWillBeTornOff != nullptr);
+	UE_NET_ASSERT_EQ(ServerObject->IntA, ClientObjectThatWillBeTornOff->IntA);
+
+	// Modify the value
+	ServerObject->IntA = 2;
+
+	// TearOff the object
+	Server->ReplicationSystem->TearOffNextUpdate(ServerObject->NetHandle);
+
+	// Send and drop packet containing the value 2
+	Server->PreSendUpdate();
+	Server->SendAndDeliverTo(Client, false);
+	Server->PostSendUpdate();
+
+	// Object should now be torn-off, so it should not copy the latest state but instead resend the last copied state (2) along with the tear-off
+	ServerObject->IntA = 3;
+
+	// Send and deliver packet
+	Server->PreSendUpdate();
+	Server->SendAndDeliverTo(Client, true);
+	Server->PostSendUpdate();
+
+	// Verify that ClientObject is TornOFf and that the expected final state was applied
+	UE_NET_ASSERT_EQ(2, ClientObjectThatWillBeTornOff->IntA);
+	UE_NET_ASSERT_TRUE(Cast<UTestReplicatedIrisObject>(Client->GetReplicationBridge()->GetReplicatedObject(ServerObject->NetHandle)) == nullptr);
+}
+
+// Test TearOff and SubObjects, SubObjects must apply state?
+UE_NET_TEST_FIXTURE(FReplicationSystemServerClientTestFixture, TestImmediateTearOffExistingObjectWithSubObject)
+{
+	UReplicationSystem* ReplicationSystem = Server->ReplicationSystem;
+
+	// Add a client
+	FReplicationSystemTestClient* Client = CreateClient();
+
+	// Spawn object on server
+	UTestReplicatedIrisObject* ServerObject = Server->CreateObject(0,0);
+
+	// Spawn second object on server as a subobject
+	UTestReplicatedIrisObject* ServerSubObject = Server->CreateSubObject(ServerObject->NetHandle, 0, 0);
+
+	// Trigger replication
+	ServerObject->IntA = 1;
+	ServerSubObject->IntA = 1;
+
+	// Send and deliver packet
+	Server->PreSendUpdate();
+	Server->SendAndDeliverTo(Client, true);
+	Server->PostSendUpdate();
+
+	// Store Pointer to objects
+	UTestReplicatedIrisObject* ClientObjectThatWillBeTornOff = Cast<UTestReplicatedIrisObject>(Client->GetReplicationBridge()->GetReplicatedObject(ServerObject->NetHandle));
+	UE_NET_ASSERT_TRUE(ClientObjectThatWillBeTornOff != nullptr);
+	UE_NET_ASSERT_EQ(ServerObject->IntA, ClientObjectThatWillBeTornOff->IntA);
+
+	UTestReplicatedIrisObject* ClientSubObjectThatWillBeTornOff = Cast<UTestReplicatedIrisObject>(Client->GetReplicationBridge()->GetReplicatedObject(ServerSubObject->NetHandle));
+	UE_NET_ASSERT_TRUE(ClientSubObjectThatWillBeTornOff != nullptr);
+	UE_NET_ASSERT_EQ(ServerSubObject->IntA, ClientSubObjectThatWillBeTornOff->IntA);
+
+	// Modify the value of subobject only
+	ServerSubObject->IntA = 2;
+
+	// TearOff the object using immediate tear-off
+	Server->ReplicationBridge->EndReplication(ServerObject, EEndReplicationFlags::TearOff);
+
+	// Send and deliver packet
+	Server->PreSendUpdate();
+	Server->SendAndDeliverTo(Client, true);
+	Server->PostSendUpdate();
+
+	// Verify that ClientObject is torn-off and that the final state was applied to subObject 
+	UE_NET_ASSERT_EQ(ServerSubObject->IntA, ClientSubObjectThatWillBeTornOff->IntA);
+	UE_NET_ASSERT_TRUE(Cast<UTestReplicatedIrisObject>(Client->GetReplicationBridge()->GetReplicatedObject(ServerSubObject->NetHandle)) == nullptr);
+}
+
+// Test TearOff and SubObjects, SubObjects must apply state?
+UE_NET_TEST_FIXTURE(FReplicationSystemServerClientTestFixture, TestImmediateTearOffExistingObjectWithSubObjectDroppedData)
+{
+	UReplicationSystem* ReplicationSystem = Server->ReplicationSystem;
+
+	// Add a client
+	FReplicationSystemTestClient* Client = CreateClient();
+
+	// Spawn object on server
+	UTestReplicatedIrisObject* ServerObject = Server->CreateObject(0,0);
+
+	// Spawn second object on server as a subobject
+	UTestReplicatedIrisObject* ServerSubObject = Server->CreateSubObject(ServerObject->NetHandle, 0, 0);
+
+	// Trigger replication
+	ServerObject->IntA = 1;
+	ServerSubObject->IntA = 1;
+
+	// Send and deliver packet
+	Server->PreSendUpdate();
+	Server->SendAndDeliverTo(Client, true);
+	Server->PostSendUpdate();
+
+	// Store Pointer to objects
+	UTestReplicatedIrisObject* ClientObjectThatWillBeTornOff = Cast<UTestReplicatedIrisObject>(Client->GetReplicationBridge()->GetReplicatedObject(ServerObject->NetHandle));
+	UE_NET_ASSERT_TRUE(ClientObjectThatWillBeTornOff != nullptr);
+	UE_NET_ASSERT_EQ(ServerObject->IntA, ClientObjectThatWillBeTornOff->IntA);
+
+	UTestReplicatedIrisObject* ClientSubObjectThatWillBeTornOff = Cast<UTestReplicatedIrisObject>(Client->GetReplicationBridge()->GetReplicatedObject(ServerSubObject->NetHandle));
+	UE_NET_ASSERT_TRUE(ClientSubObjectThatWillBeTornOff != nullptr);
+	UE_NET_ASSERT_EQ(ServerSubObject->IntA, ClientSubObjectThatWillBeTornOff->IntA);
+
+	// Modify the value of subobject only
+	ServerSubObject->IntA = 2;
+
+	// TearOff the object using immediate tear-off
+	Server->ReplicationBridge->EndReplication(ServerObject, EEndReplicationFlags::TearOff);
+
+	// Send and do not deliver packet
+	Server->PreSendUpdate();
+	Server->SendAndDeliverTo(Client, false);
+	Server->PostSendUpdate();
+
+	// Send and deliver packet
+	Server->PreSendUpdate();
+	Server->SendAndDeliverTo(Client, true);
+	Server->PostSendUpdate();
+
+	// Verify that ClientObject is torn-off and that the final state was applied to subObject 
+	UE_NET_ASSERT_EQ(ServerSubObject->IntA, ClientSubObjectThatWillBeTornOff->IntA);
+	UE_NET_ASSERT_TRUE(Cast<UTestReplicatedIrisObject>(Client->GetReplicationBridge()->GetReplicatedObject(ServerSubObject->NetHandle)) == nullptr);
+}
+
+// Test TearOff and SubObjects, SubObjects must apply state?
+UE_NET_TEST_FIXTURE(FReplicationSystemServerClientTestFixture, TestTearOffNextUpdateExistingObjectWithSubObject)
+{
+	UReplicationSystem* ReplicationSystem = Server->ReplicationSystem;
+
+	// Add a client
+	FReplicationSystemTestClient* Client = CreateClient();
+
+	// Spawn object on server
+	UTestReplicatedIrisObject* ServerObject = Server->CreateObject(0,0);
+
+	// Spawn second object on server as a subobject
+	UTestReplicatedIrisObject* ServerSubObject = Server->CreateSubObject(ServerObject->NetHandle, 0, 0);
+
+	// Trigger replication
+	ServerObject->IntA = 1;
+	ServerSubObject->IntA = 1;
+
+	// Send and deliver packet
+	Server->PreSendUpdate();
+	Server->SendAndDeliverTo(Client, true);
+	Server->PostSendUpdate();
+
+	// Store Pointer to objects
+	UTestReplicatedIrisObject* ClientObjectThatWillBeTornOff = Cast<UTestReplicatedIrisObject>(Client->GetReplicationBridge()->GetReplicatedObject(ServerObject->NetHandle));
+	UE_NET_ASSERT_TRUE(ClientObjectThatWillBeTornOff != nullptr);
+	UE_NET_ASSERT_EQ(ServerObject->IntA, ClientObjectThatWillBeTornOff->IntA);
+
+	UTestReplicatedIrisObject* ClientSubObjectThatWillBeTornOff = Cast<UTestReplicatedIrisObject>(Client->GetReplicationBridge()->GetReplicatedObject(ServerSubObject->NetHandle));
+	UE_NET_ASSERT_TRUE(ClientSubObjectThatWillBeTornOff != nullptr);
+	UE_NET_ASSERT_EQ(ServerSubObject->IntA, ClientSubObjectThatWillBeTornOff->IntA);
+
+	// Modify the value of subobject only
+	ServerSubObject->IntA = 2;
+
+	// TearOff the object using immediate tear-off
+	Server->ReplicationSystem->TearOffNextUpdate(ServerObject->NetHandle);
+
+	// Send and deliver packet
+	Server->PreSendUpdate();
+	Server->SendAndDeliverTo(Client, true);
+	Server->PostSendUpdate();
+
+	// Verify that ClientObject is torn-off and that the final state was applied to subObject 
+	UE_NET_ASSERT_EQ(ServerSubObject->IntA, ClientSubObjectThatWillBeTornOff->IntA);
+	UE_NET_ASSERT_TRUE(Cast<UTestReplicatedIrisObject>(Client->GetReplicationBridge()->GetReplicatedObject(ServerSubObject->NetHandle)) == nullptr);
+}
+
+// Test TearOff and destroy of SubObjects that are still pending create/tearoff
+UE_NET_TEST_FIXTURE(FReplicationSystemServerClientTestFixture, TestTearOffNextUpdateExistingObjectWithSubObjectPendingCreation)
+{
+	UReplicationSystem* ReplicationSystem = Server->ReplicationSystem;
+	FNetHandleManager* NetHandleManager = &ReplicationSystem->GetReplicationSystemInternal()->GetNetHandleManager();
+
+	// Add a client
+	FReplicationSystemTestClient* Client = CreateClient();
+
+	// Spawn object on server
+	UTestReplicatedIrisObject* ServerObject = Server->CreateObject(0,0);
+
+	// Spawn second object on server as a subobject
+	UTestReplicatedIrisObject* ServerSubObject = Server->CreateSubObject(ServerObject->NetHandle, 0, 0);
+
+	// Trigger replication
+	ServerObject->IntA = 1;
+	ServerSubObject->IntA = 1;
+
+	const FInternalNetHandle ServerObjectInternalIndex = NetHandleManager->GetInternalIndex(ServerObject->NetHandle);
+	const FInternalNetHandle SubObjectObjectInternalIndex = NetHandleManager->GetInternalIndex(ServerSubObject->NetHandle);
+
+	// Trigger presend without send to add the objects to scope
+	Server->PreSendUpdate();
+
+	UE_NET_ASSERT_EQ(uint16(1), NetHandleManager->GetNetObjectRefCount(ServerObjectInternalIndex));
+	UE_NET_ASSERT_EQ(uint16(1), NetHandleManager->GetNetObjectRefCount(SubObjectObjectInternalIndex));
+
+	// TearOff the object this will also tear-off subobject
+	Server->ReplicationSystem->TearOffNextUpdate(ServerObject->NetHandle);
+
+	// Update logic, object should be removed from scope but still exist as pending create in
+	Server->PreSendUpdate();
+
+	UE_NET_ASSERT_EQ(uint16(1), NetHandleManager->GetNetObjectRefCount(ServerObjectInternalIndex));
+	UE_NET_ASSERT_EQ(uint16(1), NetHandleManager->GetNetObjectRefCount(SubObjectObjectInternalIndex));
+
+	// Destroy the object
+	Server->DestroyObject(ServerObject);
+
+	// Verify that we no longer have any references to the object
+	UE_NET_ASSERT_EQ(uint16(0), NetHandleManager->GetNetObjectRefCount(ServerObjectInternalIndex));
+	UE_NET_ASSERT_EQ(uint16(0), NetHandleManager->GetNetObjectRefCount(SubObjectObjectInternalIndex));
+	
+	Server->SendAndDeliverTo(Client, true);
+	Server->PostSendUpdate();
+}
+
+// Test that we can replicate an object with no replicated properties
+UE_NET_TEST_FIXTURE(FReplicationSystemServerClientTestFixture, TestReplicatedObjectWithNoReplicatedProperties)
+{
+	UReplicationSystem* ReplicationSystem = Server->ReplicationSystem;
+
+	// Add a client
+	FReplicationSystemTestClient* Client = CreateClient();
+
+	// Spawn object on server
+	UTestReplicatedIrisObjectWithNoReplicatedMembers* ServerObject = Server->CreateObject<UTestReplicatedIrisObjectWithNoReplicatedMembers>();
+	const FNetHandle ServerHandle = ServerObject->NetHandle;
+
+	UE_NET_ASSERT_TRUE(ServerHandle.IsValid());
+
+	// Send and deliver packet
+	Server->PreSendUpdate();
+	Server->SendAndDeliverTo(Client, true);
+	Server->PostSendUpdate();
+
+	UTestReplicatedIrisObjectWithNoReplicatedMembers* ClientObject = Cast<UTestReplicatedIrisObjectWithNoReplicatedMembers>(Client->GetReplicationBridge()->GetReplicatedObject(ServerHandle));
+	UE_NET_ASSERT_TRUE(ClientObject != nullptr);
+
+	// Destroy object
+	Server->DestroyObject(ServerObject);
+
+	// Send and deliver packet
+	Server->PreSendUpdate();
+	Server->SendAndDeliverTo(Client, true);
+	Server->PostSendUpdate();
+
+	ClientObject = Cast<UTestReplicatedIrisObjectWithNoReplicatedMembers>(Client->GetReplicationBridge()->GetReplicatedObject(ServerHandle));
+	UE_NET_ASSERT_TRUE(ClientObject == nullptr);
+}
+
+UE_NET_TEST_FIXTURE(FReplicationSystemServerClientTestFixture, TestObjectPollFramePeriod)
+{
+	UReplicationSystem* ReplicationSystem = Server->ReplicationSystem;
+
+	// Add a client
+	FReplicationSystemTestClient* Client = CreateClient();
+
+	// Spawn object on server
+	UTestReplicatedIrisObject* ServerObject = Server->CreateObject(0, 0);
+
+	// Spawn second object on server that later will be addedded as a dependent object
+	UTestReplicatedIrisObject* ServerObjectPolledEveryOtherFrame = Server->CreateObject(0, 0);
+	Server->ReplicationBridge->SetPollFramePeriod(ServerObjectPolledEveryOtherFrame, 1U);
+
+	// Send and deliver packet
+	Server->PreSendUpdate();
+	Server->SendAndDeliverTo(Client, true);
+	Server->PostSendUpdate();
+
+	// Store Pointer to objects and verify state after initial replication
+	UTestReplicatedIrisObject* ClientObject = Cast<UTestReplicatedIrisObject>(Client->GetReplicationBridge()->GetReplicatedObject(ServerObject->NetHandle));
+	UTestReplicatedIrisObject* ClientObjectPolledEveryOtherFrame = Cast<UTestReplicatedIrisObject>(Client->GetReplicationBridge()->GetReplicatedObject(ServerObjectPolledEveryOtherFrame->NetHandle));
+
+	UE_NET_ASSERT_NE(ClientObjectPolledEveryOtherFrame, nullptr);
+	UE_NET_ASSERT_NE(ClientObject, nullptr);
+	UE_NET_ASSERT_EQ(ClientObject->IntA, ServerObject->IntA);
+	UE_NET_ASSERT_EQ(ClientObjectPolledEveryOtherFrame->IntA, ServerObjectPolledEveryOtherFrame->IntA);
+
+	// Trigger replication
+	ServerObject->IntA = 1;
+	ServerObjectPolledEveryOtherFrame->IntA = 1;
+
+	// Send and deliver packet
+	Server->PreSendUpdate();
+	Server->SendAndDeliverTo(Client, true);
+	Server->PostSendUpdate();
+
+	// Verify that only the server object has been updated
+	UE_NET_ASSERT_EQ(ClientObject->IntA, ServerObject->IntA);
+	UE_NET_ASSERT_NE(ClientObjectPolledEveryOtherFrame->IntA, ServerObjectPolledEveryOtherFrame->IntA);
+
+	// Send and deliver packet
+	Server->PreSendUpdate();
+	Server->SendAndDeliverTo(Client, true);
+	Server->PostSendUpdate();
+
+	// Verify that both objects now are in sync
+	UE_NET_ASSERT_EQ(ClientObject->IntA, ServerObject->IntA);
+	UE_NET_ASSERT_EQ(ClientObjectPolledEveryOtherFrame->IntA, ServerObjectPolledEveryOtherFrame->IntA);
+}
+
+}
diff --git a/Engine/Plugins/Runtime/ReplicationSystemTestPlugin/Source/Private/Tests/Serialization/TestObjectNetSerializer.cpp b/Engine/Plugins/Runtime/ReplicationSystemTestPlugin/Source/Private/Tests/Serialization/TestObjectNetSerializer.cpp
new file mode 100644
--- /dev/null
+++ b/Engine/Plugins/Runtime/ReplicationSystemTestPlugin/Source/Private/Tests/Serialization/TestObjectNetSerializer.cpp
@@ -0,0 +1,162 @@
+// Copyright Epic Games, Inc. All Rights Reserved.
+
+#include "TestNetSerializerFixture.h"
+#include "Iris/Serialization/ObjectNetSerializer.h"
+#include "Iris/Core/NetObjectReference.h"
+#include "Iris/Serialization/InternalNetSerializationContext.h"
+#include "Tests/ReplicationSystem/ReplicationSystemServerClientTestFixture.h"
+#include "Tests/ReplicationSystem/ReplicatedTestObject.h"
+
+namespace UE::Net::Private
+{
+
+static FTestMessage& PrintObjectNetSerializerConfig(FTestMessage& Message, const FNetSerializerConfig& InConfig)
+{
+	return Message;
+}
+
+UE_NET_TEST_FIXTURE(FReplicationSystemServerClientTestFixture, TestObjectReference)
+{
+	// Add a client
+	FReplicationSystemTestClient* Client = CreateClient();
+
+	// Spawn objects on server
+	UTestReplicatedIrisObjectWithObjectReference* ObjectD = Server->CreateObject<UTestReplicatedIrisObjectWithObjectReference>();
+	UTestReplicatedIrisObjectWithObjectReference* ObjectC = Server->CreateObject<UTestReplicatedIrisObjectWithObjectReference>();
+	UTestReplicatedIrisObjectWithObjectReference* ObjectB = Server->CreateObject<UTestReplicatedIrisObjectWithObjectReference>();
+	UTestReplicatedIrisObjectWithObjectReference* ObjectA = Server->CreateObject<UTestReplicatedIrisObjectWithObjectReference>();
+
+	ObjectA->RawObjectPtrRef = ObjectB;
+	ObjectB->RawObjectPtrRef = ObjectC;
+
+	// This will not work since we do not have any dependencies
+	ObjectD->RawObjectPtrRef = ObjectA;
+
+	// Send and deliver packet
+	Server->PreSendUpdate();
+	Server->SendAndDeliverTo(Client, true);
+	Server->PostSendUpdate();
+
+	// Check state, we cheat a bit and use the server nethandle for lookup
+	// THIS will fail if we start using full handle compares
+	auto ClientObjectA = Client->GetObjectAs<UTestReplicatedIrisObjectWithObjectReference>(ObjectA->NetHandle);
+	auto ClientObjectB = Client->GetObjectAs<UTestReplicatedIrisObjectWithObjectReference>(ObjectB->NetHandle);
+	auto ClientObjectC = Client->GetObjectAs<UTestReplicatedIrisObjectWithObjectReference>(ObjectC->NetHandle);
+	auto ClientObjectD = Client->GetObjectAs<UTestReplicatedIrisObjectWithObjectReference>(ObjectD->NetHandle);
+	
+	UE_NET_ASSERT_TRUE(ClientObjectA != nullptr);
+	UE_NET_ASSERT_TRUE(ClientObjectB != nullptr);
+	UE_NET_ASSERT_TRUE(ClientObjectC != nullptr);
+	UE_NET_ASSERT_TRUE(ClientObjectD != nullptr);
+
+	// Verify that we managed to resolved the references
+	UE_NET_ASSERT_EQ((UObject*)ClientObjectA->RawObjectPtrRef, (UObject*)ClientObjectB);
+	UE_NET_ASSERT_EQ((UObject*)ClientObjectB->RawObjectPtrRef, (UObject*)ClientObjectC);
+	UE_NET_ASSERT_EQ((UObject*)ClientObjectC->RawObjectPtrRef, (UObject*)nullptr);
+	UE_NET_ASSERT_EQ((UObject*)ClientObjectD->RawObjectPtrRef, (UObject*)ClientObjectA);
+}
+
+UE_NET_TEST_FIXTURE(FReplicationSystemServerClientTestFixture, TestCircularReference)
+{
+	// Add a client
+	FReplicationSystemTestClient* Client = CreateClient();
+
+	// Spawn objects on server
+	UTestReplicatedIrisObjectWithObjectReference* ObjectD = Server->CreateObject<UTestReplicatedIrisObjectWithObjectReference>();
+	UTestReplicatedIrisObjectWithObjectReference* ObjectC = Server->CreateObject<UTestReplicatedIrisObjectWithObjectReference>();
+	UTestReplicatedIrisObjectWithObjectReference* ObjectB = Server->CreateObject<UTestReplicatedIrisObjectWithObjectReference>();
+	UTestReplicatedIrisObjectWithObjectReference* ObjectA = Server->CreateObject<UTestReplicatedIrisObjectWithObjectReference>();
+
+	ObjectA->RawObjectPtrRef = ObjectB;
+	ObjectB->RawObjectPtrRef = ObjectC;
+	ObjectC->RawObjectPtrRef = ObjectD;
+	ObjectD->RawObjectPtrRef = ObjectA;
+
+	// Send and deliver packet
+	Server->PreSendUpdate();
+	Server->SendAndDeliverTo(Client, true);
+	Server->PostSendUpdate();
+
+	// Check state, we cheat a bit and use the server nethandle for lookup
+	// THIS will fail if we start using full handle compares
+	auto ClientObjectA = Client->GetObjectAs<UTestReplicatedIrisObjectWithObjectReference>(ObjectA->NetHandle);
+	auto ClientObjectB = Client->GetObjectAs<UTestReplicatedIrisObjectWithObjectReference>(ObjectB->NetHandle);
+	auto ClientObjectC = Client->GetObjectAs<UTestReplicatedIrisObjectWithObjectReference>(ObjectC->NetHandle);
+	auto ClientObjectD = Client->GetObjectAs<UTestReplicatedIrisObjectWithObjectReference>(ObjectD->NetHandle);
+	
+	UE_NET_ASSERT_TRUE(ClientObjectA != nullptr);
+	UE_NET_ASSERT_TRUE(ClientObjectB != nullptr);
+	UE_NET_ASSERT_TRUE(ClientObjectC != nullptr);
+	UE_NET_ASSERT_TRUE(ClientObjectD != nullptr);
+
+	// Verify that we managed to resolved the references
+	UE_NET_ASSERT_EQ((UObject*)ClientObjectA->RawObjectPtrRef, (UObject*)ClientObjectB);
+	UE_NET_ASSERT_EQ((UObject*)ClientObjectB->RawObjectPtrRef, (UObject*)ClientObjectC);
+	UE_NET_ASSERT_EQ((UObject*)ClientObjectC->RawObjectPtrRef, (UObject*)ClientObjectD);
+	UE_NET_ASSERT_EQ((UObject*)ClientObjectD->RawObjectPtrRef, (UObject*)ClientObjectA);
+}
+
+UE_NET_TEST_FIXTURE(FReplicationSystemServerClientTestFixture, TestCircularReferenceWithLimitedThroughput)
+{
+	// Limit packet size
+	Server->SetMaxSendPacketSize(128U);
+
+	// Add a client
+	FReplicationSystemTestClient* Client = CreateClient();
+
+	// Spawn objects on server
+	UTestReplicatedIrisObjectWithObjectReference* ObjectD = Server->CreateObject<UTestReplicatedIrisObjectWithObjectReference>();
+	UTestReplicatedIrisObjectWithObjectReference* ObjectC = Server->CreateObject<UTestReplicatedIrisObjectWithObjectReference>();
+	UTestReplicatedIrisObjectWithObjectReference* ObjectB = Server->CreateObject<UTestReplicatedIrisObjectWithObjectReference>();
+	UTestReplicatedIrisObjectWithObjectReference* ObjectA = Server->CreateObject<UTestReplicatedIrisObjectWithObjectReference>();
+
+	ObjectA->RawObjectPtrRef = ObjectB;
+	ObjectB->RawObjectPtrRef = ObjectC;
+	ObjectC->RawObjectPtrRef = ObjectD;
+	ObjectD->RawObjectPtrRef = ObjectA;
+
+	// Send and deliver packets
+	Server->PreSendUpdate();
+	Server->SendAndDeliverTo(Client, true);
+	Server->PostSendUpdate();
+
+	// Send and deliver packets
+	Server->PreSendUpdate();
+	Server->SendAndDeliverTo(Client, true);
+	Server->PostSendUpdate();
+
+	// Send and deliver packets
+	Server->PreSendUpdate();
+	Server->SendAndDeliverTo(Client, true);
+	Server->PostSendUpdate();
+
+	// Send and deliver packets
+	Server->PreSendUpdate();
+	Server->SendAndDeliverTo(Client, true);
+	Server->PostSendUpdate();
+
+	// Send and deliver packets
+	Server->PreSendUpdate();
+	Server->SendAndDeliverTo(Client, true);
+	Server->PostSendUpdate();
+
+	// Check state, we cheat a bit and use the server nethandle for lookup
+	// THIS will fail if we start using full handle compares
+	auto ClientObjectA = Client->GetObjectAs<UTestReplicatedIrisObjectWithObjectReference>(ObjectA->NetHandle);
+	auto ClientObjectB = Client->GetObjectAs<UTestReplicatedIrisObjectWithObjectReference>(ObjectB->NetHandle);
+	auto ClientObjectC = Client->GetObjectAs<UTestReplicatedIrisObjectWithObjectReference>(ObjectC->NetHandle);
+	auto ClientObjectD = Client->GetObjectAs<UTestReplicatedIrisObjectWithObjectReference>(ObjectD->NetHandle);
+	
+	UE_NET_ASSERT_TRUE(ClientObjectA != nullptr);
+	UE_NET_ASSERT_TRUE(ClientObjectB != nullptr);
+	UE_NET_ASSERT_TRUE(ClientObjectC != nullptr);
+	UE_NET_ASSERT_TRUE(ClientObjectD != nullptr);
+
+	// Verify that we managed to resolved the references
+	UE_NET_ASSERT_EQ((UObject*)ClientObjectA->RawObjectPtrRef, (UObject*)ClientObjectB);
+	UE_NET_ASSERT_EQ((UObject*)ClientObjectB->RawObjectPtrRef, (UObject*)ClientObjectC);
+	UE_NET_ASSERT_EQ((UObject*)ClientObjectC->RawObjectPtrRef, (UObject*)ClientObjectD);
+	UE_NET_ASSERT_EQ((UObject*)ClientObjectD->RawObjectPtrRef, (UObject*)ClientObjectA);
+}
+
+}
diff --git a/Engine/Plugins/Runtime/ReplicationSystemTestPlugin/Source/Private/Tests/Serialization/TestPolymorphicStructNetSerializer.h b/Engine/Plugins/Runtime/ReplicationSystemTestPlugin/Source/Private/Tests/Serialization/TestPolymorphicStructNetSerializer.h
new file mode 100644
--- /dev/null
+++ b/Engine/Plugins/Runtime/ReplicationSystemTestPlugin/Source/Private/Tests/Serialization/TestPolymorphicStructNetSerializer.h
@@ -0,0 +1,228 @@
+// Copyright Epic Games, Inc. All Rights Reserved.
+
+#pragma once
+
+#include "Iris/Serialization/PolymorphicNetSerializer.h"
+#include "UObject/ObjectMacros.h"
+#include "TestPolymorphicStructNetSerializer.generated.h"
+
+/**
+* Polymorphic setup used in gamecode which we need to support
+*/
+USTRUCT()
+struct FExamplePolymorphicArrayItem
+{
+	GENERATED_USTRUCT_BODY()
+
+	virtual ~FExamplePolymorphicArrayItem() { }
+
+	/** Returns the serialization data, must always be overridden */
+	virtual UScriptStruct* GetScriptStruct() const
+	{
+		return FExamplePolymorphicArrayItem::StaticStruct();
+	}
+};
+
+USTRUCT()
+struct FExamplePolymorphicArrayStruct
+{
+	GENERATED_BODY()
+
+	FExamplePolymorphicArrayStruct() { }
+
+	FExamplePolymorphicArrayStruct(FExamplePolymorphicArrayItem* DataPtr)
+	{
+		Data.Add(TSharedPtr<FExamplePolymorphicArrayItem>(DataPtr));
+	}
+
+	FExamplePolymorphicArrayStruct(FExamplePolymorphicArrayStruct&& Other) : Data(MoveTemp(Other.Data))	{ }
+	FExamplePolymorphicArrayStruct(const FExamplePolymorphicArrayStruct& Other) : Data(Other.Data) { }
+
+	FExamplePolymorphicArrayStruct& operator=(FExamplePolymorphicArrayStruct&& Other) { Data = MoveTemp(Other.Data); return *this; }
+	FExamplePolymorphicArrayStruct& operator=(const FExamplePolymorphicArrayStruct& Other) { Data = Other.Data; return *this; }
+
+private:
+	/** Raw storage of target data, do not modify this directly */
+	TArray<TSharedPtr<FExamplePolymorphicArrayItem>, TInlineAllocator<1> >	Data;
+
+public:
+	/** Resets handle to have no targets */
+	void Clear()
+	{
+		Data.Reset();
+	}
+
+	/** Returns number of target data, not number of actors/targets as target data may contain multiple actors */
+	int32 Num() const
+	{
+		return Data.Num();
+	}
+
+	/** Returns true if there are any valid targets */
+	bool IsValid(int32 Index) const
+	{
+		return (Index < Data.Num() && Data[Index].IsValid());
+	}
+
+	/** Returns data at index, or nullptr if invalid */
+	const FExamplePolymorphicArrayItem* Get(int32 Index) const
+	{
+		return IsValid(Index) ? Data[Index].Get() : nullptr;
+	}
+
+	/** Returns data at index, or nullptr if invalid */
+	FExamplePolymorphicArrayItem* Get(int32 Index)
+	{
+		return IsValid(Index) ? Data[Index].Get() : nullptr;
+	}
+
+	/** Adds a new target data to handle, it must have been created with new */
+	void Add(FExamplePolymorphicArrayItem* DataPtr)
+	{
+		Data.Add(TSharedPtr<FExamplePolymorphicArrayItem>(DataPtr));
+	}
+
+	/** Does a shallow copy of target data from one handle to another */
+	void Append(const FExamplePolymorphicArrayStruct& OtherHandle)
+	{
+		Data.Append(OtherHandle.Data);
+	}
+
+	/** Serialize for networking, handles polymorphism */
+	bool NetSerialize(FArchive& Ar, class UPackageMap* Map, bool& bOutSuccess) { return true; };
+
+	/** Comparison operator */
+	bool operator==(const FExamplePolymorphicArrayStruct& Other) const
+	{
+		// Both invalid structs or both valid and Pointer compare (???) // deep comparison equality
+		// From the original code this is odd, It should really compare the data instead as if we clone 
+		if (Data.Num() != Other.Data.Num())
+		{
+			return false;
+		}
+		for (int32 i = 0; i < Data.Num(); ++i)
+		{
+			if (Data[i].IsValid() != Other.Data[i].IsValid())
+			{
+				return false;
+			}
+			if (Data[i].Get() != Other.Data[i].Get())
+			{
+				return false;
+			}
+		}
+		return true;
+	}
+
+	/** Comparison operator */
+	bool operator!=(const FExamplePolymorphicArrayStruct& Other) const
+	{
+		return !(FExamplePolymorphicArrayStruct::operator==(Other));
+	}
+
+	/** For PolymorphicStructArrayNetSerializer */
+	static TArrayView<TSharedPtr<FExamplePolymorphicArrayItem>> GetArray(FExamplePolymorphicArrayStruct& ArrayContainer)
+	{
+		return MakeArrayView(ArrayContainer.Data);
+	}
+
+	static void SetArrayNum(FExamplePolymorphicArrayStruct& ArrayContainer, SIZE_T Num)
+	{
+		ArrayContainer.Data.SetNum(static_cast<SSIZE_T>(Num));
+	}
+};
+
+template<>
+struct TStructOpsTypeTraits<FExamplePolymorphicArrayStruct> : public TStructOpsTypeTraitsBase2<FExamplePolymorphicArrayStruct>
+{
+	enum
+	{
+		WithCopy = true,		// Necessary so that TSharedPtr<FExamplePolymorphicArrayItem> Data is copied around
+		WithNetSerializer = true,
+		WithIdenticalViaEquality = true,
+	};
+};
+
+// Some test structs
+USTRUCT()
+struct FExamplePolymorphicStructA : public FExamplePolymorphicArrayItem
+{
+	GENERATED_USTRUCT_BODY()
+
+	UPROPERTY()
+	int32 SomeInt;
+
+	/** Returns the serialization data, must always be overridden */
+	virtual UScriptStruct* GetScriptStruct() const override
+	{
+		return FExamplePolymorphicStructA::StaticStruct();
+	}
+};
+
+USTRUCT()
+struct FExamplePolymorphicStructB : public FExamplePolymorphicArrayItem
+{
+	GENERATED_USTRUCT_BODY()
+
+	UPROPERTY()
+	int32 SomeFloat;
+
+	/** Returns the serialization data, must always be overridden */
+	virtual UScriptStruct* GetScriptStruct() const override
+	{
+		return FExamplePolymorphicStructB::StaticStruct();
+	}
+};
+
+USTRUCT()
+struct FExamplePolymorphicStructC : public FExamplePolymorphicArrayItem
+{
+	GENERATED_USTRUCT_BODY()
+
+	UPROPERTY()
+	bool SomeBool;
+
+	UPROPERTY()
+	UObject* SomeObjectRef;
+
+	/** Returns the serialization data, must always be overridden */
+	virtual UScriptStruct* GetScriptStruct() const override
+	{
+		return FExamplePolymorphicStructC::StaticStruct();
+	}
+};
+
+USTRUCT()
+struct FExamplePolymorphicStructD : public FExamplePolymorphicArrayItem
+{
+	GENERATED_USTRUCT_BODY()
+
+	uint32 SomeValue;
+
+	/** Returns the serialization data, must always be overridden */
+	virtual UScriptStruct* GetScriptStruct() const override
+	{
+		return FExamplePolymorphicStructD::StaticStruct();
+	}
+};
+
+USTRUCT()
+struct FExamplePolymorphicStructDNetSerializerConfig : public FNetSerializerConfig
+{
+	GENERATED_BODY()
+};
+
+USTRUCT()
+struct FTestPolymorphicArrayStructNetSerializerConfig : public FPolymorphicArrayStructNetSerializerConfig
+{
+	GENERATED_USTRUCT_BODY()
+};
+
+
+namespace UE::Net
+{
+
+UE_NET_DECLARE_SERIALIZER(FExamplePolymorphicStructDNetSerializer, REPLICATIONSYSTEMTESTPLUGIN_API);
+UE_NET_DECLARE_SERIALIZER(FTestPolymorphicArrayStructNetSerializer, REPLICATIONSYSTEMTESTPLUGIN_API);
+
+}
diff --git a/Engine/Plugins/Runtime/ReplicationSystemTestPlugin/Source/Private/Tests/Serialization/TestRotatorNetSerializer.cpp b/Engine/Plugins/Runtime/ReplicationSystemTestPlugin/Source/Private/Tests/Serialization/TestRotatorNetSerializer.cpp
new file mode 100644
--- /dev/null
+++ b/Engine/Plugins/Runtime/ReplicationSystemTestPlugin/Source/Private/Tests/Serialization/TestRotatorNetSerializer.cpp
@@ -0,0 +1,229 @@
+// Copyright Epic Games, Inc. All Rights Reserved.
+
+#include "TestNetSerializerFixture.h"
+#include "Iris/Serialization/RotatorNetSerializers.h"
+#include <limits>
+
+namespace UE::Net
+{
+
+FTestMessage& operator<<(FTestMessage& Message, const FRotator& Rotator)
+{
+	return Message << Rotator.ToString();
+}
+
+}
+
+namespace UE::Net::Private
+{
+
+static FTestMessage& PrintRotatorNetSerializerConfig(FTestMessage& Message, const FNetSerializerConfig& InConfig)
+{
+	return Message;
+}
+
+class FTestRotatorNetSerializerBase : public TTestNetSerializerFixture<PrintRotatorNetSerializerConfig, FRotator>
+{
+private:
+	typedef TTestNetSerializerFixture<PrintRotatorNetSerializerConfig, FRotator> Super;
+
+public:
+	FTestRotatorNetSerializerBase(const FNetSerializer& Serializer) : Super(Serializer) {}
+
+	void TestSerialize();
+	void TestSerializeDelta();
+	void TestIsEqual();
+	void TestValidate();
+
+protected:
+	using ScalarType = decltype(FRotator::Pitch);
+
+	struct FFloatTriplet
+	{
+		ScalarType X;
+		ScalarType Y;
+		ScalarType Z;
+	};
+
+	static const FFloatTriplet BadValues[];
+	static const SIZE_T BadValueCount;
+	static const FRotator Values[];
+	static const SIZE_T ValueCount;
+	float ValueTolerance = 0.0f;
+	const FNetSerializerConfig* NetSerializerConfig = nullptr;
+};
+
+// In case there's any sort of sanity checking in the FRotator constructors this will avoid them.
+const FTestRotatorNetSerializerBase::FFloatTriplet FTestRotatorNetSerializerBase::BadValues[] =
+{
+	{-INFINITY, NAN, -NAN},
+	{-INFINITY, INFINITY, -INFINITY},
+	{NAN, NAN, NAN},
+	{-NAN, -NAN, -NAN},
+	{-1.0f, 0.0f, 0.0f},
+	{0.0f, -1.0f, 0.0f},
+	{0.0f, 0.0f, -1.0f},
+	{5.0f, 6.0f, 360.0f},
+	{std::numeric_limits<ScalarType>::max(), std::numeric_limits<ScalarType>::max(), std::numeric_limits<ScalarType>::max()},
+	{std::numeric_limits<ScalarType>::lowest(), std::numeric_limits<ScalarType>::lowest(), std::numeric_limits<ScalarType>::lowest()},
+};
+const SIZE_T FTestRotatorNetSerializerBase::BadValueCount = sizeof(BadValues)/sizeof(BadValues[0]);
+
+const FRotator FTestRotatorNetSerializerBase::Values[] =
+{
+	FRotator(0.0f),
+	FRotator(1.0f, 0.0f, 0.0f),
+	FRotator(0.0f, 1.0f, 0.0f),
+	FRotator(0.0f, 0.0f, 1.0f),
+	FRotator(-0.0f), // 0.0f == -0.0f so we expect this to be handled gracefully
+	FRotator(359.999f, 129.130f, 45.0f),
+	FRotator(1.26698556E-23),
+	FRotator(47.11f),
+};
+const SIZE_T FTestRotatorNetSerializerBase::ValueCount = sizeof(Values)/sizeof(Values[0]);
+
+class FTestRotatorAsShortNetSerializer : public FTestRotatorNetSerializerBase
+{
+private:
+	typedef FTestRotatorNetSerializerBase Super;
+
+public:
+	FTestRotatorAsShortNetSerializer()
+	: Super(UE_NET_GET_SERIALIZER(FRotatorAsShortNetSerializer))
+	{
+		ValueTolerance = 1.5f*(360.0f/65536.0f);
+		NetSerializerConfig = &RotatorAsShortNetSerializerConfig;
+	}
+
+private:
+	inline static const FRotatorAsShortNetSerializerConfig RotatorAsShortNetSerializerConfig;
+};
+
+class FTestRotatorAsByteNetSerializer : public FTestRotatorNetSerializerBase
+{
+private:
+	typedef FTestRotatorNetSerializerBase Super;
+
+public:
+	FTestRotatorAsByteNetSerializer()
+	: Super(UE_NET_GET_SERIALIZER(FRotatorAsByteNetSerializer))
+	{
+		ValueTolerance = 1.5f*(360.0f/256.0f);
+		NetSerializerConfig = &RotatorAsByteNetSerializerConfig;
+	}
+
+private:
+	inline static const FRotatorAsByteNetSerializerConfig RotatorAsByteNetSerializerConfig;
+};
+
+// FRotatorAsShortNetSerializer tests
+UE_NET_TEST_FIXTURE(FTestRotatorAsShortNetSerializer, TestIsEqual)
+{
+	TestIsEqual();
+}
+
+UE_NET_TEST_FIXTURE(FTestRotatorAsShortNetSerializer, TestSerialize)
+{
+	TestSerialize();
+}
+
+UE_NET_TEST_FIXTURE(FTestRotatorAsShortNetSerializer, TestSerializeDelta)
+{
+	TestSerializeDelta();
+}
+
+UE_NET_TEST_FIXTURE(FTestRotatorAsShortNetSerializer, TestValidate)
+{
+	TestValidate();
+}
+
+// FRotatorAsByteNetSerializer tests
+UE_NET_TEST_FIXTURE(FTestRotatorAsByteNetSerializer, TestIsEqual)
+{
+	TestIsEqual();
+}
+
+UE_NET_TEST_FIXTURE(FTestRotatorAsByteNetSerializer, TestSerialize)
+{
+	TestSerialize();
+}
+
+UE_NET_TEST_FIXTURE(FTestRotatorAsByteNetSerializer, TestSerializeDelta)
+{
+	TestSerializeDelta();
+}
+
+UE_NET_TEST_FIXTURE(FTestRotatorAsByteNetSerializer, TestValidate)
+{
+	TestValidate();
+}
+
+void FTestRotatorNetSerializerBase::TestIsEqual()
+{
+	FRotator CompareValues[2][sizeof(Values)/sizeof(Values[0])];
+	bool ExpectedResults[2][sizeof(Values)/sizeof(Values[0])];
+
+	for (SIZE_T ValueIt = 0; ValueIt < ValueCount; ++ValueIt)
+	{
+		// We need to memcopy the data to avoid any validation/fixing in the constructor
+		FPlatformMemory::Memcpy(&CompareValues[0][ValueIt], &Values[ValueIt], sizeof(FRotator));
+		ExpectedResults[0][ValueIt] = true;
+	}
+	for (SIZE_T ValueIt = 0; ValueIt < ValueCount; ++ValueIt)
+	{
+		FPlatformMemory::Memcpy(&CompareValues[1][ValueIt], &Values[(ValueIt + 1) % ValueCount], sizeof(FRotator));
+		ExpectedResults[1][ValueIt] = false;
+	}
+
+	for (SIZE_T TestRoundIt = 0, TestRoundEndIt = 2; TestRoundIt != TestRoundEndIt; ++TestRoundIt)
+	{
+		constexpr bool bQuantizedCompare = true;
+		bool bSuccess = Super::TestIsEqual(Values, CompareValues[TestRoundIt], ExpectedResults[TestRoundIt], ValueCount, *NetSerializerConfig, bQuantizedCompare);
+		if (!bSuccess)
+		{
+			return;
+		}
+	}
+}
+
+void FTestRotatorNetSerializerBase::TestSerialize()
+{
+	constexpr bool bQuantizedCompare = true;
+	Super::TestSerialize(Values, Values, ValueCount, *NetSerializerConfig, bQuantizedCompare);
+}
+
+void FTestRotatorNetSerializerBase::TestSerializeDelta()
+{
+	Super::TestSerializeDelta(Values, ValueCount, *NetSerializerConfig);
+}
+
+void FTestRotatorNetSerializerBase::TestValidate()
+{
+	// Test bad values
+	{
+		static_assert(sizeof(BadValues[0]) == sizeof(FRotator), "BadValues and FRotator type mismatch.");
+		bool ExpectedResults[BadValueCount] = {};
+		const bool bSuccess = Super::TestValidate(reinterpret_cast<const FRotator*>(BadValues), ExpectedResults, BadValueCount, *NetSerializerConfig);
+		if (!bSuccess)
+		{
+			return;
+		}
+	}
+
+	// Test conforming values
+	{
+		bool ExpectedResults[ValueCount];
+		for (SIZE_T ValueIt = 0; ValueIt < ValueCount; ++ValueIt)
+		{
+			ExpectedResults[ValueIt] = true;
+		}
+
+		const bool bSuccess = Super::TestValidate(Values, ExpectedResults, ValueCount, *NetSerializerConfig);
+		if (!bSuccess)
+		{
+			return;
+		}
+	}
+}
+
+}
diff --git a/Engine/Plugins/Runtime/ReplicationSystemTestPlugin/Source/ReplicationSystemTestPlugin.Build.cs b/Engine/Plugins/Runtime/ReplicationSystemTestPlugin/Source/ReplicationSystemTestPlugin.Build.cs
new file mode 100644
--- /dev/null
+++ b/Engine/Plugins/Runtime/ReplicationSystemTestPlugin/Source/ReplicationSystemTestPlugin.Build.cs
@@ -0,0 +1,57 @@
+// Copyright 1998-2018 Epic Games, Inc. All Rights Reserved.
+
+using System.IO;
+
+namespace UnrealBuildTool.Rules
+{
+	public class ReplicationSystemTestPlugin : ModuleRules
+	{
+		public ReplicationSystemTestPlugin(ReadOnlyTargetRules Target) : base(Target)
+		{
+			var EngineDir = Path.GetFullPath(Target.RelativeEnginePath);
+
+			// We never want to precompile this plugin
+			PrecompileForTargets = PrecompileTargetsType.None;
+
+			PublicIncludePaths.AddRange(
+				new string[] {
+					ModuleDirectory + "/Public",
+
+				}
+				);
+
+			PrivateIncludePaths.AddRange(
+				new string[]
+				{
+                    Path.Combine(EngineDir, "Source/Runtime/Experimental/Iris/Core/Private"),
+				}
+				);
+
+			PublicDependencyModuleNames.AddRange(
+				new string[]
+				{
+					"Core",
+					"Projects",
+				}
+				);
+
+			PrivateDependencyModuleNames.AddRange(
+				new string[]
+				{
+					"CoreUObject",
+					"Engine",
+					"NetCore",
+				}
+				);
+
+			DynamicallyLoadedModuleNames.AddRange(
+				new string[]
+				{
+					// ... add any modules that your module loads dynamically here ...
+				}
+				);
+
+			SetupIrisSupport(Target);
+		}
+	}
+}
diff --git a/Engine/Source/Developer/GameplayDebugger/Private/GameplayDebuggerModule.cpp b/Engine/Source/Developer/GameplayDebugger/Private/GameplayDebuggerModule.cpp
--- a/Engine/Source/Developer/GameplayDebugger/Private/GameplayDebuggerModule.cpp
+++ b/Engine/Source/Developer/GameplayDebugger/Private/GameplayDebuggerModule.cpp
@@ -21,6 +21,10 @@
 #include "Editor/GameplayDebuggerEdMode.h"
 #endif
 
+#if UE_WITH_IRIS
+#include "Iris/IrisConfig.h"
+#endif // UE_WITH_IRIS
+
 class FGameplayDebuggerModule : public IGameplayDebugger
 {
 public:
@@ -177,6 +181,14 @@ AGameplayDebuggerPlayerManager& FGameplayDebuggerModule::GetPlayerManager(UWorld
 
 void FGameplayDebuggerModule::OnWorldInitialized(UWorld* World, const UWorld::InitializationValues IVS)
 {
+#if UE_WITH_IRIS
+	if (UE::Net::ShouldUseIrisReplication())
+	{
+		// $IRIS: $TODO: Implement custom Iris NetSerilalizer to be able to support this., https://jira.it.epicgames.com/browse/UE-158457
+		return;
+	}
+#endif // UE_WITH_IRIS
+
 	// make sure that world has valid player manager, create when it doesn't
 	if (World && (World->IsGameWorld() || World->IsEditorWorld()))
 	{
diff --git a/Engine/Source/Programs/ReplicationSystemTest/Private/AutomationTestRunner.cpp b/Engine/Source/Programs/ReplicationSystemTest/Private/AutomationTestRunner.cpp
new file mode 100644
--- /dev/null
+++ b/Engine/Source/Programs/ReplicationSystemTest/Private/AutomationTestRunner.cpp
@@ -0,0 +1,60 @@
+// Copyright Epic Games, Inc. All Rights Reserved.
+
+#include "AutomationTestRunner.h"
+#include "ReplicationSystemTestPlugin/NetworkAutomationTest.h"
+
+#if WITH_AUTOMATION_WORKER
+
+namespace UE::Net
+{
+
+FAutomationTestRunner::FAutomationTestRunner()
+{
+}
+
+void
+FAutomationTestRunner::RunTests()
+{
+	constexpr int32 ExpectedTestCount = 2048;
+
+	TArray<FAutomationTestInfo> TestInfos;
+	TestInfos.Empty(ExpectedTestCount);
+
+	FAutomationTestFramework& TestFramework = FAutomationTestFramework::Get();
+	TestFramework.SetRequestedTestFilter(EAutomationTestFlags::SmokeFilter | EAutomationTestFlags::EngineFilter);
+	TestFramework.GetValidTestNames(TestInfos);
+	int TestCount = TestInfos.Num();
+	if (TestCount <= 0)
+		return;
+
+	// Stack walking doesn't work properly on Windows when omitting frame pointers. See WindowsPlatformStackWalk.cpp.
+	const bool bCaptureStack = TestFramework.GetCaptureStack();
+	TestFramework.SetCaptureStack(false);
+
+	const FString TestPrefix("Net.");
+	const double TestStartTime = FPlatformTime::Seconds();
+	for (int TestIt = 0; TestIt != TestCount; ++TestIt)
+	{
+		const FAutomationTestInfo& TestInfo = TestInfos[TestIt];
+		const FString& TestCommand = TestInfo.GetTestName();
+		if (!TestCommand.StartsWith(TestPrefix) && !TestInfo.GetDisplayName().StartsWith(TestPrefix))
+			continue;
+
+		constexpr int32 RoleIndex = 0;
+		TestFramework.StartTestByName(TestCommand, RoleIndex);
+	}
+	const double TestEndTime = FPlatformTime::Seconds();
+
+	const double TestTime = TestEndTime - TestStartTime;
+	UE_LOG(LogNetworkAutomationTest, Display, TEXT("Tests took %.3lf seconds to execute"), TestTime);
+
+	FAutomationTestExecutionInfo ExecutionInfo;
+	TestFramework.StopTest(ExecutionInfo);
+	TestFramework.SetCaptureStack(bCaptureStack);
+
+	PrintNetworkAutomationTestSummary();
+}
+
+}
+
+#endif
