Check that the serializer's quantized type size is small enough to fit the test harness' buffer.

#rnx

diff --git a/Engine/Plugins/Runtime/ReplicationSystemTestPlugin/Source/Private/Tests/Serialization/TestNetSerializerFixture.cpp b/Engine/Plugins/Runtime/ReplicationSystemTestPlugin/Source/Private/Tests/Serialization/TestNetSerializerFixture.cpp
--- a/Engine/Plugins/Runtime/ReplicationSystemTestPlugin/Source/Private/Tests/Serialization/TestNetSerializerFixture.cpp
+++ b/Engine/Plugins/Runtime/ReplicationSystemTestPlugin/Source/Private/Tests/Serialization/TestNetSerializerFixture.cpp
@@ -57,6 +57,8 @@ FTestNetSerializerFixture::FTestNetSerializerFixture(const FNetSerializer& InSer
 , SerializerVersionOverride(InSerializer.Version)
 , Context(&Reader, &Writer)
 {
+	check(Serializer.QuantizedTypeSize <= sizeof(QuantizedBuffer[0]));
+
 	// If this serializer could have allocated state then Taint quantized state buffers before testing.
 	// Assume buffers are properly freed if necessary after every use of these buffers.
 	if (EnumHasAnyFlags(InSerializer.Traits, ENetSerializerTraits::IsForwardingSerializer | ENetSerializerTraits::HasDynamicState))
