Iris
- Converted Tests to use TObjectPtr

diff --git a/Engine/Plugins/Runtime/ReplicationSystemTestPlugin/Source/Private/Tests/ReplicationState/TestFastArrayReplicationState.h b/Engine/Plugins/Runtime/ReplicationSystemTestPlugin/Source/Private/Tests/ReplicationState/TestFastArrayReplicationState.h
--- a/Engine/Plugins/Runtime/ReplicationSystemTestPlugin/Source/Private/Tests/ReplicationState/TestFastArrayReplicationState.h
+++ b/Engine/Plugins/Runtime/ReplicationSystemTestPlugin/Source/Private/Tests/ReplicationState/TestFastArrayReplicationState.h
@@ -26,7 +26,7 @@ struct FTestFastArrayReplicationState_FastArrayItem : public FFastArraySerialize
 	int32 RepInt32;
 
 	UPROPERTY()
-	UObject* ObjectRef = nullptr;
+	TObjectPtr<UObject> ObjectRef = nullptr;
 
 	// Callbacks
 	void PostReplicatedAdd(const struct FTestFastArrayReplicationState_FastArraySerializer& InArraySerializer);
diff --git a/Engine/Plugins/Runtime/ReplicationSystemTestPlugin/Source/Private/Tests/ReplicationState/TestReplicationStateDescriptorBuilder.h b/Engine/Plugins/Runtime/ReplicationSystemTestPlugin/Source/Private/Tests/ReplicationState/TestReplicationStateDescriptorBuilder.h
--- a/Engine/Plugins/Runtime/ReplicationSystemTestPlugin/Source/Private/Tests/ReplicationState/TestReplicationStateDescriptorBuilder.h
+++ b/Engine/Plugins/Runtime/ReplicationSystemTestPlugin/Source/Private/Tests/ReplicationState/TestReplicationStateDescriptorBuilder.h
@@ -362,7 +362,7 @@ struct FTestReplicationStateDescriptor_TestStructWithRef
 	GENERATED_BODY();
 
 	UPROPERTY()
-	UObject* ObjectRef;
+	TObjectPtr<UObject> ObjectRef;
 };
 
 USTRUCT()
@@ -376,7 +376,7 @@ struct FTestReplicationStateDescriptor_TestStructWithRefCArray
 	GENERATED_BODY();
 
 	UPROPERTY()
-	UObject* ObjectRef[ArrayElementCount];
+	TObjectPtr<UObject> ObjectRef[ArrayElementCount];
 };
 
 USTRUCT()
@@ -385,7 +385,7 @@ struct FTestReplicationStateDescriptor_TestStructWithRefTArray
 	GENERATED_BODY();
 
 	UPROPERTY()
-	TArray<UObject*> ObjectRef;
+	TArray<TObjectPtr<UObject>> ObjectRef;
 };
 
 UCLASS()
@@ -397,7 +397,7 @@ public:
 	UTestReplicationStateDescriptor_TestClassWithRef() : UObject() {}
 
 	UPROPERTY(Transient, Replicated)
-	UObject* ObjectRef;
+	TObjectPtr<UObject> ObjectRef;
 };
 
 UCLASS()
@@ -449,7 +449,7 @@ public:
 	UTestReplicationStateDescriptor_TestClassWithRefInCArray() : UObject() {}
 
 	UPROPERTY(Transient, Replicated)
-	UObject* ObjectRef[ArrayElementCount];
+	TObjectPtr<UObject> ObjectRef[ArrayElementCount];
 };
 
 UCLASS()
@@ -461,7 +461,7 @@ public:
 	UTestReplicationStateDescriptor_TestClassWithRefInTArray() : UObject() {}
 
 	UPROPERTY(Transient, Replicated)
-	TArray<UObject*> ObjectRefInArray;
+	TArray<TObjectPtr<UObject>> ObjectRefInArray;
 };
 
 UCLASS()
diff --git a/Engine/Plugins/Runtime/ReplicationSystemTestPlugin/Source/Private/Tests/ReplicationSystem/ReplicatedTestObject.h b/Engine/Plugins/Runtime/ReplicationSystemTestPlugin/Source/Private/Tests/ReplicationSystem/ReplicatedTestObject.h
--- a/Engine/Plugins/Runtime/ReplicationSystemTestPlugin/Source/Private/Tests/ReplicationSystem/ReplicatedTestObject.h
+++ b/Engine/Plugins/Runtime/ReplicationSystemTestPlugin/Source/Private/Tests/ReplicationSystem/ReplicatedTestObject.h
@@ -108,7 +108,7 @@ public:
 	int32 IntA;
 
 	UPROPERTY(Transient, Replicated)
-	UObject* RawObjectPtrRef;
+	TObjectPtr<UObject> RawObjectPtrRef;
 
 	UPROPERTY(Transient, Replicated)
 	TWeakObjectPtr<UObject> WeakObjectPtrObjectRef;
@@ -338,7 +338,7 @@ public:
 	int8 IntC;
 
 	UPROPERTY(Transient, Replicated)
-	UObject* RawObjectPtrRef;
+	TObjectPtr<UObject> RawObjectPtrRef;
 
 	UPROPERTY(Transient, Replicated)
 	TWeakObjectPtr<UObject> WeakObjectPtrObjectRef;
@@ -364,7 +364,7 @@ public:
 	int32 IntA;
 
 	UPROPERTY(Transient, Replicated)
-	UObject* OtherSubObject;
+	TObjectPtr<UObject> OtherSubObject;
 
 	uint32 LastRepOrderCounter = 0U;
 
diff --git a/Engine/Plugins/Runtime/ReplicationSystemTestPlugin/Source/Private/Tests/ReplicationSystem/TestObjectReferences.cpp b/Engine/Plugins/Runtime/ReplicationSystemTestPlugin/Source/Private/Tests/ReplicationSystem/TestObjectReferences.cpp
--- a/Engine/Plugins/Runtime/ReplicationSystemTestPlugin/Source/Private/Tests/ReplicationSystem/TestObjectReferences.cpp
+++ b/Engine/Plugins/Runtime/ReplicationSystemTestPlugin/Source/Private/Tests/ReplicationSystem/TestObjectReferences.cpp
@@ -809,7 +809,7 @@ UE_NET_TEST_FIXTURE(FReplicationSystemServerClientTestFixture, TestExternalSubOb
 
 	// Verify that we could resolve the subobject ref and that we got the expected reference
 	UE_NET_ASSERT_TRUE(ClientObject->ObjectRef != nullptr);
-	UE_NET_ASSERT_EQ(ClientObject->ObjectRef, (UObject*)ClientExternalSubObject);
+	UE_NET_ASSERT_EQ(ClientObject->ObjectRef, TObjectPtr<UObject>(ClientExternalSubObject));
 }
 
 // Test partial resolve of references in array
@@ -1208,7 +1208,7 @@ UE_NET_TEST_FIXTURE(FReplicationSystemServerClientTestFixture, TestRemovedRefere
 	UE_NET_ASSERT_TRUE(ClientObject != nullptr);
 
 	// Verify that object can be resolved
-	UE_NET_ASSERT_NE(ClientObject->StructWithRef.ObjectRef, nullptr);
+	UE_NET_ASSERT_NE(ClientObject->StructWithRef.ObjectRef, TObjectPtr<UObject>(nullptr));
 
 	// Disable replication for ServerReferenceObjectA
 	Server->ReplicationSystem->AddToGroup(NotReplicatedNetObjectGroupHandle, ServerReferencedObjectA->NetHandle);
@@ -1227,7 +1227,7 @@ UE_NET_TEST_FIXTURE(FReplicationSystemServerClientTestFixture, TestRemovedRefere
 	Server->PostSendUpdate();
 
 	// Verify that object can no longer be resolved
-	UE_NET_ASSERT_EQ(ClientObject->StructWithRef.ObjectRef, nullptr);
+	UE_NET_ASSERT_EQ(ClientObject->StructWithRef.ObjectRef, TObjectPtr<UObject>(nullptr));
 
 	// Enable replication for ServerReferenceObjectA
 	Server->ReplicationSystem->RemoveFromGroup(NotReplicatedNetObjectGroupHandle, ServerReferencedObjectA->NetHandle);
@@ -1242,7 +1242,7 @@ UE_NET_TEST_FIXTURE(FReplicationSystemServerClientTestFixture, TestRemovedRefere
 	Client->PostSendUpdate();
 
 	// Verify that we can now resolve the referenced object again
-	UE_NET_ASSERT_NE(ClientObject->StructWithRef.ObjectRef, nullptr);
+	UE_NET_ASSERT_NE(ClientObject->StructWithRef.ObjectRef, TObjectPtr<UObject>(nullptr));
 }
 
 // Replicate an object with a reference to another object, in two different changemasks, which is later filtered out and then replicated again.
@@ -1275,8 +1275,8 @@ UE_NET_TEST_FIXTURE(FReplicationSystemServerClientTestFixture, TestRemovedRefere
 	UE_NET_ASSERT_TRUE(ClientObject != nullptr);
 
 	// Verify that object can be resolved
-	UE_NET_ASSERT_NE(ClientObject->StructWithRef.ObjectRef, nullptr);
-	UE_NET_ASSERT_NE(ClientObject->StructWithRef_CArray[0].ObjectRef, nullptr);
+	UE_NET_ASSERT_NE(ClientObject->StructWithRef.ObjectRef, TObjectPtr<UObject>(nullptr));
+	UE_NET_ASSERT_NE(ClientObject->StructWithRef_CArray[0].ObjectRef, TObjectPtr<UObject>(nullptr));
 
 	// Disable replication for ServerReferenceObjectA
 	Server->ReplicationSystem->AddToGroup(NotReplicatedNetObjectGroupHandle, ServerReferencedObjectA->NetHandle);
@@ -1296,8 +1296,8 @@ UE_NET_TEST_FIXTURE(FReplicationSystemServerClientTestFixture, TestRemovedRefere
 	Server->PostSendUpdate();
 
 	// Verify that object can no longer be resolved
-	UE_NET_ASSERT_EQ(ClientObject->StructWithRef.ObjectRef, nullptr);
-	UE_NET_ASSERT_EQ(ClientObject->StructWithRef_CArray[0].ObjectRef, nullptr);
+	UE_NET_ASSERT_EQ(ClientObject->StructWithRef.ObjectRef, TObjectPtr<UObject>(nullptr));
+	UE_NET_ASSERT_EQ(ClientObject->StructWithRef_CArray[0].ObjectRef, TObjectPtr<UObject>(nullptr));
 
 	// Enable replication for ServerReferenceObjectA
 	Server->ReplicationSystem->RemoveFromGroup(NotReplicatedNetObjectGroupHandle, ServerReferencedObjectA->NetHandle);
@@ -1312,8 +1312,8 @@ UE_NET_TEST_FIXTURE(FReplicationSystemServerClientTestFixture, TestRemovedRefere
 	Client->PostSendUpdate();
 
 	// Verify that we can now resolve the referenced object again
-	UE_NET_ASSERT_NE(ClientObject->StructWithRef.ObjectRef, nullptr);
-	UE_NET_ASSERT_NE(ClientObject->StructWithRef_CArray[0].ObjectRef, nullptr);
+	UE_NET_ASSERT_NE(ClientObject->StructWithRef.ObjectRef, TObjectPtr<UObject>(nullptr));
+	UE_NET_ASSERT_NE(ClientObject->StructWithRef_CArray[0].ObjectRef, TObjectPtr<UObject>(nullptr));
 }
 
 // Replicate an object with references to multiple objects, in a single changemask, which are later filtered out and then replicated again.
@@ -1359,10 +1359,10 @@ UE_NET_TEST_FIXTURE(FReplicationSystemServerClientTestFixture, TestRemovingRefer
 	auto& ClientTArray = ClientObject->StructWithRef_TArray;
 
 	// Verify that objects can be resolved
-	UE_NET_ASSERT_NE(ClientTArray[0].ObjectRef, nullptr);
-	UE_NET_ASSERT_NE(ClientTArray[1].ObjectRef, nullptr);
-	UE_NET_ASSERT_NE(ClientTArray[2].ObjectRef, nullptr);
-	UE_NET_ASSERT_NE(ClientTArray[3].ObjectRef, nullptr);
+	UE_NET_ASSERT_NE(ClientTArray[0].ObjectRef, TObjectPtr<UObject>(nullptr));
+	UE_NET_ASSERT_NE(ClientTArray[1].ObjectRef, TObjectPtr<UObject>(nullptr));
+	UE_NET_ASSERT_NE(ClientTArray[2].ObjectRef, TObjectPtr<UObject>(nullptr));
+	UE_NET_ASSERT_NE(ClientTArray[3].ObjectRef, TObjectPtr<UObject>(nullptr));
 
 	// Disable replication for ServerReferenceObjectA, B and D.
 	Server->ReplicationSystem->AddToGroup(NotReplicatedNetObjectGroupHandle, ServerReferencedObjectA->NetHandle);
@@ -1394,10 +1394,10 @@ UE_NET_TEST_FIXTURE(FReplicationSystemServerClientTestFixture, TestRemovingRefer
 	Client->PostSendUpdate();
 
 	// Verify that objects A, B and D can no longer be resolved, while C still can be resolved.
-	UE_NET_ASSERT_EQ(ClientTArray[0].ObjectRef, nullptr);
-	UE_NET_ASSERT_EQ(ClientTArray[1].ObjectRef, nullptr);
-	UE_NET_ASSERT_NE(ClientTArray[2].ObjectRef, nullptr);
-	UE_NET_ASSERT_EQ(ClientTArray[3].ObjectRef, nullptr);
+	UE_NET_ASSERT_EQ(ClientTArray[0].ObjectRef, TObjectPtr<UObject>(nullptr));
+	UE_NET_ASSERT_EQ(ClientTArray[1].ObjectRef, TObjectPtr<UObject>(nullptr));
+	UE_NET_ASSERT_NE(ClientTArray[2].ObjectRef, TObjectPtr<UObject>(nullptr));
+	UE_NET_ASSERT_EQ(ClientTArray[3].ObjectRef, TObjectPtr<UObject>(nullptr));
 
 	// Enable replication for ServerReferenceObjectA, B and D.
 	Server->ReplicationSystem->RemoveFromGroup(NotReplicatedNetObjectGroupHandle, ServerReferencedObjectA->NetHandle);
@@ -1414,10 +1414,10 @@ UE_NET_TEST_FIXTURE(FReplicationSystemServerClientTestFixture, TestRemovingRefer
 	Client->PostSendUpdate();
 
 	// Verify that we can now resolve all the referenced objects again
-	UE_NET_ASSERT_NE(ClientTArray[0].ObjectRef, nullptr);
-	UE_NET_ASSERT_NE(ClientTArray[1].ObjectRef, nullptr);
-	UE_NET_ASSERT_NE(ClientTArray[2].ObjectRef, nullptr);
-	UE_NET_ASSERT_NE(ClientTArray[3].ObjectRef, nullptr);
+	UE_NET_ASSERT_NE(ClientTArray[0].ObjectRef, TObjectPtr<UObject>(nullptr));
+	UE_NET_ASSERT_NE(ClientTArray[1].ObjectRef, TObjectPtr<UObject>(nullptr));
+	UE_NET_ASSERT_NE(ClientTArray[2].ObjectRef, TObjectPtr<UObject>(nullptr));
+	UE_NET_ASSERT_NE(ClientTArray[3].ObjectRef, TObjectPtr<UObject>(nullptr));
 }
 
 // Replicate an object with reference to an object which is later torn off.
@@ -1449,7 +1449,7 @@ UE_NET_TEST_FIXTURE(FReplicationSystemServerClientTestFixture, TestTearOffDoesIn
 	UE_NET_ASSERT_NE(ClientObject, nullptr);
 
 	// Verify that object can be resolved
-	UE_NET_ASSERT_NE(ClientObject->StructWithRef.ObjectRef, nullptr);
+	UE_NET_ASSERT_NE(ClientObject->StructWithRef.ObjectRef, TObjectPtr<UObject>(nullptr));
 
 	// Tear off ServerReferenceObjectA
 	Server->ReplicationBridge->EndReplication(ServerReferencedObjectA, EEndReplicationFlags::TearOff);
@@ -1476,7 +1476,7 @@ UE_NET_TEST_FIXTURE(FReplicationSystemServerClientTestFixture, TestTearOffDoesIn
 	Client->PostSendUpdate();
 
 	// Verify that object cannot be resolved
-	UE_NET_ASSERT_EQ(ClientObject->StructWithRef.ObjectRef, nullptr);
+	UE_NET_ASSERT_EQ(ClientObject->StructWithRef.ObjectRef, TObjectPtr<UObject>(nullptr));
 }
 
 }
diff --git a/Engine/Plugins/Runtime/ReplicationSystemTestPlugin/Source/Private/Tests/ReplicationSystem/TestObjectReferences.h b/Engine/Plugins/Runtime/ReplicationSystemTestPlugin/Source/Private/Tests/ReplicationSystem/TestObjectReferences.h
--- a/Engine/Plugins/Runtime/ReplicationSystemTestPlugin/Source/Private/Tests/ReplicationSystem/TestObjectReferences.h
+++ b/Engine/Plugins/Runtime/ReplicationSystemTestPlugin/Source/Private/Tests/ReplicationSystem/TestObjectReferences.h
@@ -17,7 +17,7 @@ struct FTestObjectReferences_TestStructWithRef
 	GENERATED_BODY();
 
 	UPROPERTY()
-	UObject* ObjectRef;
+	TObjectPtr<UObject> ObjectRef;
 
 	UPROPERTY()
 	int32 IntValue;
@@ -31,7 +31,7 @@ struct FTestObjectReferences_TestStructWithRefCArray
 	GENERATED_BODY();
 
 	UPROPERTY()
-	UObject* Ref_CArray[ArrayElementCount];
+	TObjectPtr<UObject> Ref_CArray[ArrayElementCount];
 };
 
 USTRUCT()
@@ -40,7 +40,7 @@ struct FTestObjectReferences_TestStructWithRefTArray
 	GENERATED_BODY();
 
 	UPROPERTY()
-	TArray<UObject*> Ref_TArray;
+	TArray<TObjectPtr<UObject>> Ref_TArray;
 };
 
 USTRUCT()
@@ -49,10 +49,10 @@ struct FTestObjectReferences_TestObjectRefFastArrayItem : public FFastArraySeria
 	GENERATED_BODY();
 
 	UPROPERTY()
-	UObject* ObjectRef;
+	TObjectPtr<UObject> ObjectRef;
 
 	UPROPERTY()
-	TArray<UObject*> Ref_TArray;
+	TArray<TObjectPtr<UObject>> Ref_TArray;
 };
 
 USTRUCT()
@@ -114,7 +114,7 @@ public:
 	virtual void RegisterReplicationFragments(UE::Net::FFragmentRegistrationContext& Fragments, UE::Net::EFragmentRegistrationFlags RegistrationFlags) override;
 
 	UPROPERTY(Transient, Replicated)
-	UObject* ObjectRef;
+	TObjectPtr<UObject> ObjectRef;
 
 	UPROPERTY(Transient, Replicated)
 	FTestObjectReferences_TestStructWithRef StructWithRef;
@@ -126,10 +126,10 @@ public:
 	FTestObjectReferences_TestStructWithRefTArray StructWithRefTArray;
 
 	UPROPERTY(Transient, Replicated)
-	UObject* Ref_CArray[ArrayElementCount];
+	TObjectPtr<UObject> Ref_CArray[ArrayElementCount];
 
 	UPROPERTY(Transient, Replicated)
-	TArray<UObject*> Ref_TArray;
+	TArray<TObjectPtr<UObject>> Ref_TArray;
 
 	UPROPERTY(Transient, Replicated)
 	FTestObjectReferences_TestStructWithRef StructWithRef_CArray[ArrayElementCount];
@@ -161,10 +161,10 @@ public:
 	virtual void RegisterReplicationFragments(UE::Net::FFragmentRegistrationContext& Fragments, UE::Net::EFragmentRegistrationFlags RegistrationFlags) override;
 
 	UPROPERTY(Transient, Replicated)
-	UObject* ObjectRef;
+	TObjectPtr<UObject> ObjectRef;
 
 	// Used to store subobject created before we start replication
 	UPROPERTY(Transient)
-	UObject* CreatedSubObjectRef;
+	TObjectPtr<UObject> CreatedSubObjectRef;
 };
 
diff --git a/Engine/Plugins/Runtime/ReplicationSystemTestPlugin/Source/Private/Tests/Serialization/TestPolymorphicStructNetSerializer.h b/Engine/Plugins/Runtime/ReplicationSystemTestPlugin/Source/Private/Tests/Serialization/TestPolymorphicStructNetSerializer.h
--- a/Engine/Plugins/Runtime/ReplicationSystemTestPlugin/Source/Private/Tests/Serialization/TestPolymorphicStructNetSerializer.h
+++ b/Engine/Plugins/Runtime/ReplicationSystemTestPlugin/Source/Private/Tests/Serialization/TestPolymorphicStructNetSerializer.h
@@ -183,7 +183,7 @@ struct FExamplePolymorphicStructC : public FExamplePolymorphicArrayItem
 	bool SomeBool;
 
 	UPROPERTY()
-	UObject* SomeObjectRef;
+	TObjectPtr<UObject> SomeObjectRef = nullptr;
 
 	/** Returns the serialization data, must always be overridden */
 	virtual UScriptStruct* GetScriptStruct() const override
