Iris
- Added support for using NetTrace for ReplicationSystemTest
- Added -TestFilter commandline options to ReplicationSystemTest allowing to only run tests where the name contains the provided string.

diff --git a/Engine/Plugins/Runtime/ReplicationSystemTestPlugin/Source/Private/Tests/ReplicationSystem/ReplicationSystemServerClientTestFixture.cpp b/Engine/Plugins/Runtime/ReplicationSystemTestPlugin/Source/Private/Tests/ReplicationSystem/ReplicationSystemServerClientTestFixture.cpp
--- a/Engine/Plugins/Runtime/ReplicationSystemTestPlugin/Source/Private/Tests/ReplicationSystem/ReplicationSystemServerClientTestFixture.cpp
+++ b/Engine/Plugins/Runtime/ReplicationSystemTestPlugin/Source/Private/Tests/ReplicationSystem/ReplicationSystemServerClientTestFixture.cpp
@@ -8,6 +8,7 @@
 #include "Iris/ReplicationSystem/ReplicationWriter.h"
 #include "Iris/ReplicationSystem/ReplicationReader.h"
 #include "Iris/PacketControl/PacketNotification.h"
+#include "Net/Core/Trace/NetTrace.h"
 
 namespace UE::Net
 {
@@ -61,7 +62,7 @@ void FDataStreamTestUtil::AddDataStreamDefinition(const TCHAR* StreamName, const
 }
 
 // FReplicationSystemTestNode implementation
-FReplicationSystemTestNode::FReplicationSystemTestNode(bool bIsServer)
+FReplicationSystemTestNode::FReplicationSystemTestNode(bool bIsServer, const TCHAR* Name)
 {
 	ReplicationBridge = NewObject<UReplicatedTestObjectBridge>();
 	CreatedObjects.Add(TStrongObjectPtr<UObject>(ReplicationBridge));
@@ -80,6 +81,13 @@ FReplicationSystemTestNode::FReplicationSystemTestNode(bool bIsServer)
 	}
 	LogIris.SetVerbosity(IrisLogVerbosity);
 	check(ReplicationBridge != nullptr);
+
+	UE_NET_TRACE_UPDATE_INSTANCE(GetNetTraceId(), bIsServer, Name);
+}
+
+uint32 FReplicationSystemTestNode::GetNetTraceId() const
+{ 
+	return ReplicationSystem ? ReplicationSystem->GetId() : ~0U;
 }
 
 UTestReplicatedIrisObject* FReplicationSystemTestNode::CreateObject(uint32 NumComponents, uint32 NumIrisComponents)
@@ -177,6 +185,9 @@ uint32 FReplicationSystemTestNode::AddConnection()
 	Connection.NetTokenDataStream = StaticCast<UNetTokenDataStream*>(Connection.DataStreamManager->GetStream("NetToken"));
 	Connection.DataStreamManager->CreateStream("Replication");
 	Connection.ReplicationDataStream = StaticCast<UReplicationDataStream*>(Connection.DataStreamManager->GetStream("Replication"));
+
+	UE_NET_TRACE_CONNECTION_CREATED(GetNetTraceId(), Connection.ConnectionId);
+	UE_NET_TRACE_CONNECTION_STATE_UPDATED(GetNetTraceId(), Connection.ConnectionId, static_cast<uint8>(3));
 	
 	// Add a connection
 	ReplicationSystem->AddConnection(Connection.ConnectionId);
@@ -213,6 +224,8 @@ bool FReplicationSystemTestNode::SendUpdate(uint32 ConnectionId)
 
 	FNetSerializationContext Context(&Writer);
 
+	Context.SetNetTraceCollector(UE_NET_TRACE_CREATE_COLLECTOR(ENetTraceVerbosity::Trace));
+
 	FConnectionInfo& Connection = GetConnectionInfo(ConnectionId);
 
 	const FDataStreamRecord* Record = nullptr;
@@ -222,6 +235,7 @@ bool FReplicationSystemTestNode::SendUpdate(uint32 ConnectionId)
 	{
 		Writer.CommitWrites();
 		Packet.BitCount = Writer.GetPosBits();
+		Packet.PacketId = PacketId++;
 
 		Connection.WriteRecords.Enqueue(Record);
 		Connection.WrittenPackets.Enqueue(Packet);
@@ -229,6 +243,10 @@ bool FReplicationSystemTestNode::SendUpdate(uint32 ConnectionId)
 
 	Connection.DataStreamManager->EndWrite();
 
+	UE_NET_TRACE_FLUSH_COLLECTOR(Context.GetTraceCollector(), GetNetTraceId(), Connection.ConnectionId, ENetTracePacketType::Outgoing);
+	UE_NET_TRACE_DESTROY_COLLECTOR(Context.GetTraceCollector());
+	UE_NET_TRACE_PACKET_SEND(GetNetTraceId(), Connection.ConnectionId, Packet.PacketId, Packet.BitCount);
+
 	return bResult;
 }
 
@@ -240,16 +258,26 @@ void FReplicationSystemTestNode::PostSendUpdate()
 void FReplicationSystemTestNode::DeliverTo(FReplicationSystemTestNode& Dest, uint32 LocalConnectionId, uint32 RemoteConnectionId, bool bDeliver)
 {
 	FConnectionInfo& Connection = GetConnectionInfo(LocalConnectionId);
+	const FPacketData& Packet = Connection.WrittenPackets.Peek();
 
 	if (bDeliver)
 	{
-		const FPacketData& Packet = Connection.WrittenPackets.Peek();
-
 		FNetBitStreamReader Reader;
 		Reader.InitBits(Packet.PacketBuffer, Packet.BitCount);
 		FNetSerializationContext Context(&Reader);
 
+		Context.SetNetTraceCollector(UE_NET_TRACE_CREATE_COLLECTOR(ENetTraceVerbosity::Trace));
+
 		Dest.RecvUpdate(RemoteConnectionId, Context);
+
+		UE_NET_TRACE_FLUSH_COLLECTOR(Context.GetTraceCollector(), Dest.GetNetTraceId(), RemoteConnectionId, ENetTracePacketType::Incoming);
+		UE_NET_TRACE_DESTROY_COLLECTOR(Context.GetTraceCollector());
+		UE_NET_TRACE_PACKET_RECV(Dest.GetNetTraceId(), RemoteConnectionId, Packet.PacketId, Packet.BitCount);
+	}
+	else
+	{
+		UE_NET_TRACE_PACKET_DROPPED(Dest.GetNetTraceId(), RemoteConnectionId, Packet.PacketId, ENetTracePacketType::Incoming);
+		UE_NET_TRACE_PACKET_DROPPED(GetNetTraceId(), RemoteConnectionId, Packet.PacketId, ENetTracePacketType::Outgoing);
 	}
 
 	// If this triggers an assert, ensure that SendTo() actually wrote any packets before.
@@ -270,20 +298,24 @@ void FReplicationSystemTestNode::RecvUpdate(uint32 ConnectionId, FNetSerializati
 
 FReplicationSystemTestNode::~FReplicationSystemTestNode()
 {
+	const uint32 NetTraceId = ReplicationSystem->GetId();
 	FReplicationSystemFactory::DestroyReplicationSystem(ReplicationSystem);
 	CreatedObjects.Empty();
+
+	// End NetTrace session for this instance
+	UE_NET_TRACE_END_SESSION(NetTraceId);
 }
 
 // FReplicationSystemTestClient implementation
-FReplicationSystemTestClient::FReplicationSystemTestClient(uint32 ReplicationSystemId)
-: FReplicationSystemTestNode(false)
+FReplicationSystemTestClient::FReplicationSystemTestClient(uint32 ReplicationSystemId, const TCHAR* Name)
+: FReplicationSystemTestNode(false, Name)
 , ConnectionIdOnServer(~0U)
 {
 }
 
 // FReplicationSystemTestServer implementation
-FReplicationSystemTestServer::FReplicationSystemTestServer()
-: FReplicationSystemTestNode(true)
+FReplicationSystemTestServer::FReplicationSystemTestServer(const TCHAR* Name)
+: FReplicationSystemTestNode(true, Name)
 {
 }
 
@@ -323,12 +355,12 @@ void FReplicationSystemServerClientTestFixture::SetUp()
 	DataStreamUtil.AddDataStreamDefinition(TEXT("Replication"), TEXT("/Script/IrisCore.ReplicationDataStream"));
 	DataStreamUtil.FixupDefinitions();
 
-	Server = new FReplicationSystemTestServer;
+	Server = new FReplicationSystemTestServer(GetName());
 }
 
 FReplicationSystemTestClient* FReplicationSystemServerClientTestFixture::CreateClient()
 {
-	FReplicationSystemTestClient* Client = new FReplicationSystemTestClient(Clients.Num() + 1U);
+	FReplicationSystemTestClient* Client = new FReplicationSystemTestClient(Clients.Num() + 1U, GetName());
 	Clients.Add(Client);
 
 	// The client needs a connection
diff --git a/Engine/Plugins/Runtime/ReplicationSystemTestPlugin/Source/Private/Tests/ReplicationSystem/ReplicationSystemServerClientTestFixture.h b/Engine/Plugins/Runtime/ReplicationSystemTestPlugin/Source/Private/Tests/ReplicationSystem/ReplicationSystemServerClientTestFixture.h
--- a/Engine/Plugins/Runtime/ReplicationSystemTestPlugin/Source/Private/Tests/ReplicationSystem/ReplicationSystemServerClientTestFixture.h
+++ b/Engine/Plugins/Runtime/ReplicationSystemTestPlugin/Source/Private/Tests/ReplicationSystem/ReplicationSystemServerClientTestFixture.h
@@ -58,6 +58,7 @@ public:
 		
 		alignas(16) uint8 PacketBuffer[MaxPacketSize];
 		uint32 BitCount;
+		uint32 PacketId;
 	};
 
 	struct FConnectionInfo
@@ -72,7 +73,7 @@ public:
 	};
 
 public:
-	explicit FReplicationSystemTestNode(bool bIsServer);
+	explicit FReplicationSystemTestNode(bool bIsServer, const TCHAR* Name);
 	~FReplicationSystemTestNode();
 
 	template<typename T>
@@ -142,6 +143,8 @@ public:
 
 	FConnectionInfo& GetConnectionInfo(uint32 ConnectionId) { return Connections[ConnectionId - 1]; }
 
+	uint32 GetNetTraceId() const;
+
 public:
 	UReplicationSystem* ReplicationSystem;
 	UReplicatedTestObjectBridge* ReplicationBridge;
@@ -149,12 +152,14 @@ public:
 
 private:
 	TArray<FConnectionInfo> Connections;
+	uint32 PacketId = 0U;
 };
 
 class FReplicationSystemTestClient : public FReplicationSystemTestNode
 {
 public:
-	explicit FReplicationSystemTestClient(uint32 ReplicationSystemId);
+	FReplicationSystemTestClient(uint32 ReplicationSystemId, const TCHAR* Name);
+
 	uint32 ConnectionIdOnServer;
 	uint32 LocalConnectionId;
 };
@@ -162,7 +167,7 @@ public:
 class FReplicationSystemTestServer : public FReplicationSystemTestNode
 {
 public:
-	explicit FReplicationSystemTestServer();
+	explicit FReplicationSystemTestServer(const TCHAR* Name);
 
 	// Send data and deliver to the client if bDeliver is true
 	bool SendAndDeliverTo(FReplicationSystemTestClient* Client, bool bDeliver);
diff --git a/Engine/Plugins/Runtime/ReplicationSystemTestPlugin/Source/Private/Tests/Serialization/TestPolymorphicStructNetSerializer.cpp b/Engine/Plugins/Runtime/ReplicationSystemTestPlugin/Source/Private/Tests/Serialization/TestPolymorphicStructNetSerializer.cpp
--- a/Engine/Plugins/Runtime/ReplicationSystemTestPlugin/Source/Private/Tests/Serialization/TestPolymorphicStructNetSerializer.cpp
+++ b/Engine/Plugins/Runtime/ReplicationSystemTestPlugin/Source/Private/Tests/Serialization/TestPolymorphicStructNetSerializer.cpp
@@ -363,7 +363,7 @@ void FTestPolymorphicArrayStructNetSerializerFixture::ValidateExpectedState(cons
 	UE_NET_ASSERT_EQ(StructA->SomeInt, AInstance->SomeInt);
 	UE_NET_ASSERT_EQ(StructB->SomeFloat, BInstance->SomeFloat);
 	UE_NET_ASSERT_EQ(StructC->SomeBool, CInstance->SomeBool);
-	UE_NET_ASSERT_EQ(StructC->SomeObjectRef, CInstance->SomeObjectRef);
+	UE_NET_ASSERT_EQ(ToRawPtr(StructC->SomeObjectRef), ToRawPtr(CInstance->SomeObjectRef));
 	UE_NET_ASSERT_EQ(StructD->SomeValue, DInstance->SomeValue);
 }
 
diff --git a/Engine/Source/Programs/ReplicationSystemTest/Private/AutomationTestRunner.cpp b/Engine/Source/Programs/ReplicationSystemTest/Private/AutomationTestRunner.cpp
--- a/Engine/Source/Programs/ReplicationSystemTest/Private/AutomationTestRunner.cpp
+++ b/Engine/Source/Programs/ReplicationSystemTest/Private/AutomationTestRunner.cpp
@@ -13,7 +13,7 @@ FAutomationTestRunner::FAutomationTestRunner()
 }
 
 void
-FAutomationTestRunner::RunTests()
+FAutomationTestRunner::RunTests(const TCHAR* TestFilter)
 {
 	constexpr int32 ExpectedTestCount = 2048;
 
@@ -40,6 +40,11 @@ FAutomationTestRunner::RunTests()
 		if (!TestCommand.StartsWith(TestPrefix) && !TestInfo.GetDisplayName().StartsWith(TestPrefix))
 			continue;
 
+		if (TestFilter && !TestCommand.Contains(TestFilter))
+		{
+			continue;
+		}
+
 		constexpr int32 RoleIndex = 0;
 		TestFramework.StartTestByName(TestCommand, RoleIndex);
 	}
diff --git a/Engine/Source/Programs/ReplicationSystemTest/Private/AutomationTestRunner.h b/Engine/Source/Programs/ReplicationSystemTest/Private/AutomationTestRunner.h
--- a/Engine/Source/Programs/ReplicationSystemTest/Private/AutomationTestRunner.h
+++ b/Engine/Source/Programs/ReplicationSystemTest/Private/AutomationTestRunner.h
@@ -14,7 +14,7 @@ class FAutomationTestRunner
 public:
 	FAutomationTestRunner();
 	
-	void RunTests();
+	void RunTests(const TCHAR* TestFilter = nullptr);
 };
 
 }
diff --git a/Engine/Source/Programs/ReplicationSystemTest/Private/NullTestRunner.h b/Engine/Source/Programs/ReplicationSystemTest/Private/NullTestRunner.h
--- a/Engine/Source/Programs/ReplicationSystemTest/Private/NullTestRunner.h
+++ b/Engine/Source/Programs/ReplicationSystemTest/Private/NullTestRunner.h
@@ -5,5 +5,5 @@
 class FNullTestRunner
 {
 public:
-	void RunTests() {}
+	void RunTests(const TCHAR* Filter = nullptr) {}
 };
diff --git a/Engine/Source/Programs/ReplicationSystemTest/Private/ReplicationSystemTest.cpp b/Engine/Source/Programs/ReplicationSystemTest/Private/ReplicationSystemTest.cpp
--- a/Engine/Source/Programs/ReplicationSystemTest/Private/ReplicationSystemTest.cpp
+++ b/Engine/Source/Programs/ReplicationSystemTest/Private/ReplicationSystemTest.cpp
@@ -6,6 +6,7 @@
 #include "Styling/UMGCoreStyle.h"
 #include "AutomationTestRunner.h"
 #include "NullTestRunner.h"
+#include "Net/Core/Trace/Private/NetTraceInternal.h"
 
 #if WITH_AUTOMATION_WORKER
 namespace UE::Net
@@ -35,12 +36,29 @@ int TestMain()
 INT32_MAIN_INT32_ARGC_TCHAR_ARGV()
 #endif
 {
+	FString OriginalCmdLine;
+
+#if !(defined(PLATFORM_XBOXONE) && PLATFORM_XBOXONE)
+	// Parse original cmdline if there is one
+	OriginalCmdLine = FCommandLine::BuildFromArgV(nullptr, ArgC, ArgV, nullptr);	
+#endif
+
 	// Due to some code not respecting nullrhi etc.
 	PRIVATE_GIsRunningCommandlet = true;
 	PRIVATE_GAllowCommandletRendering = false;
 	PRIVATE_GAllowCommandletAudio = false;
 
-	FCommandLine::Set(TEXT("-nullrhi -log -NoAsyncLoadingThread -NoAsyncPostLoad -noedl -unattended -ENGINEINI=Engine.ini"));
+	// Init cmd line used for test
+	{
+		FString CmdLineOverride(TEXT("-nullrhi -log -NoAsyncLoadingThread -NoAsyncPostLoad -noedl -unattended -ENGINEINI=Engine.ini"));
+		FString NetTraceVerbosity;
+		if(FParse::Value(ToCStr(OriginalCmdLine), TEXT("NetTrace="), NetTraceVerbosity))
+		{
+			CmdLineOverride.Appendf(TEXT(" -trace=net,log,cpu -nettrace=%s"), ToCStr(NetTraceVerbosity));
+		}
+
+		FCommandLine::Set(ToCStr(CmdLineOverride));
+	}
 
 	PreInit();
 	LoadModules();
@@ -50,7 +68,17 @@ INT32_MAIN_INT32_ARGC_TCHAR_ARGV()
 
 	{
 		TUniquePtr<TestRunner> Runner(new TestRunner());
-		Runner->RunTests();
+
+		FString TestFilter;
+		if (FParse::Value(ToCStr(OriginalCmdLine), TEXT("TestFilter="), TestFilter))
+		{
+			Runner->RunTests(ToCStr(TestFilter));
+		}
+		else
+		{
+			Runner->RunTests();
+		}
+		
 	}
 
 	TearDown();
@@ -85,6 +113,8 @@ static void TearDown()
 		delete GConfig;
 		GConfig = nullptr;
 	}
+
+	FTraceAuxiliary::Shutdown();
 }
 
 static void PreInit()
@@ -100,6 +130,22 @@ static void PreInit()
 #endif
 	FPlatformMemory::Init();
 
+	// Initialize trace
+	FString Parameter;
+	if (FParse::Value(FCommandLine::Get(), TEXT("-trace="), Parameter, false))
+	{
+		FTraceAuxiliary::Initialize(FCommandLine::Get());
+		FTraceAuxiliary::TryAutoConnect();
+	}
+
+#if UE_NET_TRACE_ENABLED
+	uint32 NetTraceVerbosity;
+	if(FParse::Value(FCommandLine::Get(), TEXT("NetTrace="), NetTraceVerbosity))
+	{
+		FNetTrace::SetTraceVerbosity(NetTraceVerbosity);
+	}
+#endif
+
 	if (IPlatformFile* WrapperFile = FPlatformFileManager::Get().GetPlatformFile(TEXT("ReplicationSystemTestFile")))
 	{
 		IPlatformFile* CurrentPlatformFile = &FPlatformFileManager::Get().GetPlatformFile();
diff --git a/Engine/Source/Programs/ReplicationSystemTest/ReplicationSystemTest.Target.cs b/Engine/Source/Programs/ReplicationSystemTest/ReplicationSystemTest.Target.cs
--- a/Engine/Source/Programs/ReplicationSystemTest/ReplicationSystemTest.Target.cs
+++ b/Engine/Source/Programs/ReplicationSystemTest/ReplicationSystemTest.Target.cs
@@ -40,6 +40,8 @@ public class ReplicationSystemTestTarget : TargetRules
 		// Network config
 		bWithPushModel = true;		
 		bUseIris = true;
+
+		GlobalDefinitions.Add("UE_TRACE_ENABLED=1");
 	}
 }
 
