Consolidated SetNetTraceCollector/GetTraceCollector to SetTraceCollector and GetTraceCollector

diff --git a/Engine/Plugins/Runtime/ReplicationSystemTestPlugin/Source/Private/Tests/ReplicationSystem/ReplicationSystemServerClientTestFixture.cpp b/Engine/Plugins/Runtime/ReplicationSystemTestPlugin/Source/Private/Tests/ReplicationSystem/ReplicationSystemServerClientTestFixture.cpp
--- a/Engine/Plugins/Runtime/ReplicationSystemTestPlugin/Source/Private/Tests/ReplicationSystem/ReplicationSystemServerClientTestFixture.cpp
+++ b/Engine/Plugins/Runtime/ReplicationSystemTestPlugin/Source/Private/Tests/ReplicationSystem/ReplicationSystemServerClientTestFixture.cpp
@@ -224,7 +224,7 @@ bool FReplicationSystemTestNode::SendUpdate(uint32 ConnectionId)
 
 	FNetSerializationContext Context(&Writer);
 
-	Context.SetNetTraceCollector(UE_NET_TRACE_CREATE_COLLECTOR(ENetTraceVerbosity::Trace));
+	Context.SetTraceCollector(UE_NET_TRACE_CREATE_COLLECTOR(ENetTraceVerbosity::Trace));
 
 	FConnectionInfo& Connection = GetConnectionInfo(ConnectionId);
 
@@ -266,7 +266,7 @@ void FReplicationSystemTestNode::DeliverTo(FReplicationSystemTestNode& Dest, uin
 		Reader.InitBits(Packet.PacketBuffer, Packet.BitCount);
 		FNetSerializationContext Context(&Reader);
 
-		Context.SetNetTraceCollector(UE_NET_TRACE_CREATE_COLLECTOR(ENetTraceVerbosity::Trace));
+		Context.SetTraceCollector(UE_NET_TRACE_CREATE_COLLECTOR(ENetTraceVerbosity::Trace));
 
 		Dest.RecvUpdate(RemoteConnectionId, Context);
 
diff --git a/Engine/Source/Runtime/Engine/Private/Net/Experimental/Iris/DataStreamChannel.cpp b/Engine/Source/Runtime/Engine/Private/Net/Experimental/Iris/DataStreamChannel.cpp
--- a/Engine/Source/Runtime/Engine/Private/Net/Experimental/Iris/DataStreamChannel.cpp
+++ b/Engine/Source/Runtime/Engine/Private/Net/Experimental/Iris/DataStreamChannel.cpp
@@ -103,7 +103,7 @@ void UDataStreamChannel::ReceivedBunch(FInBunch& Bunch)
 	FNetSerializationContext SerializationContext(&BitReader);
 
 	// For packet stats
-	SerializationContext.SetNetTraceCollector(Connection->GetInTraceCollector());
+	SerializationContext.SetTraceCollector(Connection->GetInTraceCollector());
 
 	DataStreamManager->ReadData(SerializationContext);
 
@@ -207,7 +207,7 @@ void UDataStreamChannel::Tick()
 #if UE_NET_TRACE_ENABLED	
 		// For Iris we can use the connection trace collector as long as we make sure that the packet is prepared 
 		FNetTraceCollector* Collector = Connection->GetOutTraceCollector();
-		SerializationContext.SetNetTraceCollector(Collector);
+		SerializationContext.SetTraceCollector(Collector);
 		UE_NET_TRACE_BEGIN_BUNCH(Collector);
 #endif
 
diff --git a/Engine/Source/Runtime/Experimental/Iris/Core/Private/Iris/ReplicationSystem/NetBlob/NetBlobAssembler.cpp b/Engine/Source/Runtime/Experimental/Iris/Core/Private/Iris/ReplicationSystem/NetBlob/NetBlobAssembler.cpp
--- a/Engine/Source/Runtime/Experimental/Iris/Core/Private/Iris/ReplicationSystem/NetBlob/NetBlobAssembler.cpp
+++ b/Engine/Source/Runtime/Experimental/Iris/Core/Private/Iris/ReplicationSystem/NetBlob/NetBlobAssembler.cpp
@@ -137,7 +137,7 @@ TRefCountPtr<FNetBlob> FNetBlobAssembler::Assemble(FNetSerializationContext& Con
 		BitReader.InitBits(Payload.GetData(), BitWriter.GetPosBits());
 
 		FNetSerializationContext ReadContext = Context.MakeSubContext(&BitReader);
-		ReadContext.SetNetTraceCollector(nullptr);
+		ReadContext.SetTraceCollector(nullptr);
 
 		if (NetHandle.IsValid())
 		{
diff --git a/Engine/Source/Runtime/Experimental/Iris/Core/Private/Iris/ReplicationSystem/ReplicationReader.cpp b/Engine/Source/Runtime/Experimental/Iris/Core/Private/Iris/ReplicationSystem/ReplicationReader.cpp
--- a/Engine/Source/Runtime/Experimental/Iris/Core/Private/Iris/ReplicationSystem/ReplicationReader.cpp
+++ b/Engine/Source/Runtime/Experimental/Iris/Core/Private/Iris/ReplicationSystem/ReplicationReader.cpp
@@ -1312,7 +1312,7 @@ void FReplicationReader::ProcessHugeObjectAttachment(FNetSerializationContext& C
 	FNetBitStreamReader HugeObjectReader;
 	HugeObjectReader.InitBits(NetObjectBlob.GetRawData().GetData(), NetObjectBlob.GetRawDataBitCount());
 	FNetSerializationContext HugeObjectSerializationContext = Context.MakeSubContext(&HugeObjectReader);
-	HugeObjectSerializationContext.SetNetTraceCollector(HugeObjectTraceCollector);
+	HugeObjectSerializationContext.SetTraceCollector(HugeObjectTraceCollector);
 
 	UE_NET_TRACE_NAMED_SCOPE(HugeObjectTraceScope, HugeObjectState, HugeObjectReader, HugeObjectTraceCollector, ENetTraceVerbosity::Trace);
 
diff --git a/Engine/Source/Runtime/Experimental/Iris/Core/Private/Iris/ReplicationSystem/ReplicationWriter.cpp b/Engine/Source/Runtime/Experimental/Iris/Core/Private/Iris/ReplicationSystem/ReplicationWriter.cpp
--- a/Engine/Source/Runtime/Experimental/Iris/Core/Private/Iris/ReplicationSystem/ReplicationWriter.cpp
+++ b/Engine/Source/Runtime/Experimental/Iris/Core/Private/Iris/ReplicationSystem/ReplicationWriter.cpp
@@ -2088,7 +2088,7 @@ int FReplicationWriter::PrepareAndSendHugeObjectPayload(FNetSerializationContext
 	const uint32 MaxHugeObjectPayLoadBytes = HugeObjectPayload.Num() * sizeof(HugeObjectStorageType);
 	HugeObjectWriter.InitBytes(HugeObjectPayload.GetData(), MaxHugeObjectPayLoadBytes);
 	FNetSerializationContext HugeObjectSerializationContext = Context.MakeSubContext(&HugeObjectWriter);
-	HugeObjectSerializationContext.SetNetTraceCollector(nullptr);
+	HugeObjectSerializationContext.SetTraceCollector(nullptr);
 
 	// Huge object header needed for the receiving side to be able to process this correctly.
 	FNetObjectBlob::FHeader HugeObjectHeader = {};
diff --git a/Engine/Source/Runtime/Experimental/Iris/Core/Public/Iris/Serialization/NetSerializationContext.h b/Engine/Source/Runtime/Experimental/Iris/Core/Public/Iris/Serialization/NetSerializationContext.h
--- a/Engine/Source/Runtime/Experimental/Iris/Core/Public/Iris/Serialization/NetSerializationContext.h
+++ b/Engine/Source/Runtime/Experimental/Iris/Core/Public/Iris/Serialization/NetSerializationContext.h
@@ -32,12 +32,12 @@ public:
 
 	FNetBitStreamReader* GetBitStreamReader() { return BitStreamReader; }
 	FNetBitStreamWriter* GetBitStreamWriter() { return BitStreamWriter; }
-	FNetTraceCollector* GetTraceCollector() { return TraceCollector; }
 
 	FNetSerializationContext MakeSubContext(FNetBitStreamWriter*) const;
 	FNetSerializationContext MakeSubContext(FNetBitStreamReader*) const;
 
-	void SetNetTraceCollector(FNetTraceCollector* InTraceCollector) { TraceCollector = InTraceCollector; }
+	void SetTraceCollector(FNetTraceCollector* InTraceCollector) { TraceCollector = InTraceCollector; }
+	FNetTraceCollector* GetTraceCollector() { return TraceCollector; }
 
 	bool HasError() const { return ErrorContext.HasError(); }
 	bool HasErrorOrOverflow() const;
