Iris
- Improved tests for PolymorphicNetSerializer and added some comments declaring requirements on external struct to work as expected.

diff --git a/Engine/Plugins/Runtime/ReplicationSystemTestPlugin/Source/Private/Tests/Serialization/TestPolymorphicStructNetSerializer.cpp b/Engine/Plugins/Runtime/ReplicationSystemTestPlugin/Source/Private/Tests/Serialization/TestPolymorphicStructNetSerializer.cpp
--- a/Engine/Plugins/Runtime/ReplicationSystemTestPlugin/Source/Private/Tests/Serialization/TestPolymorphicStructNetSerializer.cpp
+++ b/Engine/Plugins/Runtime/ReplicationSystemTestPlugin/Source/Private/Tests/Serialization/TestPolymorphicStructNetSerializer.cpp
@@ -10,13 +10,16 @@
 #include "Iris/Serialization/InternalNetSerializationContext.h"
 #include "Iris/ReplicationState/PropertyNetSerializerInfoRegistry.h"
 #include "Iris/Serialization/NetReferenceCollector.h"
+#include "Net/UnrealNetwork.h"
+#include "Core/Public/Iris/ReplicationSystem/ReplicationFragmentUtil.h"
 
 namespace UE::Net
 {
 
-// Custom NetSerializer declaration based on GameplayAbilityTargetHandle
-struct FTestPolymorphicArrayStructNetSerializer : public TPolymorphicArrayStructNetSerializerImpl<FExamplePolymorphicArrayStruct, FExamplePolymorphicArrayItem, FExamplePolymorphicArrayStruct::GetArray, FExamplePolymorphicArrayStruct::SetArrayNum>
+struct FTestPolymorphicArrayStructNetSerializer : public TPolymorphicArrayStructNetSerializerImpl<FExamplePolymorphicArrayStruct, FExamplePolymorphicStructBase, FExamplePolymorphicArrayStruct::GetArray, FExamplePolymorphicArrayStruct::SetArrayNum>
 {
+	typedef TPolymorphicArrayStructNetSerializerImpl<FExamplePolymorphicArrayStruct, FExamplePolymorphicStructBase, FExamplePolymorphicArrayStruct::GetArray, FExamplePolymorphicArrayStruct::SetArrayNum> InternalNetSerializerType;
+
 	static const uint32 Version = 0;
 	static const ConfigType DefaultConfig;
 };
@@ -25,6 +28,23 @@ const FTestPolymorphicArrayStructNetSerializer::ConfigType FTestPolymorphicArray
 
 UE_NET_IMPLEMENT_SERIALIZER(FTestPolymorphicArrayStructNetSerializer);
 
+struct FTestPolymorphicStructNetSerializer : public TPolymorphicStructNetSerializerImpl<FExamplePolymorphicStruct, FExamplePolymorphicStructBase, FExamplePolymorphicStruct::GetItem>
+{
+	typedef TPolymorphicStructNetSerializerImpl<FExamplePolymorphicStruct, FExamplePolymorphicStructBase, FExamplePolymorphicStruct::GetItem> InternalNetSerializerType;
+
+	static void InitTypeCache()
+	{
+		InternalNetSerializerType::InitTypeCache<FTestPolymorphicStructNetSerializer>();
+	}
+
+	static const uint32 Version = 0;
+	static const ConfigType DefaultConfig;
+};
+
+const FTestPolymorphicStructNetSerializer::ConfigType FTestPolymorphicStructNetSerializer::DefaultConfig;
+
+UE_NET_IMPLEMENT_SERIALIZER(FTestPolymorphicStructNetSerializer);
+
 }
 
 namespace UE::Net::Private
@@ -60,6 +80,9 @@ protected:
 	FTestPolymorphicArrayStructNetSerializerConfig PolymorphicNetSerializerConfig;
 	const FNetSerializer* PolymorphicNetSerializer = &UE_NET_GET_SERIALIZER(FTestPolymorphicArrayStructNetSerializer);
 
+	FTestPolymorphicStructNetSerializerConfig PolymorphicStructNetSerializerConfig;
+	const FNetSerializer* PolymorphicStructNetSerializer = &UE_NET_GET_SERIALIZER(FTestPolymorphicStructNetSerializer);
+
 	FExamplePolymorphicArrayStruct PolymorphicStructInstance0;
 	FExamplePolymorphicArrayStruct PolymorphicStructInstance1;
 
@@ -173,7 +196,8 @@ namespace UE::Net::Private
 
 static const FName PropertyNetSerializerRegistry_NAME_ExamplePolymorphicStructD(TEXT("ExamplePolymorphicStructD"));
 UE_NET_IMPLEMENT_NAMED_STRUCT_NETSERIALIZER_INFO(PropertyNetSerializerRegistry_NAME_ExamplePolymorphicStructD, FExamplePolymorphicStructDNetSerializer);
-
+static const FName PropertyNetSerializerRegistry_NAME_TestPolymorphicStructNetSerializer(TEXT("ExamplePolymorphicStruct"));
+UE_NET_IMPLEMENT_NAMED_STRUCT_NETSERIALIZER_INFO(PropertyNetSerializerRegistry_NAME_TestPolymorphicStructNetSerializer, FTestPolymorphicStructNetSerializer);
 
 UE_NET_TEST_FIXTURE(FTestPolymorphicArrayStructNetSerializerFixture, TestQuantize)
 {
@@ -343,11 +367,14 @@ void FTestPolymorphicArrayStructNetSerializerFixture::SetUp()
 	if (!bHasRegisteredStructDNetSerializer)
 	{
 		UE_NET_REGISTER_NETSERIALIZER_INFO(PropertyNetSerializerRegistry_NAME_ExamplePolymorphicStructD);
+		UE_NET_REGISTER_NETSERIALIZER_INFO(PropertyNetSerializerRegistry_NAME_TestPolymorphicStructNetSerializer);
+
 		bHasRegisteredStructDNetSerializer = true;
 	}
 
 	// Init type registry, note that we set it up for the context we provide rather than the default config
-	PolymorphicNetSerializerConfig.RegisteredTypes.InitForType(FExamplePolymorphicArrayItem::StaticStruct());
+	PolymorphicNetSerializerConfig.RegisteredTypes.InitForType(FExamplePolymorphicStructBase::StaticStruct());
+	PolymorphicStructNetSerializerConfig.RegisteredTypes.InitForType(FExamplePolymorphicStructBase::StaticStruct());
 }
 
 void FTestPolymorphicArrayStructNetSerializerFixture::ValidateExpectedState(const FExamplePolymorphicArrayStruct& StructInstance)
@@ -552,4 +579,145 @@ void FTestPolymorphicArrayStructNetSerializerFixture::Validate()
 	PolymorphicNetSerializer->Validate(NetSerializationContext, Args);
 }
 
+UE_NET_TEST_FIXTURE(FTestPolymorphicArrayStructNetSerializerFixture, TestPolymorphicStructNetSerializer_Modify)
+{
+	// InitTypeCache as we do actual serialization
+	FTestPolymorphicStructNetSerializer::InitTypeCache();
+
+	// Add a client
+	FReplicationSystemTestClient* Client = CreateClient();
+
+	// Spawn objects on server
+	UTestPolymorphicStructNetSerializer_TestObject* ServerObject = Server->CreateObject<UTestPolymorphicStructNetSerializer_TestObject>();
+	
+	ServerObject->PolyStruct.Raise<FExamplePolymorphicStructB>();
+	ServerObject->PolyStruct.GetAs<FExamplePolymorphicStructB>().SomeFloat = 12.0f;
+
+	// Replicate
+	Server->PreSendUpdate();
+	Server->SendAndDeliverTo(Client, true);
+	Server->PostSendUpdate();
+	
+	auto ClientObject = Client->GetObjectAs<UTestPolymorphicStructNetSerializer_TestObject>(ServerObject->NetHandle);
+	UE_NET_ASSERT_TRUE(ClientObject != nullptr);
+	UE_NET_ASSERT_TRUE(ClientObject->PolyStruct.GetAs<FExamplePolymorphicStructB>().SomeFloat == 12.0f);
+
+	// Modify
+	ServerObject->PolyStruct.GetAs<FExamplePolymorphicStructB>().SomeFloat += 1.0f;
+
+	// Replicate
+	Server->PreSendUpdate();
+	Server->SendAndDeliverTo(Client, true);
+	Server->PostSendUpdate();
+
+	// Verify that we detected the modification
+	UE_NET_ASSERT_TRUE(ClientObject->PolyStruct.GetAs<FExamplePolymorphicStructB>().SomeFloat == ServerObject->PolyStruct.GetAs<FExamplePolymorphicStructB>().SomeFloat);
+
+	// Switch type
+	ServerObject->PolyStruct.Raise<FExamplePolymorphicStructD>();
+	ServerObject->PolyStruct.GetAs<FExamplePolymorphicStructD>().SomeValue = 100;
+
+	// Replicate
+	Server->PreSendUpdate();
+	Server->SendAndDeliverTo(Client, true);
+	Server->PostSendUpdate();
+
+	UE_NET_ASSERT_TRUE(ClientObject->PolyStruct.GetAs<FExamplePolymorphicStructD>().SomeValue == ServerObject->PolyStruct.GetAs<FExamplePolymorphicStructD>().SomeValue);	
+}
+
+UE_NET_TEST_FIXTURE(FTestPolymorphicArrayStructNetSerializerFixture, TestPolymorphicStructNetSerializer_ModifyFastArray)
+{
+	// InitTypeCache as we do actual serialization
+	FTestPolymorphicStructNetSerializer::InitTypeCache();
+
+	// Add a client
+	FReplicationSystemTestClient* Client = CreateClient();
+
+	// Spawn objects on server
+	UTestPolymorphicStructNetSerializer_TestObject* ServerObject = Server->CreateObject<UTestPolymorphicStructNetSerializer_TestObject>();
+
+	// Add a few entries for the fastarray
+	FExamplePolymorphicStructFastArrayItem Item;
+
+	Item.PolyStruct.Raise<FExamplePolymorphicStructA>();
+	ServerObject->PolyStructFastArray.Edit().Add(Item);
+
+	Item.PolyStruct.Raise<FExamplePolymorphicStructB>();
+	ServerObject->PolyStructFastArray.Edit().Add(Item);
+
+	Item.PolyStruct.Raise<FExamplePolymorphicStructC>();
+	ServerObject->PolyStructFastArray.Edit().Add(Item);
+
+	Item.PolyStruct.Raise<FExamplePolymorphicStructD>();
+	ServerObject->PolyStructFastArray.Edit().Add(Item);
+	ServerObject->PolyStructFastArray.Edit()[3].PolyStruct.GetAs<FExamplePolymorphicStructD>().SomeValue = 100;
+	
+	// Replicate
+	Server->PreSendUpdate();
+	Server->SendAndDeliverTo(Client, true);
+	Server->PostSendUpdate();
+	
+	auto ClientObject = Client->GetObjectAs<UTestPolymorphicStructNetSerializer_TestObject>(ServerObject->NetHandle);
+	UE_NET_ASSERT_TRUE(ClientObject != nullptr);
+	UE_NET_ASSERT_EQ(ClientObject->PolyStructFastArray.GetItemArray().Num(), ServerObject->PolyStructFastArray.GetItemArray().Num());
+	UE_NET_ASSERT_EQ(ClientObject->PolyStructFastArray.GetItemArray()[3].PolyStruct.GetAs<FExamplePolymorphicStructD>().SomeValue, 100U);
+
+	// Modify value and see that it is replicated as expected
+	ServerObject->PolyStructFastArray.Edit()[3].PolyStruct.GetAs<FExamplePolymorphicStructD>().SomeValue += 3;
+
+	// Replicate
+	Server->PreSendUpdate();
+	Server->SendAndDeliverTo(Client, true);
+	Server->PostSendUpdate();
+
+	// Verified that the client got the modified value
+	UE_NET_ASSERT_EQ(ClientObject->PolyStructFastArray.GetItemArray()[3].PolyStruct.GetAs<FExamplePolymorphicStructD>().SomeValue, 103U);
+
+	// Switch type
+	ServerObject->PolyStructFastArray.Edit()[2].PolyStruct.Raise<FExamplePolymorphicStructD>();
+	ServerObject->PolyStructFastArray.Edit()[2].PolyStruct.GetAs<FExamplePolymorphicStructD>().SomeValue = 1;
+
+	// Replicate
+	Server->PreSendUpdate();
+	Server->SendAndDeliverTo(Client, true);
+	Server->PostSendUpdate();
+
+	// Verified that the client got the modified value
+	UE_NET_ASSERT_EQ(ClientObject->PolyStructFastArray.GetItemArray()[2].PolyStruct.GetAs<FExamplePolymorphicStructD>().SomeValue, 1U);
+}
+
+}
+
+void FExamplePolymorphicStructFastArrayItem::PostReplicatedAdd(const struct FExamplePolymorphicStructFastArraySerializer& InArraySerializer)
+{
+	const_cast<FExamplePolymorphicStructFastArraySerializer&>(InArraySerializer).bHitReplicatedAdd = true;
+}
+
+void FExamplePolymorphicStructFastArrayItem::PostReplicatedChange(const struct FExamplePolymorphicStructFastArraySerializer& InArraySerializer)
+{
+	const_cast<FExamplePolymorphicStructFastArraySerializer&>(InArraySerializer).bHitReplicatedChange = true;
+}
+
+void FExamplePolymorphicStructFastArrayItem::PreReplicatedRemove(const struct FExamplePolymorphicStructFastArraySerializer& InArraySerializer)
+{
+	const_cast<FExamplePolymorphicStructFastArraySerializer&>(InArraySerializer).bHitReplicatedRemove = true;
+}
+
+UTestPolymorphicStructNetSerializer_TestObject::UTestPolymorphicStructNetSerializer_TestObject()
+: UReplicatedTestObject()
+{
+}
+
+void UTestPolymorphicStructNetSerializer_TestObject::GetLifetimeReplicatedProps(TArray<class FLifetimeProperty>& OutLifetimeProps) const
+{
+	FDoRepLifetimeParams Params;
+	Params.bIsPushBased = false;
+
+	DOREPLIFETIME_WITH_PARAMS_FAST(UTestPolymorphicStructNetSerializer_TestObject, PolyStruct, Params);
+	DOREPLIFETIME_WITH_PARAMS_FAST(UTestPolymorphicStructNetSerializer_TestObject, PolyStructFastArray, Params);
+}
+
+void UTestPolymorphicStructNetSerializer_TestObject::RegisterReplicationFragments(UE::Net::FFragmentRegistrationContext& Context, UE::Net::EFragmentRegistrationFlags RegistrationFlags)
+{
+	UE::Net::FReplicationFragmentUtil::CreateAndRegisterFragmentsForObject(this, Context, RegistrationFlags);
 }
diff --git a/Engine/Plugins/Runtime/ReplicationSystemTestPlugin/Source/Private/Tests/Serialization/TestPolymorphicStructNetSerializer.h b/Engine/Plugins/Runtime/ReplicationSystemTestPlugin/Source/Private/Tests/Serialization/TestPolymorphicStructNetSerializer.h
--- a/Engine/Plugins/Runtime/ReplicationSystemTestPlugin/Source/Private/Tests/Serialization/TestPolymorphicStructNetSerializer.h
+++ b/Engine/Plugins/Runtime/ReplicationSystemTestPlugin/Source/Private/Tests/Serialization/TestPolymorphicStructNetSerializer.h
@@ -4,25 +4,32 @@
 
 #include "Iris/Serialization/PolymorphicNetSerializer.h"
 #include "UObject/ObjectMacros.h"
+#include "Tests/ReplicationSystem/ReplicatedTestObject.h"
+#include "Core/Public/Iris/ReplicationState/IrisFastArraySerializer.h"
+#include "Core/Public/Iris/ReplicationState/Private/IrisFastArraySerializerInternal.h"
 #include "TestPolymorphicStructNetSerializer.generated.h"
 
-/**
-* Polymorphic setup used in gamecode which we need to support
-*/
+// Example BaseStruct used to support replication of polymorphic structs
+// Note: this is not an optimal setup but is needed to ensure that we support existing setups used by existing engine code.
 USTRUCT()
-struct FExamplePolymorphicArrayItem
+struct FExamplePolymorphicStructBase
 {
 	GENERATED_USTRUCT_BODY()
 
-	virtual ~FExamplePolymorphicArrayItem() { }
+	virtual ~FExamplePolymorphicStructBase() { }
 
 	/** Returns the serialization data, must always be overridden */
 	virtual UScriptStruct* GetScriptStruct() const
 	{
-		return FExamplePolymorphicArrayItem::StaticStruct();
+		return FExamplePolymorphicStructBase::StaticStruct();
 	}
 };
 
+//////////////////////////////////////////////////////////////////////////
+
+// Test struct based on examples from existing engine code, 
+// This is a very specific type of replication where structs contained in the array only are replicated based on the stored pointer.
+// i.e. it does not compare the actual instance data stored in the polymorphic struct, but instead expects user to explicitly clone and insert a new struct to trigger replication of the data contained in the struct
 USTRUCT()
 struct FExamplePolymorphicArrayStruct
 {
@@ -30,9 +37,9 @@ struct FExamplePolymorphicArrayStruct
 
 	FExamplePolymorphicArrayStruct() { }
 
-	FExamplePolymorphicArrayStruct(FExamplePolymorphicArrayItem* DataPtr)
+	FExamplePolymorphicArrayStruct(FExamplePolymorphicStructBase* DataPtr)
 	{
-		Data.Add(TSharedPtr<FExamplePolymorphicArrayItem>(DataPtr));
+		Data.Add(TSharedPtr<FExamplePolymorphicStructBase>(DataPtr));
 	}
 
 	FExamplePolymorphicArrayStruct(FExamplePolymorphicArrayStruct&& Other) : Data(MoveTemp(Other.Data))	{ }
@@ -41,61 +48,51 @@ struct FExamplePolymorphicArrayStruct
 	FExamplePolymorphicArrayStruct& operator=(FExamplePolymorphicArrayStruct&& Other) { Data = MoveTemp(Other.Data); return *this; }
 	FExamplePolymorphicArrayStruct& operator=(const FExamplePolymorphicArrayStruct& Other) { Data = Other.Data; return *this; }
 
-private:
-	/** Raw storage of target data, do not modify this directly */
-	TArray<TSharedPtr<FExamplePolymorphicArrayItem>, TInlineAllocator<1> >	Data;
-
 public:
-	/** Resets handle to have no targets */
 	void Clear()
 	{
 		Data.Reset();
 	}
 
-	/** Returns number of target data, not number of actors/targets as target data may contain multiple actors */
 	int32 Num() const
 	{
 		return Data.Num();
 	}
 
-	/** Returns true if there are any valid targets */
+	/** Returns true if index is valid */
 	bool IsValid(int32 Index) const
 	{
 		return (Index < Data.Num() && Data[Index].IsValid());
 	}
 
 	/** Returns data at index, or nullptr if invalid */
-	const FExamplePolymorphicArrayItem* Get(int32 Index) const
+	const FExamplePolymorphicStructBase* Get(int32 Index) const
 	{
 		return IsValid(Index) ? Data[Index].Get() : nullptr;
 	}
 
 	/** Returns data at index, or nullptr if invalid */
-	FExamplePolymorphicArrayItem* Get(int32 Index)
+	FExamplePolymorphicStructBase* Get(int32 Index)
 	{
 		return IsValid(Index) ? Data[Index].Get() : nullptr;
 	}
 
-	/** Adds a new target data to handle, it must have been created with new */
-	void Add(FExamplePolymorphicArrayItem* DataPtr)
+	/** Adds a new FExamplePolymorphicArrayStruct data to FExamplePolymorphicArrayStruct, it must have been created with new */
+	void Add(FExamplePolymorphicStructBase* DataPtr)
 	{
-		Data.Add(TSharedPtr<FExamplePolymorphicArrayItem>(DataPtr));
+		Data.Add(TSharedPtr<FExamplePolymorphicStructBase>(DataPtr));
 	}
 
-	/** Does a shallow copy of target data from one handle to another */
+	/** Does a shallow copy of data from one handle to another */
 	void Append(const FExamplePolymorphicArrayStruct& OtherHandle)
 	{
 		Data.Append(OtherHandle.Data);
 	}
 
-	/** Serialize for networking, handles polymorphism */
-	bool NetSerialize(FArchive& Ar, class UPackageMap* Map, bool& bOutSuccess) { return true; };
-
 	/** Comparison operator */
 	bool operator==(const FExamplePolymorphicArrayStruct& Other) const
 	{
-		// Both invalid structs or both valid and Pointer compare (???) // deep comparison equality
-		// From the original code this is odd, It should really compare the data instead as if we clone 
+		// This comparison operator only compares actual instance pointers
 		if (Data.Num() != Other.Data.Num())
 		{
 			return false;
@@ -121,7 +118,7 @@ public:
 	}
 
 	/** For PolymorphicStructArrayNetSerializer */
-	static TArrayView<TSharedPtr<FExamplePolymorphicArrayItem>> GetArray(FExamplePolymorphicArrayStruct& ArrayContainer)
+	static TArrayView<TSharedPtr<FExamplePolymorphicStructBase>> GetArray(FExamplePolymorphicArrayStruct& ArrayContainer)
 	{
 		return MakeArrayView(ArrayContainer.Data);
 	}
@@ -130,6 +127,10 @@ public:
 	{
 		ArrayContainer.Data.SetNum(static_cast<SSIZE_T>(Num));
 	}
+
+private:
+	/** Raw storage of target data, do not modify this directly */
+	TArray<TSharedPtr<FExamplePolymorphicStructBase>, TInlineAllocator<1> >	Data;
 };
 
 template<>
@@ -138,14 +139,237 @@ struct TStructOpsTypeTraits<FExamplePolymorphicArrayStruct> : public TStructOpsT
 	enum
 	{
 		WithCopy = true,		// Necessary so that TSharedPtr<FExamplePolymorphicArrayItem> Data is copied around
-		WithNetSerializer = true,
 		WithIdenticalViaEquality = true,
 	};
 };
 
-// Some test structs
+namespace UE::Net
+{
+
+UE_NET_DECLARE_SERIALIZER(FTestPolymorphicArrayStructNetSerializer, REPLICATIONSYSTEMTESTPLUGIN_API);
+
+}
+
+USTRUCT()
+struct FTestPolymorphicArrayStructNetSerializerConfig : public FPolymorphicArrayStructNetSerializerConfig
+{
+	GENERATED_USTRUCT_BODY()
+};
+
+//////////////////////////////////////////////////////////////////////////
+
+// Example of how to setup a Polymorphic struct that can work as a replicated property
+USTRUCT()
+struct FExamplePolymorphicStruct
+{
+	GENERATED_BODY()
+
+	FExamplePolymorphicStruct() { }
+
+	FExamplePolymorphicStruct(const FExamplePolymorphicStruct& Other)
+	{
+		*this = Other;
+	}
+
+	FExamplePolymorphicStruct& operator=(FExamplePolymorphicStruct&& Other)
+	{
+		Data = MoveTemp(Other.Data); return *this;
+	}
+
+	FExamplePolymorphicStruct& operator=(const FExamplePolymorphicStruct& Other) 
+	{
+		// When used as a replicated property we need to do deep copy as we store copies of the state in order to detect changes
+		if (this != &Other)
+		{
+			const UScriptStruct* DstScriptStruct = Data.IsValid() ? Data->GetScriptStruct() : nullptr;
+			const UScriptStruct* SrcScriptStruct = Other.Data.IsValid() ? Other.Data->GetScriptStruct() : nullptr;
+
+			const FExamplePolymorphicStructBase* SrcPolymorphicStruct = Other.Data.IsValid() ? Other.Data.Get() : nullptr;
+			if (SrcPolymorphicStruct)
+			{
+				FExamplePolymorphicStructBase* const DstPolymorphicStruct = Data.IsValid() ? Data.Get() : nullptr;
+				if (DstPolymorphicStruct == SrcPolymorphicStruct)
+				{
+					SrcScriptStruct->CopyScriptStruct(DstPolymorphicStruct, SrcPolymorphicStruct);
+				}
+				else
+				{
+					FExamplePolymorphicStructBase* NewPolymorphicStruct = static_cast<FExamplePolymorphicStructBase*>(FMemory::Malloc(SrcScriptStruct->GetStructureSize(), SrcScriptStruct->GetMinAlignment()));
+					SrcScriptStruct->InitializeStruct(NewPolymorphicStruct);
+					SrcScriptStruct->CopyScriptStruct(NewPolymorphicStruct, SrcPolymorphicStruct);
+					Data = MakeShareable(NewPolymorphicStruct);								
+				}				
+			}
+			else
+			{
+				Data.Reset();
+			}
+		}
+
+		return *this;
+	}
+
+	void Reset()
+	{
+		Data.Reset();
+	}
+
+	bool IsValid() const
+	{
+		return Data.IsValid();
+	}
+
+	/** Initalize Data to a struct of wanted type */
+	template<typename T>
+	void Raise()
+	{
+		Data = TSharedPtr<T>(new T);
+	}
+
+	/** Get data as expected type */
+	template<typename T>
+	T& GetAs()
+	{
+		check(Data.Get()->GetScriptStruct() == T::StaticStruct());
+		return *static_cast<T*>(Data.Get());
+	}
+
+	/** Comparison operator that compares instance data */
+	bool operator==(const FExamplePolymorphicStruct& Other) const
+	{
+		// When used as replicated property and we want to detect changes of the intance data we need to compare the data stored in the inner structs.
+		const FExamplePolymorphicStructBase* PolymorphicStruct = Data.Get();
+		const FExamplePolymorphicStructBase* OtherPolymorphicStruct = Other.Data.Get();
+		
+		const UScriptStruct* PolymorphicStructScriptStruct = PolymorphicStruct ? PolymorphicStruct->GetScriptStruct() : nullptr;
+		const UScriptStruct* OtherPolymorphicStructScriptStruct = OtherPolymorphicStruct ? OtherPolymorphicStruct->GetScriptStruct() : nullptr;
+
+		if (PolymorphicStructScriptStruct != OtherPolymorphicStructScriptStruct)
+		{
+			return false;
+		}
+
+		if (PolymorphicStructScriptStruct)
+		{
+			// Compare actual struct data	
+			return PolymorphicStructScriptStruct->CompareScriptStruct(PolymorphicStruct, OtherPolymorphicStruct, 0);
+		}
+		else
+		{
+			return false;
+		}		
+	}
+
+	bool operator!=(const FExamplePolymorphicStruct& Other) const
+	{
+		return !(FExamplePolymorphicStruct::operator==(Other));
+	}
+
+	static TSharedPtr<FExamplePolymorphicStructBase>& GetItem(FExamplePolymorphicStruct& Source)
+	{
+		return Source.Data;
+	}
+
+private:
+	TSharedPtr<FExamplePolymorphicStructBase> Data;
+};
+
+template<>
+struct TStructOpsTypeTraits<FExamplePolymorphicStruct> : public TStructOpsTypeTraitsBase2<FExamplePolymorphicStruct>
+{
+	enum
+	{
+		WithCopy = true,					// Necessary to do a deep copy
+		WithIdenticalViaEquality = true,	// We have a custom compare operator
+	};
+};
+
+namespace UE::Net
+{
+
+UE_NET_DECLARE_SERIALIZER(FTestPolymorphicStructNetSerializer, REPLICATIONSYSTEMTESTPLUGIN_API);
+
+}
+
 USTRUCT()
-struct FExamplePolymorphicStructA : public FExamplePolymorphicArrayItem
+struct FTestPolymorphicStructNetSerializerConfig : public FPolymorphicStructNetSerializerConfig
+{
+	GENERATED_USTRUCT_BODY()
+};
+
+//////////////////////////////////////////////////////////////////////////
+
+// FastArray declaration to test using a Polymorphic struct in a fastarray.
+USTRUCT()
+struct FExamplePolymorphicStructFastArrayItem : public FFastArraySerializerItem
+{
+	GENERATED_BODY()
+
+	UPROPERTY()
+	FExamplePolymorphicStruct PolyStruct;
+
+	// Callbacks
+	void PostReplicatedAdd(const struct FExamplePolymorphicStructFastArraySerializer& InArraySerializer);
+	void PostReplicatedChange(const struct FExamplePolymorphicStructFastArraySerializer& InArraySerializer);
+	void PreReplicatedRemove(const struct FExamplePolymorphicStructFastArraySerializer& InArraySerializer);
+};
+
+USTRUCT()
+struct FExamplePolymorphicStructFastArraySerializer : public FIrisFastArraySerializer
+{
+	GENERATED_BODY()
+
+	FExamplePolymorphicStructFastArraySerializer()
+	: FIrisFastArraySerializer()
+	, bHitReplicatedAdd(false)
+	, bHitReplicatedChange(false)
+	, bHitReplicatedRemove(false)
+	, bHitPostReplicatedReceive(false)
+	, bPostReplicatedReceiveWasHitWithUnresolvedReferences(false)
+	{
+	}
+
+	typedef TArray<FExamplePolymorphicStructFastArrayItem> ItemArrayType;
+	const ItemArrayType& GetItemArray() const
+	{
+		return Items;
+	}
+
+	ItemArrayType& GetItemArray()
+	{
+		return Items;
+	}
+
+	typedef UE::Net::TIrisFastArrayEditor<FExamplePolymorphicStructFastArraySerializer> FFastArrayEditor;
+	FFastArrayEditor Edit()
+	{
+		return FFastArrayEditor(*this);
+	}
+	
+	void PostReplicatedReceive(const FFastArraySerializer::FPostReplicatedReceiveParameters& Parameters)
+	{
+		bHitPostReplicatedReceive = 1U;
+		bPostReplicatedReceiveWasHitWithUnresolvedReferences = Parameters.bHasMoreUnmappedReferences;
+	}
+
+public:
+	uint8 bHitReplicatedAdd    : 1;
+	uint8 bHitReplicatedChange : 1;
+	uint8 bHitReplicatedRemove : 1;
+	uint8 bHitPostReplicatedReceive : 1;
+
+	bool bPostReplicatedReceiveWasHitWithUnresolvedReferences;
+
+private:
+	UPROPERTY()
+	TArray<FExamplePolymorphicStructFastArrayItem> Items;
+};
+
+//////////////////////////////////////////////////////////////////////////
+
+// Polymorphic structs based on FExamplePolymorphicStructBase used by test code
+USTRUCT()
+struct FExamplePolymorphicStructA : public FExamplePolymorphicStructBase
 {
 	GENERATED_USTRUCT_BODY()
 
@@ -154,13 +378,13 @@ struct FExamplePolymorphicStructA : public FExamplePolymorphicArrayItem
 
 	/** Returns the serialization data, must always be overridden */
 	virtual UScriptStruct* GetScriptStruct() const override
-	{
+	{ 
 		return FExamplePolymorphicStructA::StaticStruct();
 	}
 };
 
 USTRUCT()
-struct FExamplePolymorphicStructB : public FExamplePolymorphicArrayItem
+struct FExamplePolymorphicStructB : public FExamplePolymorphicStructBase
 {
 	GENERATED_USTRUCT_BODY()
 
@@ -175,7 +399,7 @@ struct FExamplePolymorphicStructB : public FExamplePolymorphicArrayItem
 };
 
 USTRUCT()
-struct FExamplePolymorphicStructC : public FExamplePolymorphicArrayItem
+struct FExamplePolymorphicStructC : public FExamplePolymorphicStructBase
 {
 	GENERATED_USTRUCT_BODY()
 
@@ -192,10 +416,15 @@ struct FExamplePolymorphicStructC : public FExamplePolymorphicArrayItem
 	}
 };
 
+// Example of a struct that does not have any properties, but relies on an explicit NetSerializer for serialization
 USTRUCT()
-struct FExamplePolymorphicStructD : public FExamplePolymorphicArrayItem
+struct FExamplePolymorphicStructD : public FExamplePolymorphicStructBase
 {
 	GENERATED_USTRUCT_BODY()
+	
+	// Since we do not have any properties, we need a custom compare operator and mark the struct with WithIdenticalViaEquality Trait
+	// Alternatively we could mark SomeValue as a non-replicated property (UProperty(NotReplicated))
+	bool operator==(const FExamplePolymorphicStructD& Other) const { return SomeValue == Other.SomeValue; }
 
 	uint32 SomeValue;
 
@@ -206,23 +435,44 @@ struct FExamplePolymorphicStructD : public FExamplePolymorphicArrayItem
 	}
 };
 
-USTRUCT()
-struct FExamplePolymorphicStructDNetSerializerConfig : public FNetSerializerConfig
+// Since this struct does not have any properties, it cannot be compared using properties and we must therefore set the trait to use native compare
+template<>
+struct TStructOpsTypeTraits<FExamplePolymorphicStructD> : public TStructOpsTypeTraitsBase2<FExamplePolymorphicStructD>
 {
-	GENERATED_BODY()
+	enum
+	{
+		WithIdenticalViaEquality = true,
+	};
 };
 
+namespace UE::Net
+{
+
+UE_NET_DECLARE_SERIALIZER(FExamplePolymorphicStructDNetSerializer, REPLICATIONSYSTEMTESTPLUGIN_API);
+
+}
+
 USTRUCT()
-struct FTestPolymorphicArrayStructNetSerializerConfig : public FPolymorphicArrayStructNetSerializerConfig
+struct FExamplePolymorphicStructDNetSerializerConfig : public FNetSerializerConfig
 {
-	GENERATED_USTRUCT_BODY()
+	GENERATED_BODY()
 };
 
+//////////////////////////////////////////////////////////////////////////
 
-namespace UE::Net
+// Test Object using FExamplePolymorphicStruct and FExamplePolymorphicStructFastArraySerializer used in tests
+UCLASS()
+class UTestPolymorphicStructNetSerializer_TestObject : public UReplicatedTestObject
 {
+	GENERATED_BODY()
+public:
+	UTestPolymorphicStructNetSerializer_TestObject();
 
-UE_NET_DECLARE_SERIALIZER(FExamplePolymorphicStructDNetSerializer, REPLICATIONSYSTEMTESTPLUGIN_API);
-UE_NET_DECLARE_SERIALIZER(FTestPolymorphicArrayStructNetSerializer, REPLICATIONSYSTEMTESTPLUGIN_API);
+	void RegisterReplicationFragments(UE::Net::FFragmentRegistrationContext& Fragments, UE::Net::EFragmentRegistrationFlags RegistrationFlags) override;
 
-}
+	UPROPERTY(Replicated)
+	FExamplePolymorphicStruct PolyStruct;
+
+	UPROPERTY(Replicated)
+	FExamplePolymorphicStructFastArraySerializer PolyStructFastArray;
+};
diff --git a/Engine/Source/Runtime/Experimental/Iris/Core/Public/Iris/Serialization/PolymorphicNetSerializerImpl.h b/Engine/Source/Runtime/Experimental/Iris/Core/Public/Iris/Serialization/PolymorphicNetSerializerImpl.h
--- a/Engine/Source/Runtime/Experimental/Iris/Core/Public/Iris/Serialization/PolymorphicNetSerializerImpl.h
+++ b/Engine/Source/Runtime/Experimental/Iris/Core/Public/Iris/Serialization/PolymorphicNetSerializerImpl.h
@@ -46,10 +46,21 @@ namespace UE::Net
   *
   * Helper to implement serializers that requires dynamic polymorphism.
   * It can either be used to declare a typed serializer or be used as an internal helper.
-  * 
   * ExternalSourceType is the class/struct that has the TSharedPtr<ExternalSourceItemType> data.
   * ExternalSourceItemType is the polymorphic struct type
   * GetItem is a function that will return a reference to the TSharedPtr<ExternalSourceItemType>
+  *
+  * !BIG DISCLAIMER:!
+  *
+  * This serializer was written to mimic the behavior seen in FGameplayAbilityTargetDataHandle and FGameplayEffectContextHandle 
+  * which both are written with the intent of being used for RPCs and not being used for replicated properties and uses a TSharedPointer to hold the polymorphic struct
+  *
+  * That said, IF the serializer is used for replicated properties it has very specific requirements on the implementation of the SourceType to work correctly.
+  *
+  * 1. The sourcetype MUST provide a custom assignment operator performing a deep-copy/clone
+  * 2. The sourcetype MUST define a comparison operator that compares the instance data of the stored ExternalSourceItemType
+  * 3. TStructOpsTypeTraits::WithCopy and TStructOpsTypeTraits::WithIdenticalViaEquality must be specified
+  *
   */
 template <typename ExternalSourceType, typename ExternalSourceItemType, TSharedPtr<ExternalSourceItemType>&(*GetItem)(ExternalSourceType&)>
 struct TPolymorphicStructNetSerializerImpl : protected Private::FPolymorphicStructNetSerializerInternal
@@ -67,6 +78,7 @@ struct TPolymorphicStructNetSerializerImpl : protected Private::FPolymorphicStru
 
 	typedef ExternalSourceType SourceType;
 	typedef FQuantizedData QuantizedType;
+	typedef FPolymorphicStructNetSerializerConfig ConfigType;
 
 	static void Serialize(FNetSerializationContext&, const FNetSerializeArgs&);
 	static void Deserialize(FNetSerializationContext&, const FNetDeserializeArgs&);
@@ -90,7 +102,6 @@ protected:
 
 	typedef ExternalSourceItemType SourceItemType;
 	typedef FPolymorphicNetSerializerScriptStructCache::FTypeInfo FTypeInfo;
-	typedef FPolymorphicStructNetSerializerConfig ConfigType;
 
 	struct FSourceItemTypeDeleter
 	{
@@ -121,6 +132,8 @@ private:
   *
   * Helper to implement array serializers that requires dynamic polymorphism.
   * It can either be used to declare a typed serializer or be used as an internal helper.
+  *
+  * @See: TPolymorphicStructNetSerializerImpl for requirements on external data
   */
 template <typename ExternalSourceType, typename ExternalSourceArrayItemType, TArrayView<TSharedPtr<ExternalSourceArrayItemType>>(*GetArray)(ExternalSourceType& Source), void(*SetArrayNum)(ExternalSourceType& Source, SIZE_T Num)>
 struct TPolymorphicArrayStructNetSerializerImpl : protected Private::FPolymorphicStructNetSerializerInternal
@@ -311,6 +324,9 @@ void TPolymorphicStructNetSerializerImpl<ExternalSourceType, ExternalSourceItemT
 		const FReplicationStateDescriptor* Descriptor = TypeInfo->Descriptor;
 		const UScriptStruct* ScriptStruct = TypeInfo->ScriptStruct;
 
+		// NOTE: We always allocate new memory in order to behave like the code we are trying to mimic expects, see GameplayEffectContextHandle
+		// this should really be a policy of this class as it is far from optimal.
+
 		// Allocate external struct, owned by external state therefore using global allocator
 		SourceItemType* NewData = static_cast<SourceItemType*>(FMemory::Malloc(ScriptStruct->GetStructureSize(), ScriptStruct->GetMinAlignment()));
 		ScriptStruct->InitializeStruct(NewData);
@@ -634,6 +650,9 @@ void TPolymorphicArrayStructNetSerializerImpl<ExternalSourceType, ExternalSource
 			const FReplicationStateDescriptor* Descriptor = TypeInfo->Descriptor;
 			const UScriptStruct* ScriptStruct = TypeInfo->ScriptStruct;
 
+			// NOTE: We always allocate new memory in order to behave like the code we are trying to mimic expects, see GameplayEffectContextHandle
+			// this should really be a policy of this class as it is far from optimal.
+
 			// Allocate external struct, owned by external state therefore using global allocator
 			SourceArrayItemType* NewData = static_cast<SourceArrayItemType*>(FMemory::Malloc(ScriptStruct->GetStructureSize(), ScriptStruct->GetMinAlignment()));
 			ScriptStruct->InitializeStruct(NewData);
