Iris
- Fixed issue with check being triggered for Structs with smaller alignment than 4

diff --git a/Engine/Plugins/Runtime/ReplicationSystemTestPlugin/Source/Private/Tests/ReplicationState/TestReplicationStateDescriptorBuilder.cpp b/Engine/Plugins/Runtime/ReplicationSystemTestPlugin/Source/Private/Tests/ReplicationState/TestReplicationStateDescriptorBuilder.cpp
--- a/Engine/Plugins/Runtime/ReplicationSystemTestPlugin/Source/Private/Tests/ReplicationState/TestReplicationStateDescriptorBuilder.cpp
+++ b/Engine/Plugins/Runtime/ReplicationSystemTestPlugin/Source/Private/Tests/ReplicationState/TestReplicationStateDescriptorBuilder.cpp
@@ -562,7 +562,7 @@ static FReplicationStateDescriptor TestClassWithInheritanceNoSuperReplictionStat
 	12, // ExternalSize (rounded up since it also contains statemask)
 	4,	// InternalSize
 	4,	// ExternalAlignment
-	4,	// InternalAlignment
+	2,	// InternalAlignment
 	2,	// MemberCount
 	0,	// FunctionCount
 	0,	// TagCount;
diff --git a/Engine/Source/Runtime/Experimental/Iris/Core/Private/Iris/ReplicationState/ReplicationStateDescriptorBuilder.cpp b/Engine/Source/Runtime/Experimental/Iris/Core/Private/Iris/ReplicationState/ReplicationStateDescriptorBuilder.cpp
--- a/Engine/Source/Runtime/Experimental/Iris/Core/Private/Iris/ReplicationState/ReplicationStateDescriptorBuilder.cpp
+++ b/Engine/Source/Runtime/Experimental/Iris/Core/Private/Iris/ReplicationState/ReplicationStateDescriptorBuilder.cpp
@@ -168,7 +168,7 @@ public:
 
 private:
 
-	static const SIZE_T MinAlignment = 4;
+	static const SIZE_T MinAlignment = 1;
 
 	using FSizeAndAlignment = FMemoryLayoutUtil::FSizeAndAlignment;
 	using FOffsetAndSize = FMemoryLayoutUtil::FOffsetAndSize;
