Fix for network tests being skipped when running ReplicationSystemTest after CL 22782511.

diff --git a/Engine/Plugins/Runtime/ReplicationSystemTestPlugin/Source/Private/NetworkAutomationTest.cpp b/Engine/Plugins/Runtime/ReplicationSystemTestPlugin/Source/Private/NetworkAutomationTest.cpp
--- a/Engine/Plugins/Runtime/ReplicationSystemTestPlugin/Source/Private/NetworkAutomationTest.cpp
+++ b/Engine/Plugins/Runtime/ReplicationSystemTestPlugin/Source/Private/NetworkAutomationTest.cpp
@@ -96,7 +96,7 @@ FNetworkAutomationTestWrapper::FNetworkAutomationTestWrapper(FNetworkAutomationT
 
 uint32 FNetworkAutomationTestWrapper::GetTestFlags() const
 {
-	return EAutomationTestFlags::EngineFilter | EAutomationTestFlags::ClientContext | EAutomationTestFlags::ServerContext;
+	return EAutomationTestFlags::EngineFilter | EAutomationTestFlags::ClientContext | EAutomationTestFlags::ServerContext | EAutomationTestFlags::CommandletContext;
 }
 
 uint32 FNetworkAutomationTestWrapper::GetRequiredDeviceNum() const
