Iris - Added config to FreeBufferScope so that serializers requiring it in their FreeDynamicState functions will work as intended.

diff --git a/Engine/Plugins/Runtime/ReplicationSystemTestPlugin/Source/Private/Tests/Serialization/TestNetSerializerFixture.cpp b/Engine/Plugins/Runtime/ReplicationSystemTestPlugin/Source/Private/Tests/Serialization/TestNetSerializerFixture.cpp
--- a/Engine/Plugins/Runtime/ReplicationSystemTestPlugin/Source/Private/Tests/Serialization/TestNetSerializerFixture.cpp
+++ b/Engine/Plugins/Runtime/ReplicationSystemTestPlugin/Source/Private/Tests/Serialization/TestNetSerializerFixture.cpp
@@ -9,10 +9,11 @@ namespace UE::Net
 class FTestNetSerializerFixture::FFreeBufferScope
 {
 public:
-	FFreeBufferScope(FTestNetSerializerFixture& InTestFixture)
+	FFreeBufferScope(FTestNetSerializerFixture& InTestFixture, const FNetSerializerConfig* InNetSerializerConfig)
 	: TestFixture(InTestFixture)
 	, BuffersToFree{}
 	, BuffersToFreeCount(0)
+	, NetSerializerConfig(InNetSerializerConfig)
 	{
 	}
 
@@ -33,9 +34,7 @@ public:
 		{
 			FNetFreeDynamicStateArgs Args;
 			Args.Version = TestFixture.SerializerVersionOverride;
-			// It's unexpected the config would be used for freeing state.
-			// If it crashes let's add a function parameter so we can assign the config.
-			Args.NetSerializerConfig = nullptr;
+			Args.NetSerializerConfig = NetSerializerConfig;
 			for (uint32 It = 0, EndIt = BuffersToFreeCount; It != EndIt; ++It)
 			{
 				Args.Source = BuffersToFree[It];
@@ -49,6 +48,7 @@ private:
 	FTestNetSerializerFixture& TestFixture;
 	NetSerializerValuePointer BuffersToFree[3];
 	uint32 BuffersToFreeCount;
+	const FNetSerializerConfig* NetSerializerConfig = nullptr;
 };
 
 FTestNetSerializerFixture::FTestNetSerializerFixture(const FNetSerializer& InSerializer)
@@ -93,7 +93,7 @@ bool FTestNetSerializerFixture::TestQuantize(const FNetSerializerConfig& Config,
 	// 2. Verify that dequantized value is equal to source
 	// 3. Verify that quantized value is equal to second quantized value.
 
-	FFreeBufferScope FreeScope(*this);
+	FFreeBufferScope FreeScope(*this, &Config);
 
 	FNetQuantizeArgs QuantizeArgs;
 	QuantizeArgs.Version = SerializerVersionOverride;
@@ -180,7 +180,7 @@ bool FTestNetSerializerFixture::TestSerialize(const FNetSerializerConfig& Config
 
 	// Verification is done using deserialize followed by a quantized or dequantized compare.
 
-	FFreeBufferScope FreeScope(*this);
+	FFreeBufferScope FreeScope(*this, &Config);
 
 	FNetQuantizeArgs QuantizeArgs;
 	QuantizeArgs.Version = SerializerVersionOverride;
@@ -300,7 +300,7 @@ bool FTestNetSerializerFixture::TestSerializeDelta(const FNetSerializerConfig& C
 	// After that SerializeDelta can be called.
 	// Verification is done by verifying that the deserialized delta value matches the quantized source value.
 
-	FFreeBufferScope FreeScope(*this);
+	FFreeBufferScope FreeScope(*this, &Config);
 
 	FNetQuantizeArgs QuantizeValueArgs;
 	FNetQuantizeArgs QuantizePrevValueArgs;
@@ -388,7 +388,7 @@ bool FTestNetSerializerFixture::TestIsEqual(const FNetSerializerConfig& Config,
 	const uint32 PrepareTestFlags = (bTaintBuffersBeforeTest && bQuantizeValues ? (TaintQuantizedBuffer0 | TaintQuantizedBuffer1) : 0U);
 	PrepareTest(PrepareTestFlags);
 
-	FFreeBufferScope FreeScope(*this);
+	FFreeBufferScope FreeScope(*this, &Config);
 
 	FNetIsEqualArgs IsEqualArgs;
 	IsEqualArgs.Version = SerializerVersionOverride;
@@ -453,7 +453,7 @@ bool FTestNetSerializerFixture::TestCloneDynamicState(const FNetSerializerConfig
 	constexpr uint32 PrepareTestFlags = TaintQuantizedBuffer1;
 	PrepareTest(PrepareTestFlags);
 
-	FFreeBufferScope FreeScope(*this);
+	FFreeBufferScope FreeScope(*this, &Config);
 
 	FNetQuantizeArgs QuantizeArgs;
 	QuantizeArgs.Version = SerializerVersionOverride;
@@ -531,7 +531,7 @@ bool FTestNetSerializerFixture::Serialize(const FNetSerializerConfig& Config, Ne
 	const uint32 PrepareTestFlags = (bTaintBuffersBeforeTest ? (TaintBitStreamBuffer | TaintQuantizedBuffer0) : 0U);
 	PrepareTest(PrepareTestFlags);
 
-	FFreeBufferScope FreeScope(*this);
+	FFreeBufferScope FreeScope(*this, &Config);
 
 	FNetQuantizeArgs QuantizeArgs;
 	QuantizeArgs.Version = SerializerVersionOverride;
