Iris - Implemented deltacompression for PolymorphicStructNetSerializer

diff --git a/Engine/Plugins/Runtime/ReplicationSystemTestPlugin/Source/Private/Tests/Serialization/TestPolymorphicStructNetSerializer.cpp b/Engine/Plugins/Runtime/ReplicationSystemTestPlugin/Source/Private/Tests/Serialization/TestPolymorphicStructNetSerializer.cpp
--- a/Engine/Plugins/Runtime/ReplicationSystemTestPlugin/Source/Private/Tests/Serialization/TestPolymorphicStructNetSerializer.cpp
+++ b/Engine/Plugins/Runtime/ReplicationSystemTestPlugin/Source/Private/Tests/Serialization/TestPolymorphicStructNetSerializer.cpp
@@ -16,6 +16,11 @@
 namespace UE::Net
 {
 
+FTestMessage& operator<<(FTestMessage& Message, const FExamplePolymorphicArrayStruct& Value)
+{
+	return Message << "FExamplePolymorphicArrayStruct of size " << Value.Num();
+}
+
 struct FTestPolymorphicArrayStructNetSerializer : public TPolymorphicArrayStructNetSerializerImpl<FExamplePolymorphicArrayStruct, FExamplePolymorphicStructBase, FExamplePolymorphicArrayStruct::GetArray, FExamplePolymorphicArrayStruct::SetArrayNum>
 {
 	typedef TPolymorphicArrayStructNetSerializerImpl<FExamplePolymorphicArrayStruct, FExamplePolymorphicStructBase, FExamplePolymorphicArrayStruct::GetArray, FExamplePolymorphicArrayStruct::SetArrayNum> InternalNetSerializerType;
@@ -50,6 +55,11 @@ UE_NET_IMPLEMENT_SERIALIZER(FTestPolymorphicStructNetSerializer);
 namespace UE::Net::Private
 {
 
+FTestMessage& PrintPolymorphicArrayStructNetSerializerConfig(FTestMessage& Message, const FNetSerializerConfig& InConfig)
+{
+	return Message;
+}
+
 class FTestPolymorphicArrayStructNetSerializerFixture : public FReplicationSystemServerClientTestFixture
 {
 public:
@@ -67,6 +77,7 @@ protected:
 	void SerializeDelta();
 	void DeserializeDelta();
 	void Quantize();
+	void QuantizeTwoStates();
 	void Dequantize();
 	bool IsEqual(bool bQuantized);
 	void Clone();
@@ -91,20 +102,42 @@ protected:
 	FExamplePolymorphicStructC* StructC;
 	FExamplePolymorphicStructD* StructD;
 
-	alignas(8) uint8 QuantizedBuffer[2048];
-	alignas(8) uint8 ClonedQuantizedBuffer[2048];
+	alignas(8) uint8 QuantizedBuffer[2][2048];
+	alignas(8) uint8 ClonedQuantizedBuffer[2][2048];
 	alignas(8) uint8 BitStreamBuffer[2048];
 
-	bool bHasQuantizedState;
-	bool bHasClonedQuantizedState;
+	bool bHasQuantizedState = false;
+	bool bHasClonedQuantizedState = false;
+
+	uint32 QuantizedStateCount = 0;
+	uint32 ClonedQuantizedStateCount = 0;
 
 	FNetBitStreamWriter Writer;
 	FNetBitStreamReader Reader;
 
-	static bool bHasRegisteredStructDNetSerializer;
+	inline static bool bHasRegisteredStructDNetSerializer = false;
 };
 
-bool FTestPolymorphicArrayStructNetSerializerFixture::bHasRegisteredStructDNetSerializer = false;
+class FTestPolymorphicArrayStructNetSerializerDeltaSerializationFixture : public TTestNetSerializerFixture<PrintPolymorphicArrayStructNetSerializerConfig, FExamplePolymorphicArrayStruct>
+{
+protected:
+	using Super = TTestNetSerializerFixture<PrintPolymorphicArrayStructNetSerializerConfig, FExamplePolymorphicArrayStruct>;
+
+	FTestPolymorphicArrayStructNetSerializerDeltaSerializationFixture() : Super(UE_NET_GET_SERIALIZER(FTestPolymorphicArrayStructNetSerializer)) {}
+
+	virtual void SetUp() override;
+	virtual void TearDown() override;
+
+	void SetArbitraryState(FExamplePolymorphicArrayStruct& Target);
+
+protected:
+	FTestPolymorphicArrayStructNetSerializerConfig Config;
+
+	FInternalNetSerializationContext InternalNetSerializationContext;
+
+	FDataStreamTestUtil DataStreamUtil;
+	FReplicationSystemTestServer* Server = nullptr;
+};
 
 }
 
@@ -237,29 +270,36 @@ UE_NET_TEST_FIXTURE(FTestPolymorphicArrayStructNetSerializerFixture, TestDequant
 
 UE_NET_TEST_FIXTURE(FTestPolymorphicArrayStructNetSerializerFixture, TestSerializeDelta)
 {
-	Quantize();
+	QuantizeTwoStates();
 	SerializeDelta();
 }
 
 UE_NET_TEST_FIXTURE(FTestPolymorphicArrayStructNetSerializerFixture, TestDeserializeDelta)
 {
-	Quantize();
-	Serialize();
-	FreeQuantizedState();
+	QuantizeTwoStates();
+	SerializeDelta();
+	DeserializeDelta();
+}
+
+UE_NET_TEST_FIXTURE(FTestPolymorphicArrayStructNetSerializerFixture, TestDequantizeDeltaSerializedState)
+{
+	QuantizeTwoStates();
+	SerializeDelta();
 	DeserializeDelta();
+	Dequantize();
+	ValidateExpectedState(PolymorphicStructInstance1);
 }
 
 UE_NET_TEST_FIXTURE(FTestPolymorphicArrayStructNetSerializerFixture, TestCollectReferencesNoRef)
 {
 	Quantize();
-	UE_NET_ASSERT_TRUE(bHasQuantizedState);
 
 	FNetReferenceCollector Collector;
 
 	FNetCollectReferencesArgs Args = {};
 	Args.Version = PolymorphicNetSerializer->Version;
 	Args.NetSerializerConfig = NetSerializerConfigParam(&PolymorphicNetSerializerConfig);
-	Args.Source = NetSerializerValuePointer(&QuantizedBuffer);
+	Args.Source = NetSerializerValuePointer(&QuantizedBuffer[0]);
 	Args.Collector = NetSerializerValuePointer(&Collector);
 
 	PolymorphicNetSerializer->CollectNetReferences(NetSerializationContext, Args);	
@@ -270,14 +310,13 @@ UE_NET_TEST_FIXTURE(FTestPolymorphicArrayStructNetSerializerFixture, TestCollect
 UE_NET_TEST_FIXTURE(FTestPolymorphicArrayStructNetSerializerFixture, TestCollectReferencesRef)
 {
 	Quantize();
-	UE_NET_ASSERT_TRUE(bHasQuantizedState);
 
 	FNetReferenceCollector Collector(ENetReferenceCollectorTraits::IncludeInvalidReferences);
 
 	FNetCollectReferencesArgs Args = {};
 	Args.Version = PolymorphicNetSerializer->Version;
 	Args.NetSerializerConfig = NetSerializerConfigParam(&PolymorphicNetSerializerConfig);
-	Args.Source = NetSerializerValuePointer(&QuantizedBuffer);
+	Args.Source = NetSerializerValuePointer(&QuantizedBuffer[0]);
 	Args.Collector = NetSerializerValuePointer(&Collector);
 
 	PolymorphicNetSerializer->CollectNetReferences(NetSerializationContext, Args);	
@@ -316,6 +355,52 @@ UE_NET_TEST_FIXTURE(FTestPolymorphicArrayStructNetSerializerFixture, TestValidat
 	Validate();
 }
 
+UE_NET_TEST_FIXTURE(FTestPolymorphicArrayStructNetSerializerDeltaSerializationFixture, TestDeltaSerialization)
+{
+	TArray<FExamplePolymorphicArrayStruct> Values;
+
+	{
+		// Add empty value
+		{
+			FExamplePolymorphicArrayStruct EmptyArray;
+			Values.Emplace(MoveTemp(EmptyArray));
+		}
+
+		// Add arbitrary array
+		{
+			FExamplePolymorphicArrayStruct ArbitraryArray;
+			SetArbitraryState(ArbitraryArray);
+			Values.Emplace(ArbitraryArray);
+		}
+
+		// Add "arbitrary array" without first value
+		{
+			FExamplePolymorphicArrayStruct ArbitraryArrayWithoutFirstElement;
+			SetArbitraryState(ArbitraryArrayWithoutFirstElement);
+			if (ArbitraryArrayWithoutFirstElement.Num() > 0)
+			{
+				ArbitraryArrayWithoutFirstElement.RemoveAt(0);
+			}
+			Values.Emplace(MoveTemp(ArbitraryArrayWithoutFirstElement));
+		}
+
+		// Add "arbitrary array" without last value
+		{
+			FExamplePolymorphicArrayStruct ArbitraryArrayWithoutLastElement;
+			SetArbitraryState(ArbitraryArrayWithoutLastElement);
+			if (ArbitraryArrayWithoutLastElement.Num() > 0)
+			{
+				ArbitraryArrayWithoutLastElement.RemoveAt(ArbitraryArrayWithoutLastElement.Num() - 1);
+			}
+			Values.Emplace(ArbitraryArrayWithoutLastElement);
+		}
+	}
+
+	Super::TestSerializeDelta(Values.GetData(), Values.Num(), Config);
+}
+
+
+// FTestPolymorphicArrayStructNetSerializerFixture implementation
 void FTestPolymorphicArrayStructNetSerializerFixture::SetExpectedState(FExamplePolymorphicArrayStruct& Target)
 {
 	StructA = new FExamplePolymorphicStructA;
@@ -328,7 +413,7 @@ void FTestPolymorphicArrayStructNetSerializerFixture::SetExpectedState(FExampleP
 	StructC->SomeBool = false;
 
 	StructD = new FExamplePolymorphicStructD;
-	StructD->SomeValue = 0xBEEEEEEF;
+	StructD->SomeValue = 0x7E577E57;
 	
 	// Add some data, we need some structs to send as well
 	Target.Add(StructA);
@@ -356,7 +441,7 @@ void FTestPolymorphicArrayStructNetSerializerFixture::SetUp()
 	InternalNetSerializationContext = MoveTemp(TempInternalNetSerializationContext);
 	NetSerializationContext.SetInternalContext(&InternalNetSerializationContext);
 
-	FMemory::Memzero(QuantizedBuffer, 0);
+	FMemory::Memzero(QuantizedBuffer, sizeof(QuantizedBuffer));
 
 	SetExpectedState(PolymorphicStructInstance0);
 
@@ -424,7 +509,7 @@ void FTestPolymorphicArrayStructNetSerializerFixture::Serialize()
 		FNetSerializeArgs Args = {};
 		Args.Version = PolymorphicNetSerializer->Version;
 		Args.NetSerializerConfig = NetSerializerConfigParam(&PolymorphicNetSerializerConfig);
-		Args.Source = NetSerializerValuePointer(&QuantizedBuffer);
+		Args.Source = NetSerializerValuePointer(&QuantizedBuffer[0]);
 		PolymorphicNetSerializer->Serialize(Context, Args);
 
 		Writer.CommitWrites();
@@ -448,7 +533,7 @@ void FTestPolymorphicArrayStructNetSerializerFixture::Deserialize()
 	FNetDeserializeArgs Args = {};
 	Args.Version = PolymorphicNetSerializer->Version;
 	Args.NetSerializerConfig = NetSerializerConfigParam(&PolymorphicNetSerializerConfig);
-	Args.Target = NetSerializerValuePointer(&QuantizedBuffer);
+	Args.Target = NetSerializerValuePointer(&QuantizedBuffer[0]);
 	PolymorphicNetSerializer->Deserialize(Context, Args);
 
 	bHasQuantizedState = true;
@@ -457,7 +542,7 @@ void FTestPolymorphicArrayStructNetSerializerFixture::Deserialize()
 void FTestPolymorphicArrayStructNetSerializerFixture::SerializeDelta()
 {
 	// Check pre-conditions
-	UE_NET_ASSERT_TRUE(bHasQuantizedState);
+	UE_NET_ASSERT_TRUE(bHasQuantizedState && QuantizedStateCount == 2U);
 	
 	// Serialize data
 	{
@@ -468,7 +553,8 @@ void FTestPolymorphicArrayStructNetSerializerFixture::SerializeDelta()
 		FNetSerializeDeltaArgs Args = {};
 		Args.Version = PolymorphicNetSerializer->Version;
 		Args.NetSerializerConfig = NetSerializerConfigParam(&PolymorphicNetSerializerConfig);
-		Args.Source = NetSerializerValuePointer(&QuantizedBuffer);
+		Args.Source = NetSerializerValuePointer(&QuantizedBuffer[0]);
+		Args.Prev = NetSerializerValuePointer(&QuantizedBuffer[1]);
 
 		PolymorphicNetSerializer->SerializeDelta(Context, Args);
 
@@ -482,7 +568,6 @@ void FTestPolymorphicArrayStructNetSerializerFixture::SerializeDelta()
 void FTestPolymorphicArrayStructNetSerializerFixture::DeserializeDelta()
 {
 	// Check pre-conditions
-	UE_NET_ASSERT_FALSE(bHasQuantizedState);
 	UE_NET_ASSERT_TRUE(Writer.GetPosBytes() > 0U);
 	
 	Reader.InitBits(BitStreamBuffer, Writer.GetPosBits());
@@ -493,10 +578,12 @@ void FTestPolymorphicArrayStructNetSerializerFixture::DeserializeDelta()
 	FNetDeserializeDeltaArgs Args = {};
 	Args.Version = PolymorphicNetSerializer->Version;
 	Args.NetSerializerConfig = NetSerializerConfigParam(&PolymorphicNetSerializerConfig);
-	Args.Target = NetSerializerValuePointer(&QuantizedBuffer);
+	Args.Target = NetSerializerValuePointer(&QuantizedBuffer[0]);
+	Args.Prev = NetSerializerValuePointer(&QuantizedBuffer[1]);
 	PolymorphicNetSerializer->DeserializeDelta(Context, Args);
 
 	bHasQuantizedState = true;
+	QuantizedStateCount = 1;
 }
 
 void FTestPolymorphicArrayStructNetSerializerFixture::Quantize()
@@ -504,23 +591,43 @@ void FTestPolymorphicArrayStructNetSerializerFixture::Quantize()
 	FNetQuantizeArgs Args = {};
 	Args.Version = PolymorphicNetSerializer->Version;
 	Args.NetSerializerConfig = NetSerializerConfigParam(&PolymorphicNetSerializerConfig);
-	Args.Target = NetSerializerValuePointer(&QuantizedBuffer);
+	Args.Target = NetSerializerValuePointer(&QuantizedBuffer[0]);
 	Args.Source = NetSerializerValuePointer(&PolymorphicStructInstance0);
 	PolymorphicNetSerializer->Quantize(NetSerializationContext, Args);
 
 	bHasQuantizedState = true;
+	QuantizedStateCount = 1;
+}
+
+void FTestPolymorphicArrayStructNetSerializerFixture::QuantizeTwoStates()
+{
+	Quantize();
+
+	FNetQuantizeArgs Args = {};
+	Args.Version = PolymorphicNetSerializer->Version;
+	Args.NetSerializerConfig = NetSerializerConfigParam(&PolymorphicNetSerializerConfig);
+	Args.Target = NetSerializerValuePointer(&QuantizedBuffer[1]);
+	Args.Source = NetSerializerValuePointer(&PolymorphicStructInstance1);
+	PolymorphicNetSerializer->Quantize(NetSerializationContext, Args);
+
+	bHasQuantizedState = true;
+	QuantizedStateCount = 2;
 }
 
+
 void FTestPolymorphicArrayStructNetSerializerFixture::Clone()
 {
 	// Check pre-conditions
 	UE_NET_ASSERT_TRUE(bHasQuantizedState);
 
 	FNetCloneDynamicStateArgs Args = {};
-	Args.Source = NetSerializerValuePointer(QuantizedBuffer);
-	Args.Target = NetSerializerValuePointer(ClonedQuantizedBuffer);
+	Args.Source = NetSerializerValuePointer(&QuantizedBuffer[0]);
+	Args.Target = NetSerializerValuePointer(&ClonedQuantizedBuffer[0]);
 	Args.NetSerializerConfig = NetSerializerConfigParam(&PolymorphicNetSerializerConfig);
 	PolymorphicNetSerializer->CloneDynamicState(NetSerializationContext, Args);
+
+	bHasClonedQuantizedState = true;
+	ClonedQuantizedStateCount = 1;
 }
 
 void FTestPolymorphicArrayStructNetSerializerFixture::FreeQuantizedState()
@@ -531,15 +638,21 @@ void FTestPolymorphicArrayStructNetSerializerFixture::FreeQuantizedState()
 
 	if (bHasQuantizedState)
 	{
-		Args.Source = NetSerializerValuePointer(&QuantizedBuffer);
-		PolymorphicNetSerializer->FreeDynamicState(NetSerializationContext, Args);
+		for (uint32 StateIt = 0, StateEndIt = QuantizedStateCount; StateIt != StateEndIt; ++StateIt)
+		{
+			Args.Source = NetSerializerValuePointer(&QuantizedBuffer[StateIt]);
+			PolymorphicNetSerializer->FreeDynamicState(NetSerializationContext, Args);
+		}
 		bHasQuantizedState = false;
 	}
 		
 	if (bHasClonedQuantizedState)
 	{
-		Args.Source = NetSerializerValuePointer(&ClonedQuantizedBuffer);
-		PolymorphicNetSerializer->FreeDynamicState(NetSerializationContext, Args);
+		for (uint32 StateIt = 0, StateEndIt = ClonedQuantizedStateCount; StateIt != StateEndIt; ++StateIt)
+		{
+			Args.Source = NetSerializerValuePointer(&ClonedQuantizedBuffer[StateIt]);
+			PolymorphicNetSerializer->FreeDynamicState(NetSerializationContext, Args);
+		}
 		bHasClonedQuantizedState = false;
 	}
 }
@@ -551,13 +664,20 @@ void FTestPolymorphicArrayStructNetSerializerFixture::Dequantize()
 	FNetDequantizeArgs Args = {};
 	Args.Version = PolymorphicNetSerializer->Version;
 	Args.NetSerializerConfig = NetSerializerConfigParam(&PolymorphicNetSerializerConfig);
-	Args.Source = NetSerializerValuePointer(&QuantizedBuffer);
+	Args.Source = NetSerializerValuePointer(&QuantizedBuffer[0]);
 	Args.Target = NetSerializerValuePointer(&PolymorphicStructInstance1);
 	PolymorphicNetSerializer->Dequantize(NetSerializationContext, Args);
 }
 
 bool FTestPolymorphicArrayStructNetSerializerFixture::IsEqual(bool bQuantized)
 {
+	if (bQuantized)
+	{
+		// Trick to return false if assertion fails.
+		UE_NET_ASSERT_TRUE(bHasQuantizedState), false;
+		UE_NET_ASSERT_TRUE(bHasClonedQuantizedState), false;
+	}
+
 	FNetSerializationContext Context;
 	FNetIsEqualArgs Args = {};
 	Args.Version = PolymorphicNetSerializer->Version;
@@ -595,19 +715,20 @@ UE_NET_TEST_FIXTURE(FTestPolymorphicArrayStructNetSerializerFixture, TestPolymor
 
 	// Replicate
 	Server->PreSendUpdate();
-	Server->SendAndDeliverTo(Client, true);
+	Server->SendAndDeliverTo(Client, DeliverPacket);
 	Server->PostSendUpdate();
 	
 	auto ClientObject = Client->GetObjectAs<UTestPolymorphicStructNetSerializer_TestObject>(ServerObject->NetHandle);
-	UE_NET_ASSERT_TRUE(ClientObject != nullptr);
-	UE_NET_ASSERT_TRUE(ClientObject->PolyStruct.GetAs<FExamplePolymorphicStructB>().SomeFloat == 12.0f);
+	UE_NET_ASSERT_NE(ClientObject, nullptr);
+	CA_ASSUME(ClientObject != nullptr);
+	UE_NET_ASSERT_EQ(ClientObject->PolyStruct.GetAs<FExamplePolymorphicStructB>().SomeFloat, 12.0f);
 
 	// Modify
 	ServerObject->PolyStruct.GetAs<FExamplePolymorphicStructB>().SomeFloat += 1.0f;
 
 	// Replicate
 	Server->PreSendUpdate();
-	Server->SendAndDeliverTo(Client, true);
+	Server->SendAndDeliverTo(Client, DeliverPacket);
 	Server->PostSendUpdate();
 
 	// Verify that we detected the modification
@@ -619,7 +740,7 @@ UE_NET_TEST_FIXTURE(FTestPolymorphicArrayStructNetSerializerFixture, TestPolymor
 
 	// Replicate
 	Server->PreSendUpdate();
-	Server->SendAndDeliverTo(Client, true);
+	Server->SendAndDeliverTo(Client, DeliverPacket);
 	Server->PostSendUpdate();
 
 	UE_NET_ASSERT_TRUE(ClientObject->PolyStruct.GetAs<FExamplePolymorphicStructD>().SomeValue == ServerObject->PolyStruct.GetAs<FExamplePolymorphicStructD>().SomeValue);	
@@ -654,11 +775,12 @@ UE_NET_TEST_FIXTURE(FTestPolymorphicArrayStructNetSerializerFixture, TestPolymor
 	
 	// Replicate
 	Server->PreSendUpdate();
-	Server->SendAndDeliverTo(Client, true);
+	Server->SendAndDeliverTo(Client, DeliverPacket);
 	Server->PostSendUpdate();
 	
 	auto ClientObject = Client->GetObjectAs<UTestPolymorphicStructNetSerializer_TestObject>(ServerObject->NetHandle);
-	UE_NET_ASSERT_TRUE(ClientObject != nullptr);
+	UE_NET_ASSERT_NE(ClientObject, nullptr);
+	CA_ASSUME(ClientObject != nullptr);
 	UE_NET_ASSERT_EQ(ClientObject->PolyStructFastArray.GetItemArray().Num(), ServerObject->PolyStructFastArray.GetItemArray().Num());
 	UE_NET_ASSERT_EQ(ClientObject->PolyStructFastArray.GetItemArray()[3].PolyStruct.GetAs<FExamplePolymorphicStructD>().SomeValue, 100U);
 
@@ -667,7 +789,7 @@ UE_NET_TEST_FIXTURE(FTestPolymorphicArrayStructNetSerializerFixture, TestPolymor
 
 	// Replicate
 	Server->PreSendUpdate();
-	Server->SendAndDeliverTo(Client, true);
+	Server->SendAndDeliverTo(Client, DeliverPacket);
 	Server->PostSendUpdate();
 
 	// Verified that the client got the modified value
@@ -679,13 +801,65 @@ UE_NET_TEST_FIXTURE(FTestPolymorphicArrayStructNetSerializerFixture, TestPolymor
 
 	// Replicate
 	Server->PreSendUpdate();
-	Server->SendAndDeliverTo(Client, true);
+	Server->SendAndDeliverTo(Client, DeliverPacket);
 	Server->PostSendUpdate();
 
 	// Verified that the client got the modified value
 	UE_NET_ASSERT_EQ(ClientObject->PolyStructFastArray.GetItemArray()[2].PolyStruct.GetAs<FExamplePolymorphicStructD>().SomeValue, 1U);
 }
 
+// FTestPolymorphicArrayStructNetSerializerDeltaSerializationFixture implementation
+void FTestPolymorphicArrayStructNetSerializerDeltaSerializationFixture::SetUp()
+{
+	Super::SetUp();
+
+	Config.RegisteredTypes.InitForType(FExamplePolymorphicStructBase::StaticStruct());
+
+	Server = new FReplicationSystemTestServer(GetName());
+
+	// Init NetSerializationContext
+	FReplicationSystemInternal* ReplicationSystemInternal = Server->GetReplicationSystem()->GetReplicationSystemInternal();
+
+	FInternalNetSerializationContext TempInternalNetSerializationContext;
+	FInternalNetSerializationContext::FInitParameters TempInternalNetSerializationContextInitParams;
+	TempInternalNetSerializationContextInitParams.ReplicationSystem = Server->ReplicationSystem;
+	TempInternalNetSerializationContextInitParams.ObjectResolveContext.RemoteNetTokenStoreState = ReplicationSystemInternal->GetNetTokenStore().GetLocalNetTokenStoreState();
+	TempInternalNetSerializationContext.Init(TempInternalNetSerializationContextInitParams);
+
+	InternalNetSerializationContext = MoveTemp(TempInternalNetSerializationContext);
+	Context.SetInternalContext(&InternalNetSerializationContext);
+}
+
+void FTestPolymorphicArrayStructNetSerializerDeltaSerializationFixture::TearDown()
+{
+	delete Server;
+	Server = nullptr;
+
+	Super::TearDown();
+}
+
+void FTestPolymorphicArrayStructNetSerializerDeltaSerializationFixture::SetArbitraryState(FExamplePolymorphicArrayStruct& Target)
+{
+	FExamplePolymorphicStructA* StructA = new FExamplePolymorphicStructA;
+	StructA->SomeInt = -29837492;
+	
+	FExamplePolymorphicStructB* StructB = new FExamplePolymorphicStructB;
+	StructB->SomeFloat = 17.0f;
+
+	FExamplePolymorphicStructC* StructC = new FExamplePolymorphicStructC;
+	StructC->SomeObjectRef = nullptr;
+	StructC->SomeBool = false;
+
+	FExamplePolymorphicStructD* StructD = new FExamplePolymorphicStructD;
+	StructD->SomeValue = 0xBEEEEEEF;
+
+	// Add values to target
+	Target.Add(StructA);
+	Target.Add(StructB);
+	Target.Add(StructC);
+	Target.Add(StructD);
+}
+
 }
 
 void FExamplePolymorphicStructFastArrayItem::PostReplicatedAdd(const struct FExamplePolymorphicStructFastArraySerializer& InArraySerializer)
diff --git a/Engine/Plugins/Runtime/ReplicationSystemTestPlugin/Source/Private/Tests/Serialization/TestPolymorphicStructNetSerializer.h b/Engine/Plugins/Runtime/ReplicationSystemTestPlugin/Source/Private/Tests/Serialization/TestPolymorphicStructNetSerializer.h
--- a/Engine/Plugins/Runtime/ReplicationSystemTestPlugin/Source/Private/Tests/Serialization/TestPolymorphicStructNetSerializer.h
+++ b/Engine/Plugins/Runtime/ReplicationSystemTestPlugin/Source/Private/Tests/Serialization/TestPolymorphicStructNetSerializer.h
@@ -14,7 +14,7 @@
 USTRUCT()
 struct FExamplePolymorphicStructBase
 {
-	GENERATED_USTRUCT_BODY()
+	GENERATED_BODY()
 
 	virtual ~FExamplePolymorphicStructBase() { }
 
@@ -77,6 +77,14 @@ public:
 		return IsValid(Index) ? Data[Index].Get() : nullptr;
 	}
 
+	void RemoveAt(int32 Index)
+	{
+		if (Index > 0 && Index < Num())
+		{
+			Data.RemoveAt(Index);
+		}
+	}
+
 	/** Adds a new FExamplePolymorphicArrayStruct data to FExamplePolymorphicArrayStruct, it must have been created with new */
 	void Add(FExamplePolymorphicStructBase* DataPtr)
 	{
@@ -97,13 +105,14 @@ public:
 		{
 			return false;
 		}
-		for (int32 i = 0; i < Data.Num(); ++i)
+
+		for (int32 It = 0, EndIt = Data.Num(); It < EndIt; ++It)
 		{
-			if (Data[i].IsValid() != Other.Data[i].IsValid())
+			if (Data[It].IsValid() != Other.Data[It].IsValid())
 			{
 				return false;
 			}
-			if (Data[i].Get() != Other.Data[i].Get())
+			if (Data[It].Get() != Other.Data[It].Get())
 			{
 				return false;
 			}
@@ -153,7 +162,7 @@ UE_NET_DECLARE_SERIALIZER(FTestPolymorphicArrayStructNetSerializer, REPLICATIONS
 USTRUCT()
 struct FTestPolymorphicArrayStructNetSerializerConfig : public FPolymorphicArrayStructNetSerializerConfig
 {
-	GENERATED_USTRUCT_BODY()
+	GENERATED_BODY()
 };
 
 //////////////////////////////////////////////////////////////////////////
@@ -187,6 +196,7 @@ struct FExamplePolymorphicStruct
 			const FExamplePolymorphicStructBase* SrcPolymorphicStruct = Other.Data.IsValid() ? Other.Data.Get() : nullptr;
 			if (SrcPolymorphicStruct)
 			{
+				CA_ASSUME(SrcScriptStruct != nullptr);
 				FExamplePolymorphicStructBase* const DstPolymorphicStruct = Data.IsValid() ? Data.Get() : nullptr;
 				if (DstPolymorphicStruct == SrcPolymorphicStruct)
 				{
@@ -294,7 +304,7 @@ UE_NET_DECLARE_SERIALIZER(FTestPolymorphicStructNetSerializer, REPLICATIONSYSTEM
 USTRUCT()
 struct FTestPolymorphicStructNetSerializerConfig : public FPolymorphicStructNetSerializerConfig
 {
-	GENERATED_USTRUCT_BODY()
+	GENERATED_BODY()
 };
 
 //////////////////////////////////////////////////////////////////////////
@@ -371,7 +381,7 @@ private:
 USTRUCT()
 struct FExamplePolymorphicStructA : public FExamplePolymorphicStructBase
 {
-	GENERATED_USTRUCT_BODY()
+	GENERATED_BODY()
 
 	UPROPERTY()
 	int32 SomeInt;
@@ -386,10 +396,10 @@ struct FExamplePolymorphicStructA : public FExamplePolymorphicStructBase
 USTRUCT()
 struct FExamplePolymorphicStructB : public FExamplePolymorphicStructBase
 {
-	GENERATED_USTRUCT_BODY()
+	GENERATED_BODY()
 
 	UPROPERTY()
-	int32 SomeFloat;
+	float SomeFloat;
 
 	/** Returns the serialization data, must always be overridden */
 	virtual UScriptStruct* GetScriptStruct() const override
@@ -401,7 +411,7 @@ struct FExamplePolymorphicStructB : public FExamplePolymorphicStructBase
 USTRUCT()
 struct FExamplePolymorphicStructC : public FExamplePolymorphicStructBase
 {
-	GENERATED_USTRUCT_BODY()
+	GENERATED_BODY()
 
 	UPROPERTY()
 	bool SomeBool;
@@ -420,7 +430,7 @@ struct FExamplePolymorphicStructC : public FExamplePolymorphicStructBase
 USTRUCT()
 struct FExamplePolymorphicStructD : public FExamplePolymorphicStructBase
 {
-	GENERATED_USTRUCT_BODY()
+	GENERATED_BODY()
 	
 	// Since we do not have any properties, we need a custom compare operator and mark the struct with WithIdenticalViaEquality Trait
 	// Alternatively we could mark SomeValue as a non-replicated property (UProperty(NotReplicated))
diff --git a/Engine/Source/Runtime/Experimental/Iris/Core/Public/Iris/Serialization/PolymorphicNetSerializer.h b/Engine/Source/Runtime/Experimental/Iris/Core/Public/Iris/Serialization/PolymorphicNetSerializer.h
--- a/Engine/Source/Runtime/Experimental/Iris/Core/Public/Iris/Serialization/PolymorphicNetSerializer.h
+++ b/Engine/Source/Runtime/Experimental/Iris/Core/Public/Iris/Serialization/PolymorphicNetSerializer.h
@@ -51,11 +51,9 @@ struct FPolymorphicStructNetSerializerConfig : public FNetSerializerConfig
 };
 
 USTRUCT()
-struct FPolymorphicArrayStructNetSerializerConfig : public FNetSerializerConfig
+struct FPolymorphicArrayStructNetSerializerConfig : public FPolymorphicStructNetSerializerConfig
 {
 	GENERATED_BODY()
-
-	UE::Net::FPolymorphicNetSerializerScriptStructCache RegisteredTypes;
 };
 
 namespace UE::Net
diff --git a/Engine/Source/Runtime/Experimental/Iris/Core/Public/Iris/Serialization/PolymorphicNetSerializerImpl.h b/Engine/Source/Runtime/Experimental/Iris/Core/Public/Iris/Serialization/PolymorphicNetSerializerImpl.h
--- a/Engine/Source/Runtime/Experimental/Iris/Core/Public/Iris/Serialization/PolymorphicNetSerializerImpl.h
+++ b/Engine/Source/Runtime/Experimental/Iris/Core/Public/Iris/Serialization/PolymorphicNetSerializerImpl.h
@@ -11,7 +11,7 @@
 #include "Iris/Serialization/NetBitStreamWriter.h"
 #include "Iris/ReplicationState/ReplicationStateDescriptor.h"
 #include "Iris/ReplicationSystem/ReplicationOperations.h"
-#include "UObject/ObjectMacros.h"
+#include "Net/Core/Trace/NetTrace.h"
 
 // NOTE: This file should not be included in a header
 
@@ -203,6 +203,8 @@ struct TPolymorphicArrayStructNetSerializerImpl : protected Private::FPolymorphi
 	}
 
 private:
+	using FItemNetSerializer = TPolymorphicStructNetSerializerImpl<ExternalSourceType, ExternalSourceArrayItemType, nullptr>;
+
 	// Allocate storage for the item array
 	static void InternalAllocateItemArray(FNetSerializationContext& Context, QuantizedType& Value, uint32 NumItems);
 
@@ -264,13 +266,85 @@ void TPolymorphicStructNetSerializerImpl<ExternalSourceType, ExternalSourceItemT
 template <typename ExternalSourceType, typename ExternalSourceItemType, TSharedPtr<ExternalSourceItemType>&(*GetItem)(ExternalSourceType&)>
 void TPolymorphicStructNetSerializerImpl<ExternalSourceType, ExternalSourceItemType, GetItem>::SerializeDelta(FNetSerializationContext& Context, const FNetSerializeDeltaArgs& Args)
 {
-	NetSerializeDeltaDefault<Serialize>(Context, Args);
+	FNetBitStreamWriter& Writer = *Context.GetBitStreamWriter();
+	const QuantizedType& Value = *reinterpret_cast<const QuantizedType*>(Args.Source);
+	const QuantizedType& PrevValue = *reinterpret_cast<const QuantizedType*>(Args.Prev);
+	const ConfigType& Config = *static_cast<const ConfigType*>(Args.NetSerializerConfig);
+
+	// If prev has the same type we can delta-compress
+	if (Writer.WriteBool(Value.TypeIndex == PrevValue.TypeIndex))
+	{
+		const FTypeInfo* TypeInfo = Config.RegisteredTypes.GetTypeInfo(Value.TypeIndex);
+		if (TypeInfo != nullptr)
+		{
+			UE::Net::FReplicationStateOperations::SerializeDelta(Context, static_cast<const uint8*>(Value.StructData), static_cast<const uint8*>(PrevValue.StructData), TypeInfo->Descriptor);
+		}
+	}
+	else
+	{
+		const FTypeInfo* TypeInfo = Config.RegisteredTypes.GetTypeInfo(Value.TypeIndex);
+		if (Writer.WriteBool(TypeInfo != nullptr))
+		{
+			CA_ASSUME(TypeInfo != nullptr);
+			Writer.WriteBits(Value.TypeIndex, FPolymorphicNetSerializerScriptStructCache::RegisteredTypeBits);
+			UE::Net::FReplicationStateOperations::Serialize(Context, static_cast<const uint8*>(Value.StructData), TypeInfo->Descriptor);
+		}
+	}
 }
 
 template <typename ExternalSourceType, typename ExternalSourceItemType, TSharedPtr<ExternalSourceItemType>&(*GetItem)(ExternalSourceType&)>
 void TPolymorphicStructNetSerializerImpl<ExternalSourceType, ExternalSourceItemType, GetItem>::DeserializeDelta(FNetSerializationContext& Context, const FNetDeserializeDeltaArgs& Args)
 {
-	NetDeserializeDeltaDefault<Deserialize>(Context, Args);
+	FNetBitStreamReader& Reader = *Context.GetBitStreamReader();
+	QuantizedType& Target = *reinterpret_cast<QuantizedType*>(Args.Target);
+	const QuantizedType& PrevValue = *reinterpret_cast<const QuantizedType*>(Args.Prev);
+	const ConfigType& Config = *static_cast<const ConfigType*>(Args.NetSerializerConfig);
+
+	InternalFreeItem(Context, Config, Target);
+
+	QuantizedType TempValue = {};
+
+	// If prev has the same type we can delta-compress
+	if (const bool bIsSameType = Reader.ReadBool())
+	{
+		const uint32 TypeIndex = PrevValue.TypeIndex;
+		if (const FTypeInfo* TypeInfo = Config.RegisteredTypes.GetTypeInfo(TypeIndex))
+		{
+			const FReplicationStateDescriptor* Descriptor = TypeInfo->Descriptor;
+
+			// Allocate storage and read struct data
+			TempValue.StructData = FPolymorphicStructNetSerializerInternal::Alloc(Context, Descriptor->InternalSize, Descriptor->InternalAlignment);
+			TempValue.TypeIndex = TypeIndex;
+
+			FMemory::Memzero(TempValue.StructData, Descriptor->InternalSize);
+			FReplicationStateOperations::DeserializeDelta(Context, static_cast<uint8*>(TempValue.StructData), static_cast<uint8*>(PrevValue.StructData), Descriptor);
+		}
+	}
+	else
+	{
+		if (const bool bIsValidType = Reader.ReadBool())
+		{
+			const uint32 TypeIndex = Reader.ReadBits(FPolymorphicNetSerializerScriptStructCache::RegisteredTypeBits);
+			if (const FTypeInfo* TypeInfo = Config.RegisteredTypes.GetTypeInfo(TypeIndex))
+			{
+				const FReplicationStateDescriptor* Descriptor = TypeInfo->Descriptor;
+
+				// Allocate storage and read struct data
+				TempValue.StructData = FPolymorphicStructNetSerializerInternal::Alloc(Context, Descriptor->InternalSize, Descriptor->InternalAlignment);
+				TempValue.TypeIndex = TypeIndex;
+
+				FMemory::Memzero(TempValue.StructData, Descriptor->InternalSize);
+				FReplicationStateOperations::Deserialize(Context, static_cast<uint8*>(TempValue.StructData), Descriptor);
+			}
+			else
+			{
+				Context.SetError(NetError_PolymorphicStructNetSerializer_InvalidStructType);
+				// Fall through to clear the target
+			}
+		}
+	}
+
+	Target = TempValue;
 }
 
 template <typename ExternalSourceType, typename ExternalSourceItemType, TSharedPtr<ExternalSourceItemType>&(*GetItem)(ExternalSourceType&)>
@@ -485,6 +559,8 @@ void TPolymorphicArrayStructNetSerializerImpl<ExternalSourceType, ExternalSource
 	Writer.WriteBits(SourceValue.NumItems, ArrayItemBits);	
 	for (uint32 It = 0, EndIt = SourceValue.NumItems; It != EndIt; ++It)
 	{
+		UE_NET_TRACE_SCOPE(Element, Writer, Context.GetTraceCollector(), ENetTraceVerbosity::VeryVerbose);
+
 		const FQuantizedItem& Item = SourceValue.Items[It];
 		const FTypeInfo* TypeInfo = Config.RegisteredTypes.GetTypeInfo(Item.TypeIndex);
 		if (Writer.WriteBool(TypeInfo != nullptr))
@@ -496,32 +572,6 @@ void TPolymorphicArrayStructNetSerializerImpl<ExternalSourceType, ExternalSource
 	}
 }
 
-template <typename ExternalSourceType, typename ExternalSourceArrayItemType, TArrayView<TSharedPtr<ExternalSourceArrayItemType>>(*GetArray)(ExternalSourceType& Source), void(*SetArrayNum)(ExternalSourceType& Source, SIZE_T Num)>
-void TPolymorphicArrayStructNetSerializerImpl<ExternalSourceType, ExternalSourceArrayItemType, GetArray, SetArrayNum>::CollectNetReferences(FNetSerializationContext& Context, const FNetCollectReferencesArgs& Args)
-{
-	const ConfigType& Config = *static_cast<const ConfigType*>(Args.NetSerializerConfig);
-	const QuantizedType& SourceValue = *reinterpret_cast<const QuantizedType*>(Args.Source);
-	FNetReferenceCollector& Collector = *reinterpret_cast<FNetReferenceCollector*>(Args.Collector);
-
-	// No references nothing to do
-	if (SourceValue.NumItems == 0U || !Config.RegisteredTypes.CanHaveNetReferences())
-	{
-		return;
-	}
-
-	for (uint32 It = 0, EndIt = SourceValue.NumItems; It != EndIt; ++It)
-	{
-		const FQuantizedItem& Item = SourceValue.Items[It];
-		const FTypeInfo* TypeInfo = Config.RegisteredTypes.GetTypeInfo(Item.TypeIndex);
-		const FReplicationStateDescriptor* Descriptor = TypeInfo ? TypeInfo->Descriptor : nullptr;
-
-		if (Descriptor && EnumHasAnyFlags(Descriptor->Traits, EReplicationStateTraits::HasObjectReference))
-		{			
-			FPolymorphicStructNetSerializerInternal::CollectReferences(Context, Collector, Args.ChangeMaskInfo, static_cast<const uint8*>(Item.StructData), Descriptor);
-		}
-	}
-}
-
 template <typename ExternalSourceType, typename ExternalSourceArrayItemType, TArrayView<TSharedPtr<ExternalSourceArrayItemType>>(*GetArray)(ExternalSourceType& Source), void(*SetArrayNum)(ExternalSourceType& Source, SIZE_T Num)>
 void TPolymorphicArrayStructNetSerializerImpl<ExternalSourceType, ExternalSourceArrayItemType, GetArray, SetArrayNum>::Deserialize(FNetSerializationContext& Context, const FNetDeserializeArgs& Args)
 {
@@ -538,7 +588,7 @@ void TPolymorphicArrayStructNetSerializerImpl<ExternalSourceType, ExternalSource
 		return;
 	}
 
-	// Currently we always free all memory even though we in could keep it if all types are matching
+	// Currently we always free all memory even though we in theory could keep it if all types are matching
 	InternalFreeItemArray(Context, TargetValue, Config);
 
 	// Allocate space for the ItemArray
@@ -547,6 +597,8 @@ void TPolymorphicArrayStructNetSerializerImpl<ExternalSourceType, ExternalSource
 	// Read polymorphic state data
 	for (uint32 It = 0, EndIt = TargetValue.NumItems; It != EndIt; ++It)
 	{
+		UE_NET_TRACE_SCOPE(Element, Reader, Context.GetTraceCollector(), ENetTraceVerbosity::VeryVerbose);
+
 		FQuantizedItem& Item = TargetValue.Items[It];
 		if (Reader.ReadBool())
 		{
@@ -575,13 +627,144 @@ void TPolymorphicArrayStructNetSerializerImpl<ExternalSourceType, ExternalSource
 template <typename ExternalSourceType, typename ExternalSourceArrayItemType, TArrayView<TSharedPtr<ExternalSourceArrayItemType>>(*GetArray)(ExternalSourceType& Source), void(*SetArrayNum)(ExternalSourceType& Source, SIZE_T Num)>
 void TPolymorphicArrayStructNetSerializerImpl<ExternalSourceType, ExternalSourceArrayItemType, GetArray, SetArrayNum>::SerializeDelta(FNetSerializationContext& Context, const FNetSerializeDeltaArgs& Args)
 {
-	NetSerializeDeltaDefault<Serialize>(Context, Args);
+	FNetBitStreamWriter& Writer = *Context.GetBitStreamWriter();
+	const ConfigType& Config = *static_cast<const ConfigType*>(Args.NetSerializerConfig);
+	const QuantizedType& Array = *reinterpret_cast<const QuantizedType*>(Args.Source);
+	const QuantizedType& PrevArray = *reinterpret_cast<const QuantizedType*>(Args.Prev);
+
+	const uint32 NumItems = Array.NumItems;
+	const uint32 PrevNumItems = PrevArray.NumItems;
+	const bool bSameSizeArray = (NumItems == PrevNumItems);
+	Writer.WriteBits(bSameSizeArray, 1U);
+	if (!bSameSizeArray)
+	{
+		Writer.WriteBits(NumItems, ArrayItemBits);
+	}
+
+	// Use delta serialization for elements available in both the previous and current state.
+	if (PrevNumItems)
+	{
+		FNetSerializeDeltaArgs ElementArgs = Args;
+
+		for (uint32 ElementIt = 0, ElementEndIt = FPlatformMath::Min(NumItems, PrevNumItems); ElementIt != ElementEndIt; ++ElementIt)
+		{
+			UE_NET_TRACE_SCOPE(Element, Writer, Context.GetTraceCollector(), ENetTraceVerbosity::VeryVerbose);
+
+			ElementArgs.Source = NetSerializerValuePointer(&Array.Items[ElementIt]);
+			ElementArgs.Prev = NetSerializerValuePointer(&PrevArray.Items[ElementIt]);
+
+			FItemNetSerializer::SerializeDelta(Context, ElementArgs);
+		}
+	}
+
+	// Serialize additional items with standard serialization.
+	if (NumItems > PrevNumItems)
+	{
+		for (uint32 ElementIt = PrevNumItems, ElementEndIt = NumItems; ElementIt < ElementEndIt; ++ElementIt)
+		{
+			UE_NET_TRACE_SCOPE(Element, Writer, Context.GetTraceCollector(), ENetTraceVerbosity::VeryVerbose);
+			const FQuantizedItem& Item = Array.Items[ElementIt];
+			const FTypeInfo* TypeInfo = Config.RegisteredTypes.GetTypeInfo(Item.TypeIndex);
+			if (Writer.WriteBool(TypeInfo != nullptr))
+			{
+				CA_ASSUME(TypeInfo != nullptr);
+				Writer.WriteBits(Item.TypeIndex, FPolymorphicNetSerializerScriptStructCache::RegisteredTypeBits);
+				FReplicationStateOperations::Serialize(Context, static_cast<const uint8*>(Item.StructData), TypeInfo->Descriptor);
+			}
+		}
+	}
 }
 
 template <typename ExternalSourceType, typename ExternalSourceArrayItemType, TArrayView<TSharedPtr<ExternalSourceArrayItemType>>(*GetArray)(ExternalSourceType& Source), void(*SetArrayNum)(ExternalSourceType& Source, SIZE_T Num)>
 void TPolymorphicArrayStructNetSerializerImpl<ExternalSourceType, ExternalSourceArrayItemType, GetArray, SetArrayNum>::DeserializeDelta(FNetSerializationContext& Context, const FNetDeserializeDeltaArgs& Args)
 {
-	NetDeserializeDeltaDefault<Deserialize>(Context, Args);
+	FNetBitStreamReader& Reader = *Context.GetBitStreamReader();
+	const ConfigType& Config = *static_cast<const ConfigType*>(Args.NetSerializerConfig);
+	QuantizedType& Array = *reinterpret_cast<QuantizedType*>(Args.Target);
+	const QuantizedType& PrevArray = *reinterpret_cast<const QuantizedType*>(Args.Prev);
+
+	const uint32 PrevNumItems = PrevArray.NumItems;
+	const bool bSameSizeArray = !!Reader.ReadBits(1U);
+	const uint32 NumItems = (bSameSizeArray ? PrevNumItems : Reader.ReadBits(ArrayItemBits));
+
+	if (Reader.IsOverflown())
+	{
+		Context.SetError(GNetError_BitStreamOverflow);
+		return;
+	}
+
+	// Currently we always free all memory even though we in theory could keep it if all types are matching
+	InternalFreeItemArray(Context, Array, Config);
+
+	// Allocate space for the ItemArray
+	InternalAllocateItemArray(Context, Array, NumItems);
+
+	// Elements in the current array up to the previous size can use delta serialization.
+	if (PrevNumItems)
+	{
+		FNetDeserializeDeltaArgs ElementArgs = Args;
+
+		for (uint32 ElementIt = 0, ElementEndIt = FPlatformMath::Min(NumItems, PrevNumItems); ElementIt != ElementEndIt; ++ElementIt)
+		{
+			UE_NET_TRACE_SCOPE(Element, Reader, Context.GetTraceCollector(), ENetTraceVerbosity::VeryVerbose);
+
+			ElementArgs.Target = NetSerializerValuePointer(&Array.Items[ElementIt]);
+			ElementArgs.Prev = NetSerializerValuePointer(&PrevArray.Items[ElementIt]);
+
+			FItemNetSerializer::DeserializeDelta(Context, ElementArgs);
+		}
+	}
+
+	if (Context.HasError())
+	{
+		InternalFreeItemArray(Context, Array, Config);
+		return;
+	}
+
+	if (Reader.IsOverflown())
+	{
+		Context.SetError(GNetError_BitStreamOverflow);
+		InternalFreeItemArray(Context, Array, Config);
+		return;
+	}
+
+	// Deserialize additional items with standard deserialization.
+	if (NumItems > PrevNumItems)
+	{
+		for (uint32 ElementIt = PrevNumItems, ElementEndIt = NumItems; ElementIt < ElementEndIt; ++ElementIt)
+		{
+			UE_NET_TRACE_SCOPE(Element, Reader, Context.GetTraceCollector(), ENetTraceVerbosity::VeryVerbose);
+			FQuantizedItem& Item = Array.Items[ElementIt];
+			if (Reader.ReadBool())
+			{
+				const uint32 TypeIndex = Reader.ReadBits(FPolymorphicNetSerializerScriptStructCache::RegisteredTypeBits);
+				if (const FTypeInfo* TypeInfo = Config.RegisteredTypes.GetTypeInfo(TypeIndex))
+				{
+					const FReplicationStateDescriptor* Descriptor = TypeInfo->Descriptor;
+
+					// Allocate storage and read struct data
+					Item.StructData = FPolymorphicStructNetSerializerInternal::Alloc(Context, Descriptor->InternalSize, Descriptor->InternalAlignment);
+					Item.TypeIndex = TypeIndex;
+
+					FMemory::Memzero(Item.StructData, Descriptor->InternalSize);
+					UE::Net::FReplicationStateOperations::Deserialize(Context, static_cast<uint8*>(Item.StructData), Descriptor);
+				}
+				else
+				{
+					InternalFreeItemArray(Context, Array, Config);
+					Context.SetError(NetError_PolymorphicStructNetSerializer_InvalidStructType);
+					return;
+				}
+			}
+		}
+	}
+
+	if (Reader.IsOverflown())
+	{
+		Context.SetError(GNetError_BitStreamOverflow);
+		InternalFreeItemArray(Context, Array, Config);
+		return;
+	}
 }
 
 template <typename ExternalSourceType, typename ExternalSourceArrayItemType, TArrayView<TSharedPtr<ExternalSourceArrayItemType>>(*GetArray)(ExternalSourceType& Source), void(*SetArrayNum)(ExternalSourceType& Source, SIZE_T Num)>
@@ -591,12 +774,19 @@ void TPolymorphicArrayStructNetSerializerImpl<ExternalSourceType, ExternalSource
 	QuantizedType& TargetValue = *reinterpret_cast<QuantizedType*>(Args.Target);
 	const ConfigType& Config = *static_cast<const ConfigType*>(Args.NetSerializerConfig);
 
+	// Make sure we can properly serialize the array.
+	TArrayView<const TSharedPtr<SourceArrayItemType>> ItemArray = GetArray(SourceValue);
+	const uint32 NumItems = static_cast<uint32>(ItemArray.Num());
+	if (NumItems > MaxArrayItems)
+	{
+		Context.SetError(GNetError_ArraySizeTooLarge);
+		return;
+	}
+
 	// First we need to free previously allocated data
 	InternalFreeItemArray(Context, TargetValue, Config);
 
 	// Allocate new array
-	TArrayView<const TSharedPtr<SourceArrayItemType>> ItemArray = GetArray(SourceValue);
-	const uint32 NumItems = static_cast<uint32>(ItemArray.Num());
 	InternalAllocateItemArray(Context, TargetValue, NumItems);
 
 	// Copy polymorphic struct data
@@ -719,7 +909,6 @@ bool TPolymorphicArrayStructNetSerializerImpl<ExternalSourceType, ExternalSource
 
 	TArrayView<const TSharedPtr<SourceArrayItemType>> ItemArray = GetArray(SourceValue);
 	const uint32 NumItems = (uint32)ItemArray.Num();
-
 	if (NumItems > MaxArrayItems)
 	{
 		return false;
@@ -745,6 +934,32 @@ bool TPolymorphicArrayStructNetSerializerImpl<ExternalSourceType, ExternalSource
 	return true;
 }
 
+template <typename ExternalSourceType, typename ExternalSourceArrayItemType, TArrayView<TSharedPtr<ExternalSourceArrayItemType>>(*GetArray)(ExternalSourceType& Source), void(*SetArrayNum)(ExternalSourceType& Source, SIZE_T Num)>
+void TPolymorphicArrayStructNetSerializerImpl<ExternalSourceType, ExternalSourceArrayItemType, GetArray, SetArrayNum>::CollectNetReferences(FNetSerializationContext& Context, const FNetCollectReferencesArgs& Args)
+{
+	const ConfigType& Config = *static_cast<const ConfigType*>(Args.NetSerializerConfig);
+	const QuantizedType& SourceValue = *reinterpret_cast<const QuantizedType*>(Args.Source);
+	FNetReferenceCollector& Collector = *reinterpret_cast<FNetReferenceCollector*>(Args.Collector);
+
+	// No references nothing to do
+	if (SourceValue.NumItems == 0U || !Config.RegisteredTypes.CanHaveNetReferences())
+	{
+		return;
+	}
+
+	for (uint32 It = 0, EndIt = SourceValue.NumItems; It != EndIt; ++It)
+	{
+		const FQuantizedItem& Item = SourceValue.Items[It];
+		const FTypeInfo* TypeInfo = Config.RegisteredTypes.GetTypeInfo(Item.TypeIndex);
+		const FReplicationStateDescriptor* Descriptor = TypeInfo ? TypeInfo->Descriptor : nullptr;
+
+		if (Descriptor && EnumHasAnyFlags(Descriptor->Traits, EReplicationStateTraits::HasObjectReference))
+		{			
+			FPolymorphicStructNetSerializerInternal::CollectReferences(Context, Collector, Args.ChangeMaskInfo, static_cast<const uint8*>(Item.StructData), Descriptor);
+		}
+	}
+}
+
 template <typename ExternalSourceType, typename ExternalSourceArrayItemType, TArrayView<TSharedPtr<ExternalSourceArrayItemType>>(*GetArray)(ExternalSourceType& Source), void(*SetArrayNum)(ExternalSourceType& Source, SIZE_T Num)>
 void TPolymorphicArrayStructNetSerializerImpl<ExternalSourceType, ExternalSourceArrayItemType, GetArray, SetArrayNum>::CloneDynamicState(FNetSerializationContext& Context, const FNetCloneDynamicStateArgs& Args)
 {
