Iris - Implemented globally unique NetHandles.

diff --git a/Engine/Plugins/Runtime/ReplicationSystemTestPlugin/Source/Private/Tests/ReplicationState/TestFastArrayReplicationState.cpp b/Engine/Plugins/Runtime/ReplicationSystemTestPlugin/Source/Private/Tests/ReplicationState/TestFastArrayReplicationState.cpp
--- a/Engine/Plugins/Runtime/ReplicationSystemTestPlugin/Source/Private/Tests/ReplicationState/TestFastArrayReplicationState.cpp
+++ b/Engine/Plugins/Runtime/ReplicationSystemTestPlugin/Source/Private/Tests/ReplicationState/TestFastArrayReplicationState.cpp
@@ -138,7 +138,7 @@ UE_NET_TEST_FIXTURE(FReplicationSystemServerClientTestFixture, TestAddItem)
 	Server->PostSendUpdate();
 
 	// Store Pointer to objects
-	UTestFastArrayReplicationState_FastArray_TestClassFastArray* ClientObject = Cast<UTestFastArrayReplicationState_FastArray_TestClassFastArray>(Client->GetReplicationBridge()->GetReplicatedObject(ServerObject->NetHandle));
+	UTestFastArrayReplicationState_FastArray_TestClassFastArray* ClientObject = Cast<UTestFastArrayReplicationState_FastArray_TestClassFastArray>(Client->GetReplicationBridge()->GetReplicatedObject(ServerObject->NetRefHandle));
 
 	// Verify initial replication
 	UE_NET_ASSERT_TRUE(ClientObject!= nullptr);
@@ -190,7 +190,7 @@ UE_NET_TEST_FIXTURE(FReplicationSystemServerClientTestFixture, TestRemoveItem)
 	Server->PostSendUpdate();
 
 	// Store Pointer to objects
-	UTestFastArrayReplicationState_FastArray_TestClassFastArray* ClientObject = Cast<UTestFastArrayReplicationState_FastArray_TestClassFastArray>(Client->GetReplicationBridge()->GetReplicatedObject(ServerObject->NetHandle));
+	UTestFastArrayReplicationState_FastArray_TestClassFastArray* ClientObject = Cast<UTestFastArrayReplicationState_FastArray_TestClassFastArray>(Client->GetReplicationBridge()->GetReplicatedObject(ServerObject->NetRefHandle));
 
 	// Verify initial replication
 	UE_NET_ASSERT_TRUE(ClientObject!= nullptr);
@@ -258,7 +258,7 @@ UE_NET_TEST_FIXTURE(FReplicationSystemServerClientTestFixture, TestLocalItemAreP
 	Server->PostSendUpdate();
 
 	// Store Pointer to objects
-	UTestFastArrayReplicationState_FastArray_TestClassFastArray* ClientObject = Cast<UTestFastArrayReplicationState_FastArray_TestClassFastArray>(Client->GetReplicationBridge()->GetReplicatedObject(ServerObject->NetHandle));
+	UTestFastArrayReplicationState_FastArray_TestClassFastArray* ClientObject = Cast<UTestFastArrayReplicationState_FastArray_TestClassFastArray>(Client->GetReplicationBridge()->GetReplicatedObject(ServerObject->NetRefHandle));
 
 	// Verify initial replication
 	UE_NET_ASSERT_TRUE(ClientObject!= nullptr);
@@ -341,7 +341,7 @@ UE_NET_TEST_FIXTURE(FReplicationSystemServerClientTestFixture, TestModifyItem)
 	Server->PostSendUpdate();
 
 	// Store Pointer to objects
-	UTestFastArrayReplicationState_FastArray_TestClassFastArray* ClientObject = Cast<UTestFastArrayReplicationState_FastArray_TestClassFastArray>(Client->GetReplicationBridge()->GetReplicatedObject(ServerObject->NetHandle));
+	UTestFastArrayReplicationState_FastArray_TestClassFastArray* ClientObject = Cast<UTestFastArrayReplicationState_FastArray_TestClassFastArray>(Client->GetReplicationBridge()->GetReplicatedObject(ServerObject->NetRefHandle));
 
 	// Verify initial replication
 	UE_NET_ASSERT_TRUE(ClientObject!= nullptr);
@@ -414,7 +414,7 @@ UE_NET_TEST_FIXTURE(FReplicationSystemServerClientTestFixture, TestInsertItem)
 	Server->PostSendUpdate();
 
 	// Store Pointer to objects
-	UTestFastArrayReplicationState_FastArray_TestClassFastArray* ClientObject = Cast<UTestFastArrayReplicationState_FastArray_TestClassFastArray>(Client->GetReplicationBridge()->GetReplicatedObject(ServerObject->NetHandle));
+	UTestFastArrayReplicationState_FastArray_TestClassFastArray* ClientObject = Cast<UTestFastArrayReplicationState_FastArray_TestClassFastArray>(Client->GetReplicationBridge()->GetReplicatedObject(ServerObject->NetRefHandle));
 
 	// Verify initial replication
 	UE_NET_ASSERT_TRUE(ClientObject!= nullptr);
@@ -496,7 +496,7 @@ UE_NET_TEST_FIXTURE(FReplicationSystemServerClientTestFixture, TestModifyItemDoe
 	Server->PostSendUpdate();
 
 	// Store Pointer to objects
-	UTestFastArrayReplicationState_FastArray_TestClassFastArray* ClientObject = Cast<UTestFastArrayReplicationState_FastArray_TestClassFastArray>(Client->GetReplicationBridge()->GetReplicatedObject(ServerObject->NetHandle));
+	UTestFastArrayReplicationState_FastArray_TestClassFastArray* ClientObject = Cast<UTestFastArrayReplicationState_FastArray_TestClassFastArray>(Client->GetReplicationBridge()->GetReplicatedObject(ServerObject->NetRefHandle));
 
 	// Verify initial replication
 	UE_NET_ASSERT_TRUE(ClientObject!= nullptr);
@@ -599,7 +599,7 @@ UE_NET_TEST_FIXTURE(FReplicationSystemServerClientTestFixture, TestAddItemWithou
 	Server->PostSendUpdate();
 
 	// Store Pointer to objects
-	UTestFastArrayReplicationState_FastArray_TestClassFastArray* ClientObject = Cast<UTestFastArrayReplicationState_FastArray_TestClassFastArray>(Client->GetReplicationBridge()->GetReplicatedObject(ServerObject->NetHandle));
+	UTestFastArrayReplicationState_FastArray_TestClassFastArray* ClientObject = Cast<UTestFastArrayReplicationState_FastArray_TestClassFastArray>(Client->GetReplicationBridge()->GetReplicatedObject(ServerObject->NetRefHandle));
 
 	// Verify that we received the data
 	UE_NET_ASSERT_TRUE(ClientObject!= nullptr);
@@ -622,7 +622,7 @@ UE_NET_TEST_FIXTURE(FReplicationSystemServerClientTestFixture, TestChangeMaskWra
 	Server->PostSendUpdate();
 
 	// Store Pointer to objects
-	UTestFastArrayReplicationState_FastArray_TestClassFastArray* ClientObject = Cast<UTestFastArrayReplicationState_FastArray_TestClassFastArray>(Client->GetReplicationBridge()->GetReplicatedObject(ServerObject->NetHandle));
+	UTestFastArrayReplicationState_FastArray_TestClassFastArray* ClientObject = Cast<UTestFastArrayReplicationState_FastArray_TestClassFastArray>(Client->GetReplicationBridge()->GetReplicatedObject(ServerObject->NetRefHandle));
 
 	// Verify initial replication
 	UE_NET_ASSERT_TRUE(ClientObject!= nullptr);
@@ -736,8 +736,8 @@ UE_NET_TEST_FIXTURE(FReplicationSystemServerClientTestFixture, TestPostReplicate
 	FastArrayEditor.Add(Item);
 
 	// Make sure that references are not replicated to client
-	Server->ReplicationSystem->AddToGroup(NotReplicatedNetObjectGroupHandle, ServerReferencedObjectA->NetHandle);
-	Server->ReplicationSystem->AddToGroup(NotReplicatedNetObjectGroupHandle, ServerReferencedObjectB->NetHandle);
+	Server->ReplicationSystem->AddToGroup(NotReplicatedNetObjectGroupHandle, ServerReferencedObjectA->NetRefHandle);
+	Server->ReplicationSystem->AddToGroup(NotReplicatedNetObjectGroupHandle, ServerReferencedObjectB->NetRefHandle);
 	
 	// Replicate
 	Server->PreSendUpdate();
@@ -745,7 +745,7 @@ UE_NET_TEST_FIXTURE(FReplicationSystemServerClientTestFixture, TestPostReplicate
 	Server->PostSendUpdate();
 
 	// Verify that object has been spawned on client
-	auto ClientObject = Client->GetObjectAs<UTestFastArrayReplicationState_FastArray_TestClassFastArray>(ServerObject->NetHandle);
+	auto ClientObject = Client->GetObjectAs<UTestFastArrayReplicationState_FastArray_TestClassFastArray>(ServerObject->NetRefHandle);
 	UE_NET_ASSERT_TRUE(ClientObject != nullptr);
 
 	auto& ClientFastArray = ClientObject->FastArray;
@@ -766,7 +766,7 @@ UE_NET_TEST_FIXTURE(FReplicationSystemServerClientTestFixture, TestPostReplicate
 	ClientFastArray.bPostReplicatedReceiveWasHitWithUnresolvedReferences = false;
 
 	// Enable replication for ServerReferenceObjectA
-	Server->ReplicationSystem->RemoveFromGroup(NotReplicatedNetObjectGroupHandle, ServerReferencedObjectA->NetHandle);
+	Server->ReplicationSystem->RemoveFromGroup(NotReplicatedNetObjectGroupHandle, ServerReferencedObjectA->NetRefHandle);
 
 	// Replicate
 	Server->PreSendUpdate();
@@ -785,7 +785,7 @@ UE_NET_TEST_FIXTURE(FReplicationSystemServerClientTestFixture, TestPostReplicate
 	ClientFastArray.bPostReplicatedReceiveWasHitWithUnresolvedReferences = false;
 
 	// Enable replication for ServerReferenceObjectB 
-	Server->ReplicationSystem->RemoveFromGroup(NotReplicatedNetObjectGroupHandle, ServerReferencedObjectB->NetHandle);
+	Server->ReplicationSystem->RemoveFromGroup(NotReplicatedNetObjectGroupHandle, ServerReferencedObjectB->NetRefHandle);
 
 	// Replicate
 	Server->PreSendUpdate();
diff --git a/Engine/Plugins/Runtime/ReplicationSystemTestPlugin/Source/Private/Tests/ReplicationSystem/Conditionals/TestConditionals.cpp b/Engine/Plugins/Runtime/ReplicationSystemTestPlugin/Source/Private/Tests/ReplicationSystem/Conditionals/TestConditionals.cpp
--- a/Engine/Plugins/Runtime/ReplicationSystemTestPlugin/Source/Private/Tests/ReplicationSystem/Conditionals/TestConditionals.cpp
+++ b/Engine/Plugins/Runtime/ReplicationSystemTestPlugin/Source/Private/Tests/ReplicationSystem/Conditionals/TestConditionals.cpp
@@ -31,7 +31,7 @@ UE_NET_TEST_FIXTURE(FTestConditionalsFixture, ToOwnerStateIsReplicatedToOwner)
 	ServerObject->ConnectionFilteredComponents[0]->ToOwnerB = 37;
 
 	// Set owner
-	Server->ReplicationSystem->SetOwningNetConnection(ServerObject->NetHandle, Client->ConnectionIdOnServer);
+	Server->ReplicationSystem->SetOwningNetConnection(ServerObject->NetRefHandle, Client->ConnectionIdOnServer);
 
 	// Send and deliver packet
 	Server->PreSendUpdate();
@@ -39,7 +39,7 @@ UE_NET_TEST_FIXTURE(FTestConditionalsFixture, ToOwnerStateIsReplicatedToOwner)
 	Server->PostSendUpdate();
 	
 	// Object should have been created on the client
-	const UTestReplicatedIrisObject* ClientObject = Cast<UTestReplicatedIrisObject>(Client->GetReplicationBridge()->GetReplicatedObject(ServerObject->NetHandle));
+	const UTestReplicatedIrisObject* ClientObject = Cast<UTestReplicatedIrisObject>(Client->GetReplicationBridge()->GetReplicatedObject(ServerObject->NetRefHandle));
 	UE_NET_ASSERT_NE(ClientObject, nullptr);
 
 	// Check that the To Owner members have the same values as on the sending side.
@@ -70,7 +70,7 @@ UE_NET_TEST_FIXTURE(FTestConditionalsFixture, ToOwnerStateIsNotReplicatedToNonOw
 	Server->PostSendUpdate();
 	
 	// Object should have been created on the client
-	const UTestReplicatedIrisObject* ClientObject = Cast<UTestReplicatedIrisObject>(Client->GetReplicationBridge()->GetReplicatedObject(ServerObject->NetHandle));
+	const UTestReplicatedIrisObject* ClientObject = Cast<UTestReplicatedIrisObject>(Client->GetReplicationBridge()->GetReplicatedObject(ServerObject->NetRefHandle));
 	UE_NET_ASSERT_NE(ClientObject, nullptr);
 
 	// Check that the To Owner members haven't been modified
@@ -94,7 +94,7 @@ UE_NET_TEST_FIXTURE(FTestConditionalsFixture, SkipOwnerStateIsNotReplicatedToOwn
 	ServerObject->ConnectionFilteredComponents[0]->SkipOwnerB = 37;
 
 	// Set owner
-	Server->ReplicationSystem->SetOwningNetConnection(ServerObject->NetHandle, Client->ConnectionIdOnServer);
+	Server->ReplicationSystem->SetOwningNetConnection(ServerObject->NetRefHandle, Client->ConnectionIdOnServer);
 
 	// Send and deliver packet
 	Server->PreSendUpdate();
@@ -102,7 +102,7 @@ UE_NET_TEST_FIXTURE(FTestConditionalsFixture, SkipOwnerStateIsNotReplicatedToOwn
 	Server->PostSendUpdate();
 	
 	// Object should have been created on the client
-	const UTestReplicatedIrisObject* ClientObject = Cast<UTestReplicatedIrisObject>(Client->GetReplicationBridge()->GetReplicatedObject(ServerObject->NetHandle));
+	const UTestReplicatedIrisObject* ClientObject = Cast<UTestReplicatedIrisObject>(Client->GetReplicationBridge()->GetReplicatedObject(ServerObject->NetRefHandle));
 	UE_NET_ASSERT_NE(ClientObject, nullptr);
 
 	// Check that the SkipOwner members haven't been modified
@@ -133,7 +133,7 @@ UE_NET_TEST_FIXTURE(FTestConditionalsFixture, SkipOwnerStateIsReplicatedToNonOwn
 	Server->PostSendUpdate();
 	
 	// Object should have been created on the client
-	const UTestReplicatedIrisObject* ClientObject = Cast<UTestReplicatedIrisObject>(Client->GetReplicationBridge()->GetReplicatedObject(ServerObject->NetHandle));
+	const UTestReplicatedIrisObject* ClientObject = Cast<UTestReplicatedIrisObject>(Client->GetReplicationBridge()->GetReplicatedObject(ServerObject->NetRefHandle));
 	UE_NET_ASSERT_NE(ClientObject, nullptr);
 
 	// Check that the SkipOwner members have been modified
@@ -166,7 +166,7 @@ UE_NET_TEST_FIXTURE(FTestConditionalsFixture, SimulatedIsReplicatedToSimulated)
 	Server->PostSendUpdate();
 	
 	// Object should have been created on the client
-	const UTestReplicatedIrisObject* ClientObject = Cast<UTestReplicatedIrisObject>(Client->GetReplicationBridge()->GetReplicatedObject(ServerObject->NetHandle));
+	const UTestReplicatedIrisObject* ClientObject = Cast<UTestReplicatedIrisObject>(Client->GetReplicationBridge()->GetReplicatedObject(ServerObject->NetRefHandle));
 	UE_NET_ASSERT_NE(ClientObject, nullptr);
 
 	// Check that the various Simulated members have been modified
@@ -194,7 +194,7 @@ UE_NET_TEST_FIXTURE(FTestConditionalsFixture, SimulatedIsNotReplicatedToAutonomo
 	ServerObject->ConnectionFilteredComponents[0]->SimulatedOrPhysicsNoReplayInt = 11;
 
 	// Set this connection to be autonomous.
-	Server->ReplicationSystem->SetReplicationConditionConnectionFilter(ServerObject->NetHandle, EReplicationCondition::RoleAutonomous, Client->ConnectionIdOnServer, true);
+	Server->ReplicationSystem->SetReplicationConditionConnectionFilter(ServerObject->NetRefHandle, EReplicationCondition::RoleAutonomous, Client->ConnectionIdOnServer, true);
 
 	// Send and deliver packet
 	Server->PreSendUpdate();
@@ -202,7 +202,7 @@ UE_NET_TEST_FIXTURE(FTestConditionalsFixture, SimulatedIsNotReplicatedToAutonomo
 	Server->PostSendUpdate();
 	
 	// Object should have been created on the client
-	const UTestReplicatedIrisObject* ClientObject = Cast<UTestReplicatedIrisObject>(Client->GetReplicationBridge()->GetReplicatedObject(ServerObject->NetHandle));
+	const UTestReplicatedIrisObject* ClientObject = Cast<UTestReplicatedIrisObject>(Client->GetReplicationBridge()->GetReplicatedObject(ServerObject->NetRefHandle));
 	UE_NET_ASSERT_NE(ClientObject, nullptr);
 
 	// Check that the various Simulated members have not been modified
@@ -227,7 +227,7 @@ UE_NET_TEST_FIXTURE(FTestConditionalsFixture, AutonomousIsReplicatedToAutonomous
 	ServerObject->ConnectionFilteredComponents[0]->AutonomousOnlyInt = 13;
 
 	// Set this connection to be autonomous.
-	Server->ReplicationSystem->SetReplicationConditionConnectionFilter(ServerObject->NetHandle, EReplicationCondition::RoleAutonomous, Client->ConnectionIdOnServer, true);
+	Server->ReplicationSystem->SetReplicationConditionConnectionFilter(ServerObject->NetRefHandle, EReplicationCondition::RoleAutonomous, Client->ConnectionIdOnServer, true);
 
 	// Send and deliver packet
 	Server->PreSendUpdate();
@@ -235,7 +235,7 @@ UE_NET_TEST_FIXTURE(FTestConditionalsFixture, AutonomousIsReplicatedToAutonomous
 	Server->PostSendUpdate();
 	
 	// Object should have been created on the client
-	const UTestReplicatedIrisObject* ClientObject = Cast<UTestReplicatedIrisObject>(Client->GetReplicationBridge()->GetReplicatedObject(ServerObject->NetHandle));
+	const UTestReplicatedIrisObject* ClientObject = Cast<UTestReplicatedIrisObject>(Client->GetReplicationBridge()->GetReplicatedObject(ServerObject->NetRefHandle));
 	UE_NET_ASSERT_NE(ClientObject, nullptr);
 
 	// Check that the Autonomous member has been modified
@@ -259,7 +259,7 @@ UE_NET_TEST_FIXTURE(FTestConditionalsFixture, AutonomousIsNotReplicatedToSimulat
 	ServerObject->ConnectionFilteredComponents[0]->AutonomousOnlyInt = 13;
 
 	// Set other connection to be autonomous.
-	Server->ReplicationSystem->SetReplicationConditionConnectionFilter(ServerObject->NetHandle, EReplicationCondition::RoleAutonomous, AutonomousClient->ConnectionIdOnServer, true);
+	Server->ReplicationSystem->SetReplicationConditionConnectionFilter(ServerObject->NetRefHandle, EReplicationCondition::RoleAutonomous, AutonomousClient->ConnectionIdOnServer, true);
 
 	// Send and deliver packet
 	Server->PreSendUpdate();
@@ -267,7 +267,7 @@ UE_NET_TEST_FIXTURE(FTestConditionalsFixture, AutonomousIsNotReplicatedToSimulat
 	Server->PostSendUpdate();
 	
 	// Object should have been created on the client
-	const UTestReplicatedIrisObject* ClientObject = Cast<UTestReplicatedIrisObject>(Client->GetReplicationBridge()->GetReplicatedObject(ServerObject->NetHandle));
+	const UTestReplicatedIrisObject* ClientObject = Cast<UTestReplicatedIrisObject>(Client->GetReplicationBridge()->GetReplicatedObject(ServerObject->NetRefHandle));
 	UE_NET_ASSERT_NE(ClientObject, nullptr);
 
 	// Check that the Autonomous member has not been modified
@@ -290,9 +290,9 @@ UE_NET_TEST_FIXTURE(FTestConditionalsFixture, SimulatedOrPhysicsIsReplicatedWhen
 	ServerObject->ConnectionFilteredComponents[0]->SimulatedOrPhysicsNoReplayInt = 11;
 
 	// Set this connection to be autonomous so we know that only the Physics condition is at play
-	Server->ReplicationSystem->SetReplicationConditionConnectionFilter(ServerObject->NetHandle, EReplicationCondition::RoleAutonomous, Client->ConnectionIdOnServer, true);
+	Server->ReplicationSystem->SetReplicationConditionConnectionFilter(ServerObject->NetRefHandle, EReplicationCondition::RoleAutonomous, Client->ConnectionIdOnServer, true);
 	// Enable ReplicatePhysics condition
-	Server->ReplicationSystem->SetReplicationCondition(ServerObject->NetHandle, EReplicationCondition::ReplicatePhysics, true);
+	Server->ReplicationSystem->SetReplicationCondition(ServerObject->NetRefHandle, EReplicationCondition::ReplicatePhysics, true);
 
 	// Send and deliver packet
 	Server->PreSendUpdate();
@@ -300,7 +300,7 @@ UE_NET_TEST_FIXTURE(FTestConditionalsFixture, SimulatedOrPhysicsIsReplicatedWhen
 	Server->PostSendUpdate();
 	
 	// Object should have been created on the client
-	const UTestReplicatedIrisObject* ClientObject = Cast<UTestReplicatedIrisObject>(Client->GetReplicationBridge()->GetReplicatedObject(ServerObject->NetHandle));
+	const UTestReplicatedIrisObject* ClientObject = Cast<UTestReplicatedIrisObject>(Client->GetReplicationBridge()->GetReplicatedObject(ServerObject->NetRefHandle));
 	UE_NET_ASSERT_NE(ClientObject, nullptr);
 
 	// Check that the various Physics members have been modified
@@ -324,9 +324,9 @@ UE_NET_TEST_FIXTURE(FTestConditionalsFixture, SimulatedOrPhysicsIsNotReplicatedW
 	ServerObject->ConnectionFilteredComponents[0]->SimulatedOrPhysicsNoReplayInt = 11;
 
 	// Set this connection to be autonomous so we know that only the Physics condition is at play
-	Server->ReplicationSystem->SetReplicationConditionConnectionFilter(ServerObject->NetHandle, EReplicationCondition::RoleAutonomous, Client->ConnectionIdOnServer, true);
+	Server->ReplicationSystem->SetReplicationConditionConnectionFilter(ServerObject->NetRefHandle, EReplicationCondition::RoleAutonomous, Client->ConnectionIdOnServer, true);
 	// Disable ReplicatePhysics condition, though that is expected to be the default
-	Server->ReplicationSystem->SetReplicationCondition(ServerObject->NetHandle, EReplicationCondition::ReplicatePhysics, false);
+	Server->ReplicationSystem->SetReplicationCondition(ServerObject->NetRefHandle, EReplicationCondition::ReplicatePhysics, false);
 
 	// Send and deliver packet
 	Server->PreSendUpdate();
@@ -334,7 +334,7 @@ UE_NET_TEST_FIXTURE(FTestConditionalsFixture, SimulatedOrPhysicsIsNotReplicatedW
 	Server->PostSendUpdate();
 	
 	// Object should have been created on the client
-	const UTestReplicatedIrisObject* ClientObject = Cast<UTestReplicatedIrisObject>(Client->GetReplicationBridge()->GetReplicatedObject(ServerObject->NetHandle));
+	const UTestReplicatedIrisObject* ClientObject = Cast<UTestReplicatedIrisObject>(Client->GetReplicationBridge()->GetReplicatedObject(ServerObject->NetRefHandle));
 	UE_NET_ASSERT_NE(ClientObject, nullptr);
 
 	// Check that the various Physics members have not been modified
@@ -357,7 +357,7 @@ UE_NET_TEST_FIXTURE(FTestConditionalsFixture, CanToggleAutonomous)
 	ServerObject->ConnectionFilteredComponents[0]->AutonomousOnlyInt = 13;
 
 	// Set this connection to be autonomous.
-	Server->ReplicationSystem->SetReplicationConditionConnectionFilter(ServerObject->NetHandle, EReplicationCondition::RoleAutonomous, Client->ConnectionIdOnServer, true);
+	Server->ReplicationSystem->SetReplicationConditionConnectionFilter(ServerObject->NetRefHandle, EReplicationCondition::RoleAutonomous, Client->ConnectionIdOnServer, true);
 
 	// Send and deliver packet
 	Server->PreSendUpdate();
@@ -365,14 +365,14 @@ UE_NET_TEST_FIXTURE(FTestConditionalsFixture, CanToggleAutonomous)
 	Server->PostSendUpdate();
 	
 	// Object should have been created on the client
-	const UTestReplicatedIrisObject* ClientObject = Cast<UTestReplicatedIrisObject>(Client->GetReplicationBridge()->GetReplicatedObject(ServerObject->NetHandle));
+	const UTestReplicatedIrisObject* ClientObject = Cast<UTestReplicatedIrisObject>(Client->GetReplicationBridge()->GetReplicatedObject(ServerObject->NetRefHandle));
 	UE_NET_ASSERT_NE(ClientObject, nullptr);
 
 	// Check that the Autonomous member has been modified
 	UE_NET_ASSERT_EQ(ClientObject->ConnectionFilteredComponents[0]->AutonomousOnlyInt, ServerObject->ConnectionFilteredComponents[0]->AutonomousOnlyInt);
 
 	// Set this connection to not be autonomous (one could also pass ConnectionId 0)
-	Server->ReplicationSystem->SetReplicationConditionConnectionFilter(ServerObject->NetHandle, EReplicationCondition::RoleAutonomous, Client->ConnectionIdOnServer, false);
+	Server->ReplicationSystem->SetReplicationConditionConnectionFilter(ServerObject->NetRefHandle, EReplicationCondition::RoleAutonomous, Client->ConnectionIdOnServer, false);
 
 	// Set values that should be replicated for autonomous objects
 	ServerObject->ConnectionFilteredComponents[0]->AutonomousOnlyInt = 37;
@@ -414,7 +414,7 @@ UE_NET_TEST_FIXTURE(FTestConditionalsFixture, CanToggleSimulated)
 	Server->PostSendUpdate();
 	
 	// Object should have been created on the client
-	const UTestReplicatedIrisObject* ClientObject = Cast<UTestReplicatedIrisObject>(Client->GetReplicationBridge()->GetReplicatedObject(ServerObject->NetHandle));
+	const UTestReplicatedIrisObject* ClientObject = Cast<UTestReplicatedIrisObject>(Client->GetReplicationBridge()->GetReplicatedObject(ServerObject->NetRefHandle));
 	UE_NET_ASSERT_NE(ClientObject, nullptr);
 
 	// Check that the various Simulated members have been modified
@@ -424,7 +424,7 @@ UE_NET_TEST_FIXTURE(FTestConditionalsFixture, CanToggleSimulated)
 	UE_NET_ASSERT_EQ(ClientObject->ConnectionFilteredComponents[0]->SimulatedOrPhysicsNoReplayInt, ServerObject->ConnectionFilteredComponents[0]->SimulatedOrPhysicsNoReplayInt);
 
 	// Set this connection to be autonomous.
-	Server->ReplicationSystem->SetReplicationConditionConnectionFilter(ServerObject->NetHandle, EReplicationCondition::RoleAutonomous, Client->ConnectionIdOnServer, true);
+	Server->ReplicationSystem->SetReplicationConditionConnectionFilter(ServerObject->NetRefHandle, EReplicationCondition::RoleAutonomous, Client->ConnectionIdOnServer, true);
 
 	// Set values that should be replicated for simulated objects
 	ServerObject->ConnectionFilteredComponents[0]->SimulatedOnlyInt = 47;
@@ -464,9 +464,9 @@ UE_NET_TEST_FIXTURE(FTestConditionalsFixture, CanMixConditions)
 	UTestReplicatedIrisObject* ServerObject = Server->CreateObject(ObjectComponents);
 
 	// Set Autonomous connection
-	Server->ReplicationSystem->SetReplicationConditionConnectionFilter(ServerObject->NetHandle, EReplicationCondition::RoleAutonomous, AutonomousClientConnectionId, true);
+	Server->ReplicationSystem->SetReplicationConditionConnectionFilter(ServerObject->NetRefHandle, EReplicationCondition::RoleAutonomous, AutonomousClientConnectionId, true);
 	// Enable ReplicatePhysics condition
-	Server->ReplicationSystem->SetReplicationCondition(ServerObject->NetHandle, EReplicationCondition::ReplicatePhysics, true);
+	Server->ReplicationSystem->SetReplicationCondition(ServerObject->NetRefHandle, EReplicationCondition::ReplicatePhysics, true);
 
 	// Set all the values
 	ServerObject->ConnectionFilteredComponents[0]->AutonomousOnlyInt = 1;
@@ -488,7 +488,7 @@ UE_NET_TEST_FIXTURE(FTestConditionalsFixture, CanMixConditions)
 	for (FReplicationSystemTestClient*& Client : ClientArray)
 	{
 		const uint32 ClientIndex = &Client - ClientArray;
-		ClientObjectArray[ClientIndex] = Cast<UTestReplicatedIrisObject>(Client->GetReplicationBridge()->GetReplicatedObject(ServerObject->NetHandle));
+		ClientObjectArray[ClientIndex] = Cast<UTestReplicatedIrisObject>(Client->GetReplicationBridge()->GetReplicatedObject(ServerObject->NetRefHandle));
 		UE_NET_ASSERT_NE(ClientObjectArray[ClientIndex], nullptr);
 	}
 
@@ -530,7 +530,7 @@ UE_NET_TEST_FIXTURE(FTestConditionalsFixture, CanSwitchAutonmousConnection)
 	UTestReplicatedIrisObject* ServerObject = Server->CreateObject(ObjectComponents);
 
 	// Set Autonomous connection
-	Server->ReplicationSystem->SetReplicationConditionConnectionFilter(ServerObject->NetHandle, EReplicationCondition::RoleAutonomous, AutonomousClientConnectionId, true);
+	Server->ReplicationSystem->SetReplicationConditionConnectionFilter(ServerObject->NetRefHandle, EReplicationCondition::RoleAutonomous, AutonomousClientConnectionId, true);
 
 	// Set all the values
 	ServerObject->ConnectionFilteredComponents[0]->AutonomousOnlyInt = 1;
@@ -552,14 +552,14 @@ UE_NET_TEST_FIXTURE(FTestConditionalsFixture, CanSwitchAutonmousConnection)
 	for (FReplicationSystemTestClient*& Client : ClientArray)
 	{
 		const uint32 ClientIndex = &Client - ClientArray;
-		ClientObjectArray[ClientIndex] = Cast<UTestReplicatedIrisObject>(Client->GetReplicationBridge()->GetReplicatedObject(ServerObject->NetHandle));
+		ClientObjectArray[ClientIndex] = Cast<UTestReplicatedIrisObject>(Client->GetReplicationBridge()->GetReplicatedObject(ServerObject->NetRefHandle));
 		UE_NET_ASSERT_NE(ClientObjectArray[ClientIndex], nullptr);
 	}
 
 	// Switch Autonomous connection
 	AutonomousClientIndex = (AutonomousClientIndex + 1)  % UE_ARRAY_COUNT(ClientArray);
 	AutonomousClientConnectionId = ClientArray[AutonomousClientIndex]->ConnectionIdOnServer;
-	Server->ReplicationSystem->SetReplicationConditionConnectionFilter(ServerObject->NetHandle, EReplicationCondition::RoleAutonomous, AutonomousClientConnectionId, true);
+	Server->ReplicationSystem->SetReplicationConditionConnectionFilter(ServerObject->NetRefHandle, EReplicationCondition::RoleAutonomous, AutonomousClientConnectionId, true);
 
 	// Set all the values
 	ServerObject->ConnectionFilteredComponents[0]->AutonomousOnlyInt = 1+1;
@@ -620,13 +620,13 @@ UE_NET_TEST_FIXTURE(FTestConditionalsFixture, SubObjectConditionsMatchesParentOb
 	UTestReplicatedIrisObject::FComponents ObjectComponents;
 	ObjectComponents.ConnectionFilteredComponentCount = ConnectionFilteredComponentCount;
 	UTestReplicatedIrisObject* ServerObject = Server->CreateObject(ObjectComponents);
-	UTestReplicatedIrisObject* ServerSubObject = Server->CreateSubObject(ServerObject->NetHandle, ObjectComponents);
+	UTestReplicatedIrisObject* ServerSubObject = Server->CreateSubObject(ServerObject->NetRefHandle, ObjectComponents);
 
 
 	// Set Autonomous connection
-	Server->ReplicationSystem->SetReplicationConditionConnectionFilter(ServerObject->NetHandle, EReplicationCondition::RoleAutonomous, AutonomousClientConnectionId, true);
+	Server->ReplicationSystem->SetReplicationConditionConnectionFilter(ServerObject->NetRefHandle, EReplicationCondition::RoleAutonomous, AutonomousClientConnectionId, true);
 	// Enable ReplicatePhysics condition
-	Server->ReplicationSystem->SetReplicationCondition(ServerObject->NetHandle, EReplicationCondition::ReplicatePhysics, true);
+	Server->ReplicationSystem->SetReplicationCondition(ServerObject->NetRefHandle, EReplicationCondition::ReplicatePhysics, true);
 
 	// Set all the values
 	for (UTestReplicatedIrisObject* Object : {ServerObject, ServerSubObject})
@@ -652,9 +652,9 @@ UE_NET_TEST_FIXTURE(FTestConditionalsFixture, SubObjectConditionsMatchesParentOb
 	for (FReplicationSystemTestClient*& Client : ClientArray)
 	{
 		const uint32 ClientIndex = &Client - ClientArray;
-		ClientObjectArray[ClientIndex] = Cast<UTestReplicatedIrisObject>(Client->GetReplicationBridge()->GetReplicatedObject(ServerObject->NetHandle));
+		ClientObjectArray[ClientIndex] = Cast<UTestReplicatedIrisObject>(Client->GetReplicationBridge()->GetReplicatedObject(ServerObject->NetRefHandle));
 		UE_NET_ASSERT_NE(ClientObjectArray[ClientIndex], nullptr);
-		ClientSubObjectArray[ClientIndex] = Cast<UTestReplicatedIrisObject>(Client->GetReplicationBridge()->GetReplicatedObject(ServerSubObject->NetHandle));
+		ClientSubObjectArray[ClientIndex] = Cast<UTestReplicatedIrisObject>(Client->GetReplicationBridge()->GetReplicatedObject(ServerSubObject->NetRefHandle));
 		UE_NET_ASSERT_NE(ClientSubObjectArray[ClientIndex], nullptr);
 	}
 
@@ -695,12 +695,12 @@ UE_NET_TEST_FIXTURE(FTestConditionalsFixture, WithSkipOwnerSubObjectIsNotReplica
 
 	// Spawn object + subobject on server
 	UTestReplicatedIrisObject* ServerObject = Server->CreateObject(0, 0);
-	UTestReplicatedIrisObject* ServerSubObject = Server->CreateSubObject(ServerObject->NetHandle, 0, 0);
+	UTestReplicatedIrisObject* ServerSubObject = Server->CreateSubObject(ServerObject->NetRefHandle, 0, 0);
 
-	Server->ReplicationBridge->SetSubObjectNetCondition(ServerSubObject->NetHandle, ELifetimeCondition::COND_SkipOwner);
+	Server->ReplicationBridge->SetSubObjectNetCondition(ServerSubObject->NetRefHandle, ELifetimeCondition::COND_SkipOwner);
 
 	// Set owner
-	Server->ReplicationSystem->SetOwningNetConnection(ServerObject->NetHandle, Client->ConnectionIdOnServer);
+	Server->ReplicationSystem->SetOwningNetConnection(ServerObject->NetRefHandle, Client->ConnectionIdOnServer);
 
 	// Set some values
 	ServerSubObject->IntA = 13; 
@@ -711,11 +711,11 @@ UE_NET_TEST_FIXTURE(FTestConditionalsFixture, WithSkipOwnerSubObjectIsNotReplica
 	Server->PostSendUpdate();
 	
 	// Object should have been created on the client
-	const UTestReplicatedIrisObject* ClientObject = Cast<UTestReplicatedIrisObject>(Client->GetReplicationBridge()->GetReplicatedObject(ServerObject->NetHandle));
+	const UTestReplicatedIrisObject* ClientObject = Cast<UTestReplicatedIrisObject>(Client->GetReplicationBridge()->GetReplicatedObject(ServerObject->NetRefHandle));
 	UE_NET_ASSERT_NE(ClientObject, nullptr);
 
 	// SubObject should not have been created on client
-	const UTestReplicatedIrisObject* ClientSubObject = Cast<UTestReplicatedIrisObject>(Client->GetReplicationBridge()->GetReplicatedObject(ServerSubObject->NetHandle));
+	const UTestReplicatedIrisObject* ClientSubObject = Cast<UTestReplicatedIrisObject>(Client->GetReplicationBridge()->GetReplicatedObject(ServerSubObject->NetRefHandle));
 	UE_NET_ASSERT_EQ(ClientSubObject, nullptr);
 }
 
@@ -727,12 +727,12 @@ UE_NET_TEST_FIXTURE(FTestConditionalsFixture, WithSkipOwnerSubObjectReplicatedTo
 
 	// Spawn object + subobject on server
 	UTestReplicatedIrisObject* ServerObject = Server->CreateObject(0, 0);
-	UTestReplicatedIrisObject* ServerSubObject = Server->CreateSubObject(ServerObject->NetHandle, 0, 0);
+	UTestReplicatedIrisObject* ServerSubObject = Server->CreateSubObject(ServerObject->NetRefHandle, 0, 0);
 
-	Server->ReplicationBridge->SetSubObjectNetCondition(ServerSubObject->NetHandle, ELifetimeCondition::COND_SkipOwner);
+	Server->ReplicationBridge->SetSubObjectNetCondition(ServerSubObject->NetRefHandle, ELifetimeCondition::COND_SkipOwner);
 
 	// Set owner
-	Server->ReplicationSystem->SetOwningNetConnection(ServerObject->NetHandle, Client->ConnectionIdOnServer);
+	Server->ReplicationSystem->SetOwningNetConnection(ServerObject->NetRefHandle, Client->ConnectionIdOnServer);
 
 	// Set some values
 	ServerSubObject->IntA = 13; 
@@ -748,19 +748,19 @@ UE_NET_TEST_FIXTURE(FTestConditionalsFixture, WithSkipOwnerSubObjectReplicatedTo
 	Server->PostSendUpdate();
 	
 	// Object should have been created on the client
-	const UTestReplicatedIrisObject* ClientObject = Cast<UTestReplicatedIrisObject>(Client->GetReplicationBridge()->GetReplicatedObject(ServerObject->NetHandle));
+	const UTestReplicatedIrisObject* ClientObject = Cast<UTestReplicatedIrisObject>(Client->GetReplicationBridge()->GetReplicatedObject(ServerObject->NetRefHandle));
 	UE_NET_ASSERT_NE(ClientObject, nullptr);
 
 	// Object should have been created on the client
-	const UTestReplicatedIrisObject* OtherClientObject = Cast<UTestReplicatedIrisObject>(Client->GetReplicationBridge()->GetReplicatedObject(ServerObject->NetHandle));
+	const UTestReplicatedIrisObject* OtherClientObject = Cast<UTestReplicatedIrisObject>(Client->GetReplicationBridge()->GetReplicatedObject(ServerObject->NetRefHandle));
 	UE_NET_ASSERT_NE(OtherClientObject, nullptr);
 
 	// SubObject should not ahve been created on client
-	const UTestReplicatedIrisObject* ClientSubObject = Cast<UTestReplicatedIrisObject>(Client->GetReplicationBridge()->GetReplicatedObject(ServerSubObject->NetHandle));
+	const UTestReplicatedIrisObject* ClientSubObject = Cast<UTestReplicatedIrisObject>(Client->GetReplicationBridge()->GetReplicatedObject(ServerSubObject->NetRefHandle));
 	UE_NET_ASSERT_EQ(ClientSubObject, nullptr);
 
 	// SubObject should not ahve been created on client
-	const UTestReplicatedIrisObject* OtherClientSubObject = Cast<UTestReplicatedIrisObject>(Client->GetReplicationBridge()->GetReplicatedObject(ServerSubObject->NetHandle));
+	const UTestReplicatedIrisObject* OtherClientSubObject = Cast<UTestReplicatedIrisObject>(Client->GetReplicationBridge()->GetReplicatedObject(ServerSubObject->NetRefHandle));
 	UE_NET_ASSERT_EQ(OtherClientSubObject, nullptr);
 }
 
@@ -771,12 +771,12 @@ UE_NET_TEST_FIXTURE(FTestConditionalsFixture, WithSkipOwnerSubObjectReplicatedTo
 
 	// Spawn object + subobject on server
 	UTestReplicatedIrisObject* ServerObject = Server->CreateObject(0, 0);
-	UTestReplicatedIrisObject* ServerSubObject = Server->CreateSubObject(ServerObject->NetHandle, 0, 0);
+	UTestReplicatedIrisObject* ServerSubObject = Server->CreateSubObject(ServerObject->NetRefHandle, 0, 0);
 
-	Server->ReplicationBridge->SetSubObjectNetCondition(ServerSubObject->NetHandle, ELifetimeCondition::COND_SkipOwner);
+	Server->ReplicationBridge->SetSubObjectNetCondition(ServerSubObject->NetRefHandle, ELifetimeCondition::COND_SkipOwner);
 
 	// Set owner
-	Server->ReplicationSystem->SetOwningNetConnection(ServerObject->NetHandle, Client->ConnectionIdOnServer);
+	Server->ReplicationSystem->SetOwningNetConnection(ServerObject->NetRefHandle, Client->ConnectionIdOnServer);
 
 	// Set some values
 	ServerSubObject->IntA = 13; 
@@ -787,17 +787,17 @@ UE_NET_TEST_FIXTURE(FTestConditionalsFixture, WithSkipOwnerSubObjectReplicatedTo
 	Server->PostSendUpdate();
 	
 	// Object should have been created on the client
-	const UTestReplicatedIrisObject* ClientObject = Cast<UTestReplicatedIrisObject>(Client->GetReplicationBridge()->GetReplicatedObject(ServerObject->NetHandle));
+	const UTestReplicatedIrisObject* ClientObject = Cast<UTestReplicatedIrisObject>(Client->GetReplicationBridge()->GetReplicatedObject(ServerObject->NetRefHandle));
 	UE_NET_ASSERT_NE(ClientObject, nullptr);
 
 	// SubObject should not have been created on client
 	{
-		const UTestReplicatedIrisObject* ClientSubObject = Cast<UTestReplicatedIrisObject>(Client->GetReplicationBridge()->GetReplicatedObject(ServerSubObject->NetHandle));
+		const UTestReplicatedIrisObject* ClientSubObject = Cast<UTestReplicatedIrisObject>(Client->GetReplicationBridge()->GetReplicatedObject(ServerSubObject->NetRefHandle));
 		UE_NET_ASSERT_EQ(ClientSubObject, nullptr);
 	}
 
 	// Remove condition
-	Server->ReplicationBridge->SetSubObjectNetCondition(ServerSubObject->NetHandle, ELifetimeCondition::COND_None);
+	Server->ReplicationBridge->SetSubObjectNetCondition(ServerSubObject->NetRefHandle, ELifetimeCondition::COND_None);
 
 	// Send and deliver packet
 	Server->PreSendUpdate();
@@ -806,7 +806,7 @@ UE_NET_TEST_FIXTURE(FTestConditionalsFixture, WithSkipOwnerSubObjectReplicatedTo
 
 	// SubObject should now have been created on client
 	{
-		const UTestReplicatedIrisObject* ClientSubObject = Cast<UTestReplicatedIrisObject>(Client->GetReplicationBridge()->GetReplicatedObject(ServerSubObject->NetHandle));
+		const UTestReplicatedIrisObject* ClientSubObject = Cast<UTestReplicatedIrisObject>(Client->GetReplicationBridge()->GetReplicatedObject(ServerSubObject->NetRefHandle));
 		UE_NET_ASSERT_NE(ClientSubObject, nullptr);
 
 		UE_NET_ASSERT_EQ(ServerSubObject->IntA, ClientSubObject->IntA);
@@ -820,14 +820,14 @@ UE_NET_TEST_FIXTURE(FTestConditionalsFixture, HierarchicalSubObjectIsNotReplicat
 
 	// Spawn object + subobjects on server
 	UTestReplicatedIrisObject* ServerObject = Server->CreateObject(0, 0);
-	UTestReplicatedIrisObject* ServerSubObject = Server->CreateSubObject(ServerObject->NetHandle, 0, 0);
-	UTestReplicatedIrisObject* ServerSubSubObject = Server->CreateSubObject<UTestReplicatedIrisObject>(ServerObject->NetHandle, ServerSubObject->NetHandle, UReplicationBridge::ESubObjectInsertionOrder::ReplicateWith);
+	UTestReplicatedIrisObject* ServerSubObject = Server->CreateSubObject(ServerObject->NetRefHandle, 0, 0);
+	UTestReplicatedIrisObject* ServerSubSubObject = Server->CreateSubObject<UTestReplicatedIrisObject>(ServerObject->NetRefHandle, ServerSubObject->NetRefHandle, UReplicationBridge::ESubObjectInsertionOrder::ReplicateWith);
 
 	// Set condition to only replicate if root is the owning connection
-	Server->ReplicationBridge->SetSubObjectNetCondition(ServerSubObject->NetHandle, ELifetimeCondition::COND_SkipOwner);
+	Server->ReplicationBridge->SetSubObjectNetCondition(ServerSubObject->NetRefHandle, ELifetimeCondition::COND_SkipOwner);
 
 	// Set owner
-	Server->ReplicationSystem->SetOwningNetConnection(ServerObject->NetHandle, Client->ConnectionIdOnServer);
+	Server->ReplicationSystem->SetOwningNetConnection(ServerObject->NetRefHandle, Client->ConnectionIdOnServer);
 
 	// Set some values
 	ServerSubObject->IntA = 13; 
@@ -838,22 +838,22 @@ UE_NET_TEST_FIXTURE(FTestConditionalsFixture, HierarchicalSubObjectIsNotReplicat
 	Server->PostSendUpdate();
 	
 	// Object should have been created on the client
-	const UTestReplicatedIrisObject* ClientObject = Cast<UTestReplicatedIrisObject>(Client->GetReplicationBridge()->GetReplicatedObject(ServerObject->NetHandle));
+	const UTestReplicatedIrisObject* ClientObject = Cast<UTestReplicatedIrisObject>(Client->GetReplicationBridge()->GetReplicatedObject(ServerObject->NetRefHandle));
 	UE_NET_ASSERT_NE(ClientObject, nullptr);
 
 	// SubObject and SubSubObject should not have been created on client
 	{
-		const UTestReplicatedIrisObject* ClientSubObject = Cast<UTestReplicatedIrisObject>(Client->GetReplicationBridge()->GetReplicatedObject(ServerSubObject->NetHandle));
+		const UTestReplicatedIrisObject* ClientSubObject = Cast<UTestReplicatedIrisObject>(Client->GetReplicationBridge()->GetReplicatedObject(ServerSubObject->NetRefHandle));
 		UE_NET_ASSERT_EQ(ClientSubObject, nullptr);
 
-		const UTestReplicatedIrisObject* ClientSubSubObject = Cast<UTestReplicatedIrisObject>(Client->GetReplicationBridge()->GetReplicatedObject(ServerSubSubObject->NetHandle));
+		const UTestReplicatedIrisObject* ClientSubSubObject = Cast<UTestReplicatedIrisObject>(Client->GetReplicationBridge()->GetReplicatedObject(ServerSubSubObject->NetRefHandle));
 		UE_NET_ASSERT_EQ(ClientSubSubObject, nullptr);
 	}
 
 	// Remove condition on ParentSubObject
-	Server->ReplicationBridge->SetSubObjectNetCondition(ServerSubObject->NetHandle, ELifetimeCondition::COND_None);
+	Server->ReplicationBridge->SetSubObjectNetCondition(ServerSubObject->NetRefHandle, ELifetimeCondition::COND_None);
 	// Add condtion on SubSubObject
-	Server->ReplicationBridge->SetSubObjectNetCondition(ServerSubSubObject->NetHandle, ELifetimeCondition::COND_SkipOwner);
+	Server->ReplicationBridge->SetSubObjectNetCondition(ServerSubSubObject->NetRefHandle, ELifetimeCondition::COND_SkipOwner);
 
 	// Send and deliver packet
 	Server->PreSendUpdate();
@@ -862,15 +862,15 @@ UE_NET_TEST_FIXTURE(FTestConditionalsFixture, HierarchicalSubObjectIsNotReplicat
 
 	// SubObject should now be replicated while SubSubObject should not
 	{
-		const UTestReplicatedIrisObject* ClientSubObject = Cast<UTestReplicatedIrisObject>(Client->GetReplicationBridge()->GetReplicatedObject(ServerSubObject->NetHandle));
+		const UTestReplicatedIrisObject* ClientSubObject = Cast<UTestReplicatedIrisObject>(Client->GetReplicationBridge()->GetReplicatedObject(ServerSubObject->NetRefHandle));
 		UE_NET_ASSERT_NE(ClientSubObject, nullptr);
 
-		const UTestReplicatedIrisObject* ClientSubSubObject = Cast<UTestReplicatedIrisObject>(Client->GetReplicationBridge()->GetReplicatedObject(ServerSubSubObject->NetHandle));
+		const UTestReplicatedIrisObject* ClientSubSubObject = Cast<UTestReplicatedIrisObject>(Client->GetReplicationBridge()->GetReplicatedObject(ServerSubSubObject->NetRefHandle));
 		UE_NET_ASSERT_EQ(ClientSubSubObject, nullptr);
 	}
 
 	// Remove condition on SubSubObject
-	Server->ReplicationBridge->SetSubObjectNetCondition(ServerSubSubObject->NetHandle, ELifetimeCondition::COND_None);
+	Server->ReplicationBridge->SetSubObjectNetCondition(ServerSubSubObject->NetRefHandle, ELifetimeCondition::COND_None);
 
 	// Send and deliver packet
 	Server->PreSendUpdate();
@@ -879,10 +879,10 @@ UE_NET_TEST_FIXTURE(FTestConditionalsFixture, HierarchicalSubObjectIsNotReplicat
 
 	// SubObject and SubSubObject should now both be replicated
 	{
-		const UTestReplicatedIrisObject* ClientSubObject = Cast<UTestReplicatedIrisObject>(Client->GetReplicationBridge()->GetReplicatedObject(ServerSubObject->NetHandle));
+		const UTestReplicatedIrisObject* ClientSubObject = Cast<UTestReplicatedIrisObject>(Client->GetReplicationBridge()->GetReplicatedObject(ServerSubObject->NetRefHandle));
 		UE_NET_ASSERT_NE(ClientSubObject, nullptr);
 
-		const UTestReplicatedIrisObject* ClientSubSubObject = Cast<UTestReplicatedIrisObject>(Client->GetReplicationBridge()->GetReplicatedObject(ServerSubSubObject->NetHandle));
+		const UTestReplicatedIrisObject* ClientSubSubObject = Cast<UTestReplicatedIrisObject>(Client->GetReplicationBridge()->GetReplicatedObject(ServerSubSubObject->NetRefHandle));
 		UE_NET_ASSERT_NE(ClientSubSubObject, nullptr);
 	}
 
diff --git a/Engine/Plugins/Runtime/ReplicationSystemTestPlugin/Source/Private/Tests/ReplicationSystem/Filtering/MockNetObjectFilter.h b/Engine/Plugins/Runtime/ReplicationSystemTestPlugin/Source/Private/Tests/ReplicationSystem/Filtering/MockNetObjectFilter.h
--- a/Engine/Plugins/Runtime/ReplicationSystemTestPlugin/Source/Private/Tests/ReplicationSystem/Filtering/MockNetObjectFilter.h
+++ b/Engine/Plugins/Runtime/ReplicationSystemTestPlugin/Source/Private/Tests/ReplicationSystem/Filtering/MockNetObjectFilter.h
@@ -9,7 +9,7 @@
 
 namespace UE::Net::Private
 {
-	typedef uint32 FInternalNetHandle;
+	typedef uint32 FInternalNetRefIndex;
 }
 
 const UE::Net::FRepTag RepTag_NetTest_FilterOut = 0x521855F5DFA298B7ULL;
diff --git a/Engine/Plugins/Runtime/ReplicationSystemTestPlugin/Source/Private/Tests/ReplicationSystem/Filtering/TestFiltering.cpp b/Engine/Plugins/Runtime/ReplicationSystemTestPlugin/Source/Private/Tests/ReplicationSystem/Filtering/TestFiltering.cpp
--- a/Engine/Plugins/Runtime/ReplicationSystemTestPlugin/Source/Private/Tests/ReplicationSystem/Filtering/TestFiltering.cpp
+++ b/Engine/Plugins/Runtime/ReplicationSystemTestPlugin/Source/Private/Tests/ReplicationSystem/Filtering/TestFiltering.cpp
@@ -83,7 +83,7 @@ UE_NET_TEST_FIXTURE(FTestFilteringFixture, OwnerFilterPreventsObjectFromReplicat
 	UReplicatedTestObject* ServerObject = Server->CreateObject(0, 0);
 
 	// Apply owner filter
-	Server->ReplicationSystem->SetFilter(ServerObject->NetHandle, ToOwnerFilterHandle);
+	Server->ReplicationSystem->SetFilter(ServerObject->NetRefHandle, ToOwnerFilterHandle);
 
 	// Send and deliver packet
 	Server->PreSendUpdate();
@@ -91,7 +91,7 @@ UE_NET_TEST_FIXTURE(FTestFilteringFixture, OwnerFilterPreventsObjectFromReplicat
 	Server->PostSendUpdate();
 
 	// Object should not have been created on the client
-	UE_NET_ASSERT_EQ(Client->GetReplicationBridge()->GetReplicatedObject(ServerObject->NetHandle), nullptr);
+	UE_NET_ASSERT_EQ(Client->GetReplicationBridge()->GetReplicatedObject(ServerObject->NetRefHandle), nullptr);
 
 	Server->DestroyObject(ServerObject);
 }
@@ -103,8 +103,8 @@ UE_NET_TEST_FIXTURE(FTestFilteringFixture, OwnerFilterAllowsObjectToReplicateToO
 
 	// Spawn object on server
 	UReplicatedTestObject* ServerObject = Server->CreateObject(0, 0);
-	Server->ReplicationSystem->SetFilter(ServerObject->NetHandle, ToOwnerFilterHandle);
-	Server->ReplicationSystem->SetOwningNetConnection(ServerObject->NetHandle, Client->ConnectionIdOnServer);
+	Server->ReplicationSystem->SetFilter(ServerObject->NetRefHandle, ToOwnerFilterHandle);
+	Server->ReplicationSystem->SetOwningNetConnection(ServerObject->NetRefHandle, Client->ConnectionIdOnServer);
 
 	// Send and deliver packet
 	Server->PreSendUpdate();
@@ -112,7 +112,7 @@ UE_NET_TEST_FIXTURE(FTestFilteringFixture, OwnerFilterAllowsObjectToReplicateToO
 	Server->PostSendUpdate();
 
 	// Object should have been created on the client
-	UE_NET_ASSERT_NE(Client->GetReplicationBridge()->GetReplicatedObject(ServerObject->NetHandle), nullptr);
+	UE_NET_ASSERT_NE(Client->GetReplicationBridge()->GetReplicatedObject(ServerObject->NetRefHandle), nullptr);
 
 	// Destroy object on server and client
 	Server->DestroyObject(ServerObject);
@@ -130,8 +130,8 @@ UE_NET_TEST_FIXTURE(FTestFilteringFixture, OwnerFilterReplicatesOnlyToOwningConn
 
 	// Spawn object on server
 	UReplicatedTestObject* ServerObject = Server->CreateObject(0, 0);
-	Server->ReplicationSystem->SetFilter(ServerObject->NetHandle, ToOwnerFilterHandle);
-	Server->ReplicationSystem->SetOwningNetConnection(ServerObject->NetHandle, ClientArray[LastClientIndex]->ConnectionIdOnServer);
+	Server->ReplicationSystem->SetFilter(ServerObject->NetRefHandle, ToOwnerFilterHandle);
+	Server->ReplicationSystem->SetOwningNetConnection(ServerObject->NetRefHandle, ClientArray[LastClientIndex]->ConnectionIdOnServer);
 
 	// Send and deliver packets
 	Server->PreSendUpdate();
@@ -147,11 +147,11 @@ UE_NET_TEST_FIXTURE(FTestFilteringFixture, OwnerFilterReplicatesOnlyToOwningConn
 		const SIZE_T ClientIndex = &Client - &ClientArray[0];
 		if (ClientIndex == LastClientIndex)
 		{
-			UE_NET_ASSERT_NE(Client->GetReplicationBridge()->GetReplicatedObject(ServerObject->NetHandle), nullptr);
+			UE_NET_ASSERT_NE(Client->GetReplicationBridge()->GetReplicatedObject(ServerObject->NetRefHandle), nullptr);
 		}
 		else
 		{
-			UE_NET_ASSERT_EQ(Client->GetReplicationBridge()->GetReplicatedObject(ServerObject->NetHandle), nullptr);
+			UE_NET_ASSERT_EQ(Client->GetReplicationBridge()->GetReplicatedObject(ServerObject->NetRefHandle), nullptr);
 		}
 	}
 
@@ -178,10 +178,10 @@ UE_NET_TEST_FIXTURE(FTestFilteringFixture, CanChangeOwningConnection)
 	Server->PostSendUpdate();
 
 	// Object should now exist on client
-	UE_NET_ASSERT_NE(Client->GetReplicationBridge()->GetReplicatedObject(ServerObject->NetHandle), nullptr);
+	UE_NET_ASSERT_NE(Client->GetReplicationBridge()->GetReplicatedObject(ServerObject->NetRefHandle), nullptr);
 
 	// Turn on owner filter
-	Server->ReplicationSystem->SetFilter(ServerObject->NetHandle, ToOwnerFilterHandle);
+	Server->ReplicationSystem->SetFilter(ServerObject->NetRefHandle, ToOwnerFilterHandle);
 
 	// Send and deliver packets
 	Server->PreSendUpdate();
@@ -189,10 +189,10 @@ UE_NET_TEST_FIXTURE(FTestFilteringFixture, CanChangeOwningConnection)
 	Server->PostSendUpdate();
 
 	// As object is now filtered it should be deleted on the client
-	UE_NET_ASSERT_EQ(Client->GetReplicationBridge()->GetReplicatedObject(ServerObject->NetHandle), nullptr);
+	UE_NET_ASSERT_EQ(Client->GetReplicationBridge()->GetReplicatedObject(ServerObject->NetRefHandle), nullptr);
 
 	// Set the client as owner
-	Server->ReplicationSystem->SetOwningNetConnection(ServerObject->NetHandle, Client->ConnectionIdOnServer);
+	Server->ReplicationSystem->SetOwningNetConnection(ServerObject->NetRefHandle, Client->ConnectionIdOnServer);
 
 	// Send and deliver packets
 	Server->PreSendUpdate();
@@ -200,10 +200,10 @@ UE_NET_TEST_FIXTURE(FTestFilteringFixture, CanChangeOwningConnection)
 	Server->PostSendUpdate();
 
 	// The object should have been created again
-	UE_NET_ASSERT_NE(Client->GetReplicationBridge()->GetReplicatedObject(ServerObject->NetHandle), nullptr);
+	UE_NET_ASSERT_NE(Client->GetReplicationBridge()->GetReplicatedObject(ServerObject->NetRefHandle), nullptr);
 
 	// Finally, remove the owning connection
-	Server->ReplicationSystem->SetOwningNetConnection(ServerObject->NetHandle, 0);
+	Server->ReplicationSystem->SetOwningNetConnection(ServerObject->NetRefHandle, 0);
 
 	// Send and deliver packets
 	Server->PreSendUpdate();
@@ -211,7 +211,7 @@ UE_NET_TEST_FIXTURE(FTestFilteringFixture, CanChangeOwningConnection)
 	Server->PostSendUpdate();
 
 	// The client is no longer owning the object so it should be deleted
-	UE_NET_ASSERT_EQ(Client->GetReplicationBridge()->GetReplicatedObject(ServerObject->NetHandle), nullptr);
+	UE_NET_ASSERT_EQ(Client->GetReplicationBridge()->GetReplicatedObject(ServerObject->NetRefHandle), nullptr);
 
 	Server->DestroyObject(ServerObject);
 }
@@ -230,10 +230,10 @@ UE_NET_TEST_FIXTURE(FTestFilteringFixture, CanChangeOwnerFilter)
 	Server->PostSendUpdate();
 
 	// Object should now exist on client
-	UE_NET_ASSERT_NE(Client->GetReplicationBridge()->GetReplicatedObject(ServerObject->NetHandle), nullptr);
+	UE_NET_ASSERT_NE(Client->GetReplicationBridge()->GetReplicatedObject(ServerObject->NetRefHandle), nullptr);
 
 	// Turn on owner filter
-	Server->ReplicationSystem->SetFilter(ServerObject->NetHandle, ToOwnerFilterHandle);
+	Server->ReplicationSystem->SetFilter(ServerObject->NetRefHandle, ToOwnerFilterHandle);
 
 	// Send and deliver packets
 	Server->PreSendUpdate();
@@ -241,10 +241,10 @@ UE_NET_TEST_FIXTURE(FTestFilteringFixture, CanChangeOwnerFilter)
 	Server->PostSendUpdate();
 
 	// As object is now filtered it should be deleted on the client
-	UE_NET_ASSERT_EQ(Client->GetReplicationBridge()->GetReplicatedObject(ServerObject->NetHandle), nullptr);
+	UE_NET_ASSERT_EQ(Client->GetReplicationBridge()->GetReplicatedObject(ServerObject->NetRefHandle), nullptr);
 
 	// Remove the owner filter
-	Server->ReplicationSystem->SetFilter(ServerObject->NetHandle, InvalidNetObjectFilterHandle);
+	Server->ReplicationSystem->SetFilter(ServerObject->NetRefHandle, InvalidNetObjectFilterHandle);
 
 	// Send and deliver packets
 	Server->PreSendUpdate();
@@ -252,7 +252,7 @@ UE_NET_TEST_FIXTURE(FTestFilteringFixture, CanChangeOwnerFilter)
 	Server->PostSendUpdate();
 
 	// The object should have been created again
-	UE_NET_ASSERT_NE(Client->GetReplicationBridge()->GetReplicatedObject(ServerObject->NetHandle), nullptr);
+	UE_NET_ASSERT_NE(Client->GetReplicationBridge()->GetReplicatedObject(ServerObject->NetRefHandle), nullptr);
 
 	// Destroy the object
 	Server->DestroyObject(ServerObject);
@@ -277,7 +277,7 @@ UE_NET_TEST_FIXTURE(FTestFilteringFixture, LateAddedSubObjectGetsOwnerPropagated
 		{
 			UReplicatedTestObject* ServerObject = Server->CreateObject(0, 0);
 			ServerObjects.Add(ServerObject);
-			Server->ReplicationSystem->SetOwningNetConnection(ServerObject->NetHandle, Client->ConnectionIdOnServer);
+			Server->ReplicationSystem->SetOwningNetConnection(ServerObject->NetRefHandle, Client->ConnectionIdOnServer);
 		}
 	}
 
@@ -287,14 +287,14 @@ UE_NET_TEST_FIXTURE(FTestFilteringFixture, LateAddedSubObjectGetsOwnerPropagated
 
 	// Create subobject to arbitrary object
 	UReplicatedTestObject* ArbitraryServerObject = ServerObjects[3];
-	UReplicatedTestObject* ServerSubObject = Server->CreateSubObject(ArbitraryServerObject->NetHandle, 1, 1);
+	UReplicatedTestObject* ServerSubObject = Server->CreateSubObject(ArbitraryServerObject->NetRefHandle, 1, 1);
 
 	// Net update
 	Server->PreSendUpdate();
 	Server->PostSendUpdate();
 
 	// Verify subobject owner is as expected
-	UE_NET_ASSERT_EQ(Server->ReplicationSystem->GetOwningNetConnection(ServerSubObject->NetHandle), Client->ConnectionIdOnServer);
+	UE_NET_ASSERT_EQ(Server->ReplicationSystem->GetOwningNetConnection(ServerSubObject->NetRefHandle), Client->ConnectionIdOnServer);
 }
 
 // Connection filtering tests
@@ -308,7 +308,7 @@ UE_NET_TEST_FIXTURE(FTestFilteringFixture, ConnectionFilterPreventsObjectFromRep
 
 	// Apply connection filter
 	TBitArray<> NoConnections;
-	const bool bFilterWasApplied = Server->ReplicationSystem->SetConnectionFilter(ServerObject->NetHandle, NoConnections, ENetFilterStatus::Allow);
+	const bool bFilterWasApplied = Server->ReplicationSystem->SetConnectionFilter(ServerObject->NetRefHandle, NoConnections, ENetFilterStatus::Allow);
 	UE_NET_ASSERT_TRUE(bFilterWasApplied);
 
 	// Send and deliver packet
@@ -317,7 +317,7 @@ UE_NET_TEST_FIXTURE(FTestFilteringFixture, ConnectionFilterPreventsObjectFromRep
 	Server->PostSendUpdate();
 
 	// Object should not have been created on the client
-	UE_NET_ASSERT_EQ(Client->GetReplicationBridge()->GetReplicatedObject(ServerObject->NetHandle), nullptr);
+	UE_NET_ASSERT_EQ(Client->GetReplicationBridge()->GetReplicatedObject(ServerObject->NetRefHandle), nullptr);
 }
 
 UE_NET_TEST_FIXTURE(FTestFilteringFixture, ConnectionFilterAllowsObjectToReplicateToAllowedConnections)
@@ -332,7 +332,7 @@ UE_NET_TEST_FIXTURE(FTestFilteringFixture, ConnectionFilterAllowsObjectToReplica
 	TBitArray<> AllowedConnections;
 	AllowedConnections.Init(false, Client->ConnectionIdOnServer + 1);
 	AllowedConnections[Client->ConnectionIdOnServer] = true;
-	Server->ReplicationSystem->SetConnectionFilter(ServerObject->NetHandle, AllowedConnections, ENetFilterStatus::Allow);
+	Server->ReplicationSystem->SetConnectionFilter(ServerObject->NetRefHandle, AllowedConnections, ENetFilterStatus::Allow);
 
 	// Send and deliver packet
 	Server->PreSendUpdate();
@@ -340,7 +340,7 @@ UE_NET_TEST_FIXTURE(FTestFilteringFixture, ConnectionFilterAllowsObjectToReplica
 	Server->PostSendUpdate();
 
 	// Object should have been created on the client
-	UE_NET_ASSERT_NE(Client->GetReplicationBridge()->GetReplicatedObject(ServerObject->NetHandle), nullptr);
+	UE_NET_ASSERT_NE(Client->GetReplicationBridge()->GetReplicatedObject(ServerObject->NetRefHandle), nullptr);
 }
 
 UE_NET_TEST_FIXTURE(FTestFilteringFixture, ConnectionFilterAllowsObjectToReplicateToLateJoiningConnections)
@@ -351,7 +351,7 @@ UE_NET_TEST_FIXTURE(FTestFilteringFixture, ConnectionFilterAllowsObjectToReplica
 	// Apply filtering that allows all connections
 	TBitArray<> NoConnections;
 	constexpr ENetFilterStatus ReplicationStatus = ENetFilterStatus::Disallow;
-	Server->ReplicationSystem->SetConnectionFilter(ServerObject->NetHandle, NoConnections, ReplicationStatus);
+	Server->ReplicationSystem->SetConnectionFilter(ServerObject->NetRefHandle, NoConnections, ReplicationStatus);
 
 	// Send and deliver packets
 	Server->PreSendUpdate();
@@ -371,7 +371,7 @@ UE_NET_TEST_FIXTURE(FTestFilteringFixture, ConnectionFilterAllowsObjectToReplica
 	// Object should have been created on all clients
 	for (FReplicationSystemTestClient*& Client : ClientArray)
 	{
-		UE_NET_ASSERT_NE(Client->GetReplicationBridge()->GetReplicatedObject(ServerObject->NetHandle), nullptr);
+		UE_NET_ASSERT_NE(Client->GetReplicationBridge()->GetReplicatedObject(ServerObject->NetRefHandle), nullptr);
 	}
 }
 
@@ -385,17 +385,17 @@ UE_NET_TEST_FIXTURE(FTestFilteringFixture, ConnectionFilterAllowsObjectToReplica
 	
 	// Setup group filter
 	FNetObjectGroupHandle GroupHandle = Server->ReplicationSystem->CreateGroup();
-	Server->ReplicationSystem->AddToGroup(GroupHandle, ServerObject0->NetHandle);
+	Server->ReplicationSystem->AddToGroup(GroupHandle, ServerObject0->NetRefHandle);
 	Server->ReplicationSystem->AddGroupFilter(GroupHandle);
 
 	// Spawn object on server
 	UReplicatedTestObject* ServerObject = Server->CreateObject(0, 0);
-	UReplicatedTestObject* ServerSubObject = Server->CreateSubObject(ServerObject->NetHandle, 0, 0);
+	UReplicatedTestObject* ServerSubObject = Server->CreateSubObject(ServerObject->NetRefHandle, 0, 0);
 
 	// Apply filtering that allows all connections
 	TBitArray<> NoConnections;
 	constexpr ENetFilterStatus ReplicationStatus = ENetFilterStatus::Disallow;
-	Server->ReplicationSystem->SetConnectionFilter(ServerObject->NetHandle, NoConnections, ReplicationStatus);
+	Server->ReplicationSystem->SetConnectionFilter(ServerObject->NetRefHandle, NoConnections, ReplicationStatus);
 
 	// Send and deliver packets
 	Server->PreSendUpdate();
@@ -408,11 +408,11 @@ UE_NET_TEST_FIXTURE(FTestFilteringFixture, ConnectionFilterAllowsObjectToReplica
 
 	// Validate status on clients
 	// Object should not have been created on the clients 
-	UE_NET_ASSERT_EQ(Client->GetReplicationBridge()->GetReplicatedObject(ServerObject0->NetHandle), nullptr);
+	UE_NET_ASSERT_EQ(Client->GetReplicationBridge()->GetReplicatedObject(ServerObject0->NetRefHandle), nullptr);
 		
 	// Object should have been crated on all clients
-	UE_NET_ASSERT_NE(Client->GetReplicationBridge()->GetReplicatedObject(ServerObject->NetHandle), nullptr);
-	UE_NET_ASSERT_NE(Client->GetReplicationBridge()->GetReplicatedObject(ServerSubObject->NetHandle), nullptr);
+	UE_NET_ASSERT_NE(Client->GetReplicationBridge()->GetReplicatedObject(ServerObject->NetRefHandle), nullptr);
+	UE_NET_ASSERT_NE(Client->GetReplicationBridge()->GetReplicatedObject(ServerSubObject->NetRefHandle), nullptr);
 }
 
 UE_NET_TEST_FIXTURE(FTestFilteringFixture, ConnectionFilterPreventsObjectFromReplicatingToFilteredOutLateJoiningConnections)
@@ -424,7 +424,7 @@ UE_NET_TEST_FIXTURE(FTestFilteringFixture, ConnectionFilterPreventsObjectFromRep
 	TBitArray<> Connection1(false, 2);
 	Connection1[1] = true;
 	constexpr ENetFilterStatus ReplicationStatus = ENetFilterStatus::Allow;
-	Server->ReplicationSystem->SetConnectionFilter(ServerObject->NetHandle, Connection1, ReplicationStatus);
+	Server->ReplicationSystem->SetConnectionFilter(ServerObject->NetRefHandle, Connection1, ReplicationStatus);
 
 	// Send and deliver packets
 	Server->PreSendUpdate();
@@ -451,11 +451,11 @@ UE_NET_TEST_FIXTURE(FTestFilteringFixture, ConnectionFilterPreventsObjectFromRep
 		const bool bShouldHaveObject = ConnectionId == 1;
 		if (bShouldHaveObject)
 		{
-			UE_NET_ASSERT_NE(Client->GetReplicationBridge()->GetReplicatedObject(ServerObject->NetHandle), nullptr);
+			UE_NET_ASSERT_NE(Client->GetReplicationBridge()->GetReplicatedObject(ServerObject->NetRefHandle), nullptr);
 		}
 		else
 		{
-			UE_NET_ASSERT_EQ(Client->GetReplicationBridge()->GetReplicatedObject(ServerObject->NetHandle), nullptr);
+			UE_NET_ASSERT_EQ(Client->GetReplicationBridge()->GetReplicatedObject(ServerObject->NetRefHandle), nullptr);
 		}
 	}
 }
@@ -474,7 +474,7 @@ UE_NET_TEST_FIXTURE(FTestFilteringFixture, ConnectionFilterReplicatesOnlyToAllow
 	TBitArray<> AllowedConnections;
 	AllowedConnections.Init(false, ConnectionIdForLastClient + 1);
 	AllowedConnections[ConnectionIdForLastClient] = true;
-	Server->ReplicationSystem->SetConnectionFilter(ServerObject->NetHandle, AllowedConnections, ENetFilterStatus::Allow);
+	Server->ReplicationSystem->SetConnectionFilter(ServerObject->NetRefHandle, AllowedConnections, ENetFilterStatus::Allow);
 
 	// Send and deliver packets
 	Server->PreSendUpdate();
@@ -490,11 +490,11 @@ UE_NET_TEST_FIXTURE(FTestFilteringFixture, ConnectionFilterReplicatesOnlyToAllow
 		const SIZE_T ClientIndex = &Client - &ClientArray[0];
 		if (ClientIndex == LastClientIndex)
 		{
-			UE_NET_ASSERT_NE(Client->GetReplicationBridge()->GetReplicatedObject(ServerObject->NetHandle), nullptr);
+			UE_NET_ASSERT_NE(Client->GetReplicationBridge()->GetReplicatedObject(ServerObject->NetRefHandle), nullptr);
 		}
 		else
 		{
-			UE_NET_ASSERT_EQ(Client->GetReplicationBridge()->GetReplicatedObject(ServerObject->NetHandle), nullptr);
+			UE_NET_ASSERT_EQ(Client->GetReplicationBridge()->GetReplicatedObject(ServerObject->NetRefHandle), nullptr);
 		}
 	}
 }
@@ -513,11 +513,11 @@ UE_NET_TEST_FIXTURE(FTestFilteringFixture, CanChangeConnectionFilter)
 	Server->PostSendUpdate();
 
 	// Object should now exist on client
-	UE_NET_ASSERT_NE(Client->GetReplicationBridge()->GetReplicatedObject(ServerObject->NetHandle), nullptr);
+	UE_NET_ASSERT_NE(Client->GetReplicationBridge()->GetReplicatedObject(ServerObject->NetRefHandle), nullptr);
 
 	// Apply connection filtering, not allowing any connection to receive the object
 	TBitArray<> NoConnections;
-	Server->ReplicationSystem->SetConnectionFilter(ServerObject->NetHandle, NoConnections, ENetFilterStatus::Allow);
+	Server->ReplicationSystem->SetConnectionFilter(ServerObject->NetRefHandle, NoConnections, ENetFilterStatus::Allow);
 
 	// Send and deliver packets
 	Server->PreSendUpdate();
@@ -525,13 +525,13 @@ UE_NET_TEST_FIXTURE(FTestFilteringFixture, CanChangeConnectionFilter)
 	Server->PostSendUpdate();
 
 	// As the object is now filtered it should be deleted on the client
-	UE_NET_ASSERT_EQ(Client->GetReplicationBridge()->GetReplicatedObject(ServerObject->NetHandle), nullptr);
+	UE_NET_ASSERT_EQ(Client->GetReplicationBridge()->GetReplicatedObject(ServerObject->NetRefHandle), nullptr);
 
 	// Set the client as an allowed connection
 	TBitArray<> AllowedConnections;
 	AllowedConnections.Init(false, Client->ConnectionIdOnServer + 1);
 	AllowedConnections[Client->ConnectionIdOnServer] = true;
-	Server->ReplicationSystem->SetConnectionFilter(ServerObject->NetHandle, AllowedConnections, ENetFilterStatus::Allow);
+	Server->ReplicationSystem->SetConnectionFilter(ServerObject->NetRefHandle, AllowedConnections, ENetFilterStatus::Allow);
 
 	// Send and deliver packets
 	Server->PreSendUpdate();
@@ -539,10 +539,10 @@ UE_NET_TEST_FIXTURE(FTestFilteringFixture, CanChangeConnectionFilter)
 	Server->PostSendUpdate();
 
 	// The object should have been created again
-	UE_NET_ASSERT_NE(Client->GetReplicationBridge()->GetReplicatedObject(ServerObject->NetHandle), nullptr);
+	UE_NET_ASSERT_NE(Client->GetReplicationBridge()->GetReplicatedObject(ServerObject->NetRefHandle), nullptr);
 
 	// Finally, set the filter to not include any connections again.
-	Server->ReplicationSystem->SetConnectionFilter(ServerObject->NetHandle, NoConnections, ENetFilterStatus::Allow);
+	Server->ReplicationSystem->SetConnectionFilter(ServerObject->NetRefHandle, NoConnections, ENetFilterStatus::Allow);
 
 	// Send and deliver packets
 	Server->PreSendUpdate();
@@ -550,7 +550,7 @@ UE_NET_TEST_FIXTURE(FTestFilteringFixture, CanChangeConnectionFilter)
 	Server->PostSendUpdate();
 
 	// The client is no longer owning the object so it should be deleted
-	UE_NET_ASSERT_EQ(Client->GetReplicationBridge()->GetReplicatedObject(ServerObject->NetHandle), nullptr);
+	UE_NET_ASSERT_EQ(Client->GetReplicationBridge()->GetReplicatedObject(ServerObject->NetRefHandle), nullptr);
 }
 
 UE_NET_TEST_FIXTURE(FTestFilteringFixture, CanToggleConnectionFilter)
@@ -567,11 +567,11 @@ UE_NET_TEST_FIXTURE(FTestFilteringFixture, CanToggleConnectionFilter)
 	Server->PostSendUpdate();
 
 	// Object should now exist on client
-	UE_NET_ASSERT_NE(Client->GetReplicationBridge()->GetReplicatedObject(ServerObject->NetHandle), nullptr);
+	UE_NET_ASSERT_NE(Client->GetReplicationBridge()->GetReplicatedObject(ServerObject->NetRefHandle), nullptr);
 
 	// Apply connection filtering, not allowing any connection to receive the object
 	TBitArray<> NoConnections;
-	Server->ReplicationSystem->SetConnectionFilter(ServerObject->NetHandle, NoConnections, ENetFilterStatus::Allow);
+	Server->ReplicationSystem->SetConnectionFilter(ServerObject->NetRefHandle, NoConnections, ENetFilterStatus::Allow);
 
 	// Send and deliver packets
 	Server->PreSendUpdate();
@@ -579,10 +579,10 @@ UE_NET_TEST_FIXTURE(FTestFilteringFixture, CanToggleConnectionFilter)
 	Server->PostSendUpdate();
 
 	// As object is now filtered it should be deleted on the client
-	UE_NET_ASSERT_EQ(Client->GetReplicationBridge()->GetReplicatedObject(ServerObject->NetHandle), nullptr);
+	UE_NET_ASSERT_EQ(Client->GetReplicationBridge()->GetReplicatedObject(ServerObject->NetRefHandle), nullptr);
 
 	// Remove the connection filter
-	Server->ReplicationSystem->SetFilter(ServerObject->NetHandle, InvalidNetObjectFilterHandle);
+	Server->ReplicationSystem->SetFilter(ServerObject->NetRefHandle, InvalidNetObjectFilterHandle);
 
 	// Send and deliver packets
 	Server->PreSendUpdate();
@@ -590,7 +590,7 @@ UE_NET_TEST_FIXTURE(FTestFilteringFixture, CanToggleConnectionFilter)
 	Server->PostSendUpdate();
 
 	// The object should have been created again
-	UE_NET_ASSERT_NE(Client->GetReplicationBridge()->GetReplicatedObject(ServerObject->NetHandle), nullptr);
+	UE_NET_ASSERT_NE(Client->GetReplicationBridge()->GetReplicatedObject(ServerObject->NetRefHandle), nullptr);
 
 	// Destroy the object
 	Server->DestroyObject(ServerObject);
@@ -614,7 +614,7 @@ UE_NET_TEST_FIXTURE(FTestFilteringFixture, CanToggleBetweenAllFilters)
 	// Toggle between owner and connection filters
 	{
 		// Apply owner filtering
-		Server->ReplicationSystem->SetFilter(ServerObject->NetHandle, ToOwnerFilterHandle);
+		Server->ReplicationSystem->SetFilter(ServerObject->NetRefHandle, ToOwnerFilterHandle);
 
 		// Send and deliver packets
 		Server->PreSendUpdate();
@@ -622,11 +622,11 @@ UE_NET_TEST_FIXTURE(FTestFilteringFixture, CanToggleBetweenAllFilters)
 		Server->PostSendUpdate();
 
 		// Object should not exist on client
-		UE_NET_ASSERT_EQ(Client->GetReplicationBridge()->GetReplicatedObject(ServerObject->NetHandle), nullptr);
+		UE_NET_ASSERT_EQ(Client->GetReplicationBridge()->GetReplicatedObject(ServerObject->NetRefHandle), nullptr);
 
 		// Apply connection filtering, not allowing any connection to receive the object
 		TBitArray<> NoConnections;
-		Server->ReplicationSystem->SetConnectionFilter(ServerObject->NetHandle, NoConnections, ENetFilterStatus::Allow);
+		Server->ReplicationSystem->SetConnectionFilter(ServerObject->NetRefHandle, NoConnections, ENetFilterStatus::Allow);
 
 		// Send and deliver packets
 		Server->PreSendUpdate();
@@ -634,10 +634,10 @@ UE_NET_TEST_FIXTURE(FTestFilteringFixture, CanToggleBetweenAllFilters)
 		Server->PostSendUpdate();
 
 		// Object should still not exist on client
-		UE_NET_ASSERT_EQ(Client->GetReplicationBridge()->GetReplicatedObject(ServerObject->NetHandle), nullptr);
+		UE_NET_ASSERT_EQ(Client->GetReplicationBridge()->GetReplicatedObject(ServerObject->NetRefHandle), nullptr);
 
 		// Finally test going from connection filtering to owner filtering
-		Server->ReplicationSystem->SetFilter(ServerObject->NetHandle, ToOwnerFilterHandle);
+		Server->ReplicationSystem->SetFilter(ServerObject->NetRefHandle, ToOwnerFilterHandle);
 
 		// Send and deliver packets
 		Server->PreSendUpdate();
@@ -645,7 +645,7 @@ UE_NET_TEST_FIXTURE(FTestFilteringFixture, CanToggleBetweenAllFilters)
 		Server->PostSendUpdate();
 
 		// Object should still not exist on client
-		UE_NET_ASSERT_EQ(Client->GetReplicationBridge()->GetReplicatedObject(ServerObject->NetHandle), nullptr);
+		UE_NET_ASSERT_EQ(Client->GetReplicationBridge()->GetReplicatedObject(ServerObject->NetRefHandle), nullptr);
 	}
 }
 
@@ -659,7 +659,7 @@ UE_NET_TEST_FIXTURE(FTestFilteringFixture, GroupFilterPreventsObjectFromReplicat
 	UReplicatedTestObject* ServerObject = Server->CreateObject(0, 0);
 
 	FNetObjectGroupHandle GroupHandle = Server->ReplicationSystem->CreateGroup();
-	Server->ReplicationSystem->AddToGroup(GroupHandle, ServerObject->NetHandle);
+	Server->ReplicationSystem->AddToGroup(GroupHandle, ServerObject->NetRefHandle);
 	Server->ReplicationSystem->AddGroupFilter(GroupHandle);
 
 	// Filter out objects in group
@@ -671,7 +671,7 @@ UE_NET_TEST_FIXTURE(FTestFilteringFixture, GroupFilterPreventsObjectFromReplicat
 	Server->PostSendUpdate();
 
 	// Object should not have been created on the client
-	UE_NET_ASSERT_EQ(Client->GetReplicationBridge()->GetReplicatedObject(ServerObject->NetHandle), nullptr);
+	UE_NET_ASSERT_EQ(Client->GetReplicationBridge()->GetReplicatedObject(ServerObject->NetRefHandle), nullptr);
 }
 
 UE_NET_TEST_FIXTURE(FTestFilteringFixture, GroupFilterAllowsObjectToReplicate)
@@ -683,7 +683,7 @@ UE_NET_TEST_FIXTURE(FTestFilteringFixture, GroupFilterAllowsObjectToReplicate)
 	UReplicatedTestObject* ServerObject = Server->CreateObject(0, 0);
 
 	FNetObjectGroupHandle GroupHandle = Server->ReplicationSystem->CreateGroup();
-	Server->ReplicationSystem->AddToGroup(GroupHandle, ServerObject->NetHandle);
+	Server->ReplicationSystem->AddToGroup(GroupHandle, ServerObject->NetRefHandle);
 	Server->ReplicationSystem->AddGroupFilter(GroupHandle);
 
 	// Filter out objects in group
@@ -695,7 +695,7 @@ UE_NET_TEST_FIXTURE(FTestFilteringFixture, GroupFilterAllowsObjectToReplicate)
 	Server->PostSendUpdate();
 
 	// Object should not have been created on the client
-	UE_NET_ASSERT_EQ(Client->GetReplicationBridge()->GetReplicatedObject(ServerObject->NetHandle), nullptr);
+	UE_NET_ASSERT_EQ(Client->GetReplicationBridge()->GetReplicatedObject(ServerObject->NetRefHandle), nullptr);
 
 	// Allow replication again
 	Server->ReplicationSystem->SetGroupFilterStatus(GroupHandle, Client->ConnectionIdOnServer, ENetFilterStatus::Allow);
@@ -706,7 +706,7 @@ UE_NET_TEST_FIXTURE(FTestFilteringFixture, GroupFilterAllowsObjectToReplicate)
 	Server->PostSendUpdate();
 
 	// Object should now have been created on the client
-	UE_NET_ASSERT_NE(Client->GetReplicationBridge()->GetReplicatedObject(ServerObject->NetHandle), nullptr);
+	UE_NET_ASSERT_NE(Client->GetReplicationBridge()->GetReplicatedObject(ServerObject->NetRefHandle), nullptr);
 }
 
 UE_NET_TEST_FIXTURE(FTestFilteringFixture, GroupFilterRestoresConnectionConnectionFilter)
@@ -723,7 +723,7 @@ UE_NET_TEST_FIXTURE(FTestFilteringFixture, GroupFilterRestoresConnectionConnecti
 	TBitArray<> AllowedConnections;
 	AllowedConnections.Init(false, ConnectionIdForLastClient + 1);
 	AllowedConnections[ConnectionIdForLastClient] = true;
-	Server->ReplicationSystem->SetConnectionFilter(ServerObject->NetHandle, AllowedConnections, ENetFilterStatus::Allow);
+	Server->ReplicationSystem->SetConnectionFilter(ServerObject->NetRefHandle, AllowedConnections, ENetFilterStatus::Allow);
 
 	// Send and deliver packets
 	Server->PreSendUpdate();
@@ -739,17 +739,17 @@ UE_NET_TEST_FIXTURE(FTestFilteringFixture, GroupFilterRestoresConnectionConnecti
 		const SIZE_T ClientIndex = &Client - &ClientArray[0];
 		if (ClientIndex == LastClientIndex)
 		{
-			UE_NET_ASSERT_NE(Client->GetReplicationBridge()->GetReplicatedObject(ServerObject->NetHandle), nullptr);
+			UE_NET_ASSERT_NE(Client->GetReplicationBridge()->GetReplicatedObject(ServerObject->NetRefHandle), nullptr);
 		}
 		else
 		{
-			UE_NET_ASSERT_EQ(Client->GetReplicationBridge()->GetReplicatedObject(ServerObject->NetHandle), nullptr);
+			UE_NET_ASSERT_EQ(Client->GetReplicationBridge()->GetReplicatedObject(ServerObject->NetRefHandle), nullptr);
 		}
 	}
 
 	// Create and set group filter for last client only
 	FNetObjectGroupHandle GroupHandle = Server->ReplicationSystem->CreateGroup();
-	Server->ReplicationSystem->AddToGroup(GroupHandle, ServerObject->NetHandle);
+	Server->ReplicationSystem->AddToGroup(GroupHandle, ServerObject->NetRefHandle);
 	Server->ReplicationSystem->AddGroupFilter(GroupHandle);
 	Server->ReplicationSystem->SetGroupFilterStatus(GroupHandle, ConnectionIdForLastClient, ENetFilterStatus::Disallow);
 
@@ -764,7 +764,7 @@ UE_NET_TEST_FIXTURE(FTestFilteringFixture, GroupFilterRestoresConnectionConnecti
 	// Object Should now have been destroyed on the last client
 	for (FReplicationSystemTestClient*& Client : ClientArray)
 	{
-		UE_NET_ASSERT_EQ(Client->GetReplicationBridge()->GetReplicatedObject(ServerObject->NetHandle), nullptr);
+		UE_NET_ASSERT_EQ(Client->GetReplicationBridge()->GetReplicatedObject(ServerObject->NetRefHandle), nullptr);
 	}
 
 	// Clear Group filter
@@ -784,11 +784,11 @@ UE_NET_TEST_FIXTURE(FTestFilteringFixture, GroupFilterRestoresConnectionConnecti
 		const SIZE_T ClientIndex = &Client - &ClientArray[0];
 		if (ClientIndex == LastClientIndex)
 		{
-			UE_NET_ASSERT_NE(Client->GetReplicationBridge()->GetReplicatedObject(ServerObject->NetHandle), nullptr);
+			UE_NET_ASSERT_NE(Client->GetReplicationBridge()->GetReplicatedObject(ServerObject->NetRefHandle), nullptr);
 		}
 		else
 		{
-			UE_NET_ASSERT_EQ(Client->GetReplicationBridge()->GetReplicatedObject(ServerObject->NetHandle), nullptr);
+			UE_NET_ASSERT_EQ(Client->GetReplicationBridge()->GetReplicatedObject(ServerObject->NetRefHandle), nullptr);
 		}
 	}
 }
@@ -822,7 +822,7 @@ UE_NET_TEST_FIXTURE(FTestFilteringFixture, DynamicFilterAddObjectAndRemoveObject
 	}
 
 	UReplicatedTestObject* ServerObject = Server->CreateObject(0, 0);
-	Server->ReplicationSystem->SetFilter(ServerObject->NetHandle, MockFilterHandle);
+	Server->ReplicationSystem->SetFilter(ServerObject->NetRefHandle, MockFilterHandle);
 
 	Server->PreSendUpdate();
 	// Filter needs to be set now.
@@ -868,7 +868,7 @@ UE_NET_TEST_FIXTURE(FTestFilteringFixture, DynamicFilterCanAllowObjectToReplicat
 
 	// Spawn object on server and set filter
 	UReplicatedTestObject* ServerObject = Server->CreateObject(0, 0);
-	Server->ReplicationSystem->SetFilter(ServerObject->NetHandle, MockFilterHandle);
+	Server->ReplicationSystem->SetFilter(ServerObject->NetRefHandle, MockFilterHandle);
 
 	// Send and deliver packets
 	Server->PreSendUpdate();
@@ -876,7 +876,7 @@ UE_NET_TEST_FIXTURE(FTestFilteringFixture, DynamicFilterCanAllowObjectToReplicat
 	Server->PostSendUpdate();
 
 	// Object should now exist on client
-	UE_NET_ASSERT_NE(Client->GetReplicationBridge()->GetReplicatedObject(ServerObject->NetHandle), nullptr);
+	UE_NET_ASSERT_NE(Client->GetReplicationBridge()->GetReplicatedObject(ServerObject->NetRefHandle), nullptr);
 }
 
 UE_NET_TEST_FIXTURE(FTestFilteringFixture, DynamicFilterCanDisallowObjectToReplicate)
@@ -895,7 +895,7 @@ UE_NET_TEST_FIXTURE(FTestFilteringFixture, DynamicFilterCanDisallowObjectToRepli
 
 	// Spawn object on server and set filter
 	UReplicatedTestObject* ServerObject = Server->CreateObject(0, 0);
-	Server->ReplicationSystem->SetFilter(ServerObject->NetHandle, MockFilterHandle);
+	Server->ReplicationSystem->SetFilter(ServerObject->NetRefHandle, MockFilterHandle);
 
 	// Send and deliver packets
 	Server->PreSendUpdate();
@@ -903,7 +903,7 @@ UE_NET_TEST_FIXTURE(FTestFilteringFixture, DynamicFilterCanDisallowObjectToRepli
 	Server->PostSendUpdate();
 
 	// Object should not exist on client
-	UE_NET_ASSERT_EQ(Client->GetReplicationBridge()->GetReplicatedObject(ServerObject->NetHandle), nullptr);
+	UE_NET_ASSERT_EQ(Client->GetReplicationBridge()->GetReplicatedObject(ServerObject->NetRefHandle), nullptr);
 }
 
 UE_NET_TEST_FIXTURE(FTestFilteringFixture, SwitchingFiltersCallsRemoveObjectOnPreviousFilter)
@@ -919,7 +919,7 @@ UE_NET_TEST_FIXTURE(FTestFilteringFixture, SwitchingFiltersCallsRemoveObjectOnPr
 
 	// Spawn object on server and set filter
 	UReplicatedTestObject* ServerObject = Server->CreateObject(0, 0);
-	Server->ReplicationSystem->SetFilter(ServerObject->NetHandle, MockFilterHandle);
+	Server->ReplicationSystem->SetFilter(ServerObject->NetRefHandle, MockFilterHandle);
 
 	// Make sure filter is set
 	Server->PreSendUpdate();
@@ -930,7 +930,7 @@ UE_NET_TEST_FIXTURE(FTestFilteringFixture, SwitchingFiltersCallsRemoveObjectOnPr
 		MockNetObjectFilter->ResetFunctionCallStatus();
 
 		const TBitArray<> NoConnections;
-		Server->ReplicationSystem->SetConnectionFilter(ServerObject->NetHandle, NoConnections, ENetFilterStatus::Disallow);
+		Server->ReplicationSystem->SetConnectionFilter(ServerObject->NetRefHandle, NoConnections, ENetFilterStatus::Disallow);
 
 		Server->PreSendUpdate();
 		Server->PostSendUpdate();
@@ -957,9 +957,9 @@ UE_NET_TEST_FIXTURE(FTestFilteringFixture, SubObjectsAreReplicatedWhenOwnerDynam
 
 	// Spawn object on server and set filter
 	UReplicatedTestObject* ServerObject = Server->CreateObject(0, 0);
-	Server->ReplicationSystem->SetFilter(ServerObject->NetHandle, MockFilterHandle);
+	Server->ReplicationSystem->SetFilter(ServerObject->NetRefHandle, MockFilterHandle);
 
-	UReplicatedTestObject* ServerSubObject = Server->CreateSubObject(ServerObject->NetHandle, 0, 0);
+	UReplicatedTestObject* ServerSubObject = Server->CreateSubObject(ServerObject->NetRefHandle, 0, 0);
 
 	// Send and deliver packets
 	Server->PreSendUpdate();
@@ -967,8 +967,8 @@ UE_NET_TEST_FIXTURE(FTestFilteringFixture, SubObjectsAreReplicatedWhenOwnerDynam
 	Server->PostSendUpdate();
 
 	// Check that both the object and subobject exist.
-	UE_NET_ASSERT_NE(Client->GetReplicationBridge()->GetReplicatedObject(ServerObject->NetHandle), nullptr);
-	UE_NET_ASSERT_NE(Client->GetReplicationBridge()->GetReplicatedObject(ServerSubObject->NetHandle), nullptr);
+	UE_NET_ASSERT_NE(Client->GetReplicationBridge()->GetReplicatedObject(ServerObject->NetRefHandle), nullptr);
+	UE_NET_ASSERT_NE(Client->GetReplicationBridge()->GetReplicatedObject(ServerSubObject->NetRefHandle), nullptr);
 }
 
 UE_NET_TEST_FIXTURE(FTestFilteringFixture, SubObjectsAreNotReplicatedWhenOwnerDynamicFilterDisallows)
@@ -987,9 +987,9 @@ UE_NET_TEST_FIXTURE(FTestFilteringFixture, SubObjectsAreNotReplicatedWhenOwnerDy
 
 	// Spawn object on server and set filter
 	UReplicatedTestObject* ServerObject = Server->CreateObject(0, 0);
-	Server->ReplicationSystem->SetFilter(ServerObject->NetHandle, MockFilterHandle);
+	Server->ReplicationSystem->SetFilter(ServerObject->NetRefHandle, MockFilterHandle);
 
-	UReplicatedTestObject* ServerSubObject = Server->CreateSubObject(ServerObject->NetHandle, 0, 0);
+	UReplicatedTestObject* ServerSubObject = Server->CreateSubObject(ServerObject->NetRefHandle, 0, 0);
 
 	// Send and deliver packets
 	Server->PreSendUpdate();
@@ -997,8 +997,8 @@ UE_NET_TEST_FIXTURE(FTestFilteringFixture, SubObjectsAreNotReplicatedWhenOwnerDy
 	Server->PostSendUpdate();
 
 	// Check that neither the object nor the subobject exist.
-	UE_NET_ASSERT_EQ(Client->GetReplicationBridge()->GetReplicatedObject(ServerObject->NetHandle), nullptr);
-	UE_NET_ASSERT_EQ(Client->GetReplicationBridge()->GetReplicatedObject(ServerSubObject->NetHandle), nullptr);
+	UE_NET_ASSERT_EQ(Client->GetReplicationBridge()->GetReplicatedObject(ServerObject->NetRefHandle), nullptr);
+	UE_NET_ASSERT_EQ(Client->GetReplicationBridge()->GetReplicatedObject(ServerSubObject->NetRefHandle), nullptr);
 }
 
 UE_NET_TEST_FIXTURE(FTestFilteringFixture, DependentObjectIsUnaffectedByDynamicFilter)
@@ -1018,7 +1018,7 @@ UE_NET_TEST_FIXTURE(FTestFilteringFixture, DependentObjectIsUnaffectedByDynamicF
 	// Spawn objects on server and set filter on dependent object
 	UReplicatedTestObject* ServerObject = Server->CreateObject(0, 0);
 	UReplicatedTestObject* ServerFutureDependentObject = Server->CreateObject(0, 0);
-	Server->ReplicationSystem->SetFilter(ServerFutureDependentObject->NetHandle, MockFilterHandle);
+	Server->ReplicationSystem->SetFilter(ServerFutureDependentObject->NetRefHandle, MockFilterHandle);
 
 	// Send and deliver packets
 	Server->PreSendUpdate();
@@ -1026,11 +1026,11 @@ UE_NET_TEST_FIXTURE(FTestFilteringFixture, DependentObjectIsUnaffectedByDynamicF
 	Server->PostSendUpdate();
 
 	// We expect the object to exist and the future dependent object not to exist
-	UE_NET_ASSERT_NE(Client->GetReplicationBridge()->GetReplicatedObject(ServerObject->NetHandle), nullptr);
-	UE_NET_ASSERT_EQ(Client->GetReplicationBridge()->GetReplicatedObject(ServerFutureDependentObject->NetHandle), nullptr);
+	UE_NET_ASSERT_NE(Client->GetReplicationBridge()->GetReplicatedObject(ServerObject->NetRefHandle), nullptr);
+	UE_NET_ASSERT_EQ(Client->GetReplicationBridge()->GetReplicatedObject(ServerFutureDependentObject->NetRefHandle), nullptr);
 
 	// Make dependent object and make sure it's now replicated.
-	Server->ReplicationBridge->AddDependentObject(ServerObject->NetHandle, ServerFutureDependentObject->NetHandle);
+	Server->ReplicationBridge->AddDependentObject(ServerObject->NetRefHandle, ServerFutureDependentObject->NetRefHandle);
 	UReplicatedTestObject* ServerDependentObject = ServerFutureDependentObject;
 
 	// Send and deliver packets
@@ -1039,10 +1039,10 @@ UE_NET_TEST_FIXTURE(FTestFilteringFixture, DependentObjectIsUnaffectedByDynamicF
 	Server->PostSendUpdate();
 
 	// We now expect the dependent object to exist
-	UE_NET_ASSERT_NE(Client->GetReplicationBridge()->GetReplicatedObject(ServerDependentObject->NetHandle), nullptr);
+	UE_NET_ASSERT_NE(Client->GetReplicationBridge()->GetReplicatedObject(ServerDependentObject->NetRefHandle), nullptr);
 
 	// Remove dependency and make sure the formerly dependent object is removed from the client
-	Server->ReplicationBridge->RemoveDependentObject(ServerObject->NetHandle, ServerDependentObject->NetHandle);
+	Server->ReplicationBridge->RemoveDependentObject(ServerObject->NetRefHandle, ServerDependentObject->NetRefHandle);
 	UReplicatedTestObject* ServerFormerDependentObject = ServerDependentObject;
 
 	// Send and deliver packets
@@ -1051,7 +1051,7 @@ UE_NET_TEST_FIXTURE(FTestFilteringFixture, DependentObjectIsUnaffectedByDynamicF
 	Server->PostSendUpdate();
 
 	// We expect the former dependent object not to exist
-	UE_NET_ASSERT_EQ(Client->GetReplicationBridge()->GetReplicatedObject(ServerFormerDependentObject->NetHandle), nullptr);
+	UE_NET_ASSERT_EQ(Client->GetReplicationBridge()->GetReplicatedObject(ServerFormerDependentObject->NetRefHandle), nullptr);
 }
 
 UE_NET_TEST_FIXTURE(FTestFilteringFixture, NestedDependentObjectIsFilteredAsParentsOrIndependent)
@@ -1073,8 +1073,8 @@ UE_NET_TEST_FIXTURE(FTestFilteringFixture, NestedDependentObjectIsFilteredAsPare
 	UReplicatedTestObject* ServerFutureDependentObject = Server->CreateObject(0, 0);
 	UReplicatedTestObject* ServerFutureNestedDependentObject = Server->CreateObject(0, 0);
 
-	Server->ReplicationSystem->SetFilter(ServerFutureDependentObject->NetHandle, MockFilterHandle);
-	Server->ReplicationSystem->SetFilter(ServerFutureNestedDependentObject->NetHandle, MockFilterHandle);
+	Server->ReplicationSystem->SetFilter(ServerFutureDependentObject->NetRefHandle, MockFilterHandle);
+	Server->ReplicationSystem->SetFilter(ServerFutureNestedDependentObject->NetRefHandle, MockFilterHandle);
 
 	// Send and deliver packets
 	Server->PreSendUpdate();
@@ -1082,13 +1082,13 @@ UE_NET_TEST_FIXTURE(FTestFilteringFixture, NestedDependentObjectIsFilteredAsPare
 	Server->PostSendUpdate();
 
 	// We expect the object to exist and the future dependent and future dependent objects not to exist
-	UE_NET_ASSERT_NE(Client->GetReplicationBridge()->GetReplicatedObject(ServerObject->NetHandle), nullptr);
-	UE_NET_ASSERT_EQ(Client->GetReplicationBridge()->GetReplicatedObject(ServerFutureDependentObject->NetHandle), nullptr);
-	UE_NET_ASSERT_EQ(Client->GetReplicationBridge()->GetReplicatedObject(ServerFutureNestedDependentObject->NetHandle), nullptr);
+	UE_NET_ASSERT_NE(Client->GetReplicationBridge()->GetReplicatedObject(ServerObject->NetRefHandle), nullptr);
+	UE_NET_ASSERT_EQ(Client->GetReplicationBridge()->GetReplicatedObject(ServerFutureDependentObject->NetRefHandle), nullptr);
+	UE_NET_ASSERT_EQ(Client->GetReplicationBridge()->GetReplicatedObject(ServerFutureNestedDependentObject->NetRefHandle), nullptr);
 
 	// Make dependent objects and make sure they're now replicated.
-	Server->ReplicationBridge->AddDependentObject(ServerObject->NetHandle, ServerFutureDependentObject->NetHandle);
-	Server->ReplicationBridge->AddDependentObject(ServerFutureDependentObject->NetHandle, ServerFutureNestedDependentObject->NetHandle);
+	Server->ReplicationBridge->AddDependentObject(ServerObject->NetRefHandle, ServerFutureDependentObject->NetRefHandle);
+	Server->ReplicationBridge->AddDependentObject(ServerFutureDependentObject->NetRefHandle, ServerFutureNestedDependentObject->NetRefHandle);
 
 	UReplicatedTestObject* ServerDependentObject = ServerFutureDependentObject;
 	UReplicatedTestObject* ServerNestedDependentObject = ServerFutureNestedDependentObject;
@@ -1099,11 +1099,11 @@ UE_NET_TEST_FIXTURE(FTestFilteringFixture, NestedDependentObjectIsFilteredAsPare
 	Server->PostSendUpdate();
 
 	// We now expect the dependent objects to exist
-	UE_NET_ASSERT_NE(Client->GetReplicationBridge()->GetReplicatedObject(ServerDependentObject->NetHandle), nullptr);
-	UE_NET_ASSERT_NE(Client->GetReplicationBridge()->GetReplicatedObject(ServerNestedDependentObject->NetHandle), nullptr);
+	UE_NET_ASSERT_NE(Client->GetReplicationBridge()->GetReplicatedObject(ServerDependentObject->NetRefHandle), nullptr);
+	UE_NET_ASSERT_NE(Client->GetReplicationBridge()->GetReplicatedObject(ServerNestedDependentObject->NetRefHandle), nullptr);
 
 	// Remove dependency on root and make sure the formerly dependent object is removed from the client thanks to the filter.
-	Server->ReplicationBridge->RemoveDependentObject(ServerObject->NetHandle, ServerDependentObject->NetHandle);
+	Server->ReplicationBridge->RemoveDependentObject(ServerObject->NetRefHandle, ServerDependentObject->NetRefHandle);
 	UReplicatedTestObject* ServerFormerDependentObject = ServerDependentObject;
 
 	// Send and deliver packets
@@ -1112,13 +1112,13 @@ UE_NET_TEST_FIXTURE(FTestFilteringFixture, NestedDependentObjectIsFilteredAsPare
 	Server->PostSendUpdate();
 
 	// We expect the former dependent object to not to exist, thanks to the filter
-	UE_NET_ASSERT_EQ(Client->GetReplicationBridge()->GetReplicatedObject(ServerFormerDependentObject->NetHandle), nullptr);
+	UE_NET_ASSERT_EQ(Client->GetReplicationBridge()->GetReplicatedObject(ServerFormerDependentObject->NetRefHandle), nullptr);
 
 	// As the former dependent object is filtered out it's ok for the nested dependent object to be filtered out.
-	UE_NET_ASSERT_EQ(Client->GetReplicationBridge()->GetReplicatedObject(ServerNestedDependentObject->NetHandle), nullptr);
+	UE_NET_ASSERT_EQ(Client->GetReplicationBridge()->GetReplicatedObject(ServerNestedDependentObject->NetRefHandle), nullptr);
 
 	// Remove filter on the nested dependent object
-	Server->ReplicationSystem->SetFilter(ServerNestedDependentObject->NetHandle, InvalidNetObjectFilterHandle);
+	Server->ReplicationSystem->SetFilter(ServerNestedDependentObject->NetRefHandle, InvalidNetObjectFilterHandle);
 
 	// Send and deliver packets
 	Server->PreSendUpdate();
@@ -1126,8 +1126,8 @@ UE_NET_TEST_FIXTURE(FTestFilteringFixture, NestedDependentObjectIsFilteredAsPare
 	Server->PostSendUpdate();
 
 	// Verify that dependent object no longer is filtered out even though its parent is
-	UE_NET_ASSERT_EQ(Client->GetReplicationBridge()->GetReplicatedObject(ServerFormerDependentObject->NetHandle), nullptr);
-	UE_NET_ASSERT_NE(Client->GetReplicationBridge()->GetReplicatedObject(ServerNestedDependentObject->NetHandle), nullptr);
+	UE_NET_ASSERT_EQ(Client->GetReplicationBridge()->GetReplicatedObject(ServerFormerDependentObject->NetRefHandle), nullptr);
+	UE_NET_ASSERT_NE(Client->GetReplicationBridge()->GetReplicatedObject(ServerNestedDependentObject->NetRefHandle), nullptr);
 }
 
 UE_NET_TEST_FIXTURE(FTestFilteringFixture, ObjectGetsFilterOutSettingFromStart)
@@ -1146,7 +1146,7 @@ UE_NET_TEST_FIXTURE(FTestFilteringFixture, ObjectGetsFilterOutSettingFromStart)
 
 	// Create object and set filter
 	UTestFilteringObject* ServerObject = Server->CreateObject<UTestFilteringObject>();
-	Server->ReplicationSystem->SetFilter(ServerObject->NetHandle, MockFilterHandle);
+	Server->ReplicationSystem->SetFilter(ServerObject->NetRefHandle, MockFilterHandle);
 
 	// We want the object to be filtered out
 	constexpr bool bFilterOut = true;
@@ -1158,7 +1158,7 @@ UE_NET_TEST_FIXTURE(FTestFilteringFixture, ObjectGetsFilterOutSettingFromStart)
 	Server->PostSendUpdate();
 
 	// Check that the object does not exist on the client.
-	UE_NET_ASSERT_EQ(Client->GetReplicationBridge()->GetReplicatedObject(ServerObject->NetHandle), nullptr);
+	UE_NET_ASSERT_EQ(Client->GetReplicationBridge()->GetReplicatedObject(ServerObject->NetRefHandle), nullptr);
 }
 
 UE_NET_TEST_FIXTURE(FTestFilteringFixture, ObjectGetsUpdatedFilterOutSetting)
@@ -1186,7 +1186,7 @@ UE_NET_TEST_FIXTURE(FTestFilteringFixture, ObjectGetsUpdatedFilterOutSetting)
 
 	// Create object and set filter
 	UTestFilteringObject* ServerObject = Server->CreateObject<UTestFilteringObject>();
-	Server->ReplicationSystem->SetFilter(ServerObject->NetHandle, MockFilterHandle);
+	Server->ReplicationSystem->SetFilter(ServerObject->NetRefHandle, MockFilterHandle);
 
 	// We don't want the object to be filtered out
 	ServerObject->SetFilterOut(false);
@@ -1197,7 +1197,7 @@ UE_NET_TEST_FIXTURE(FTestFilteringFixture, ObjectGetsUpdatedFilterOutSetting)
 	Server->PostSendUpdate();
 
 	// Check that the object exists on the client.
-	UE_NET_ASSERT_NE(Client->GetReplicationBridge()->GetReplicatedObject(ServerObject->NetHandle), nullptr);
+	UE_NET_ASSERT_NE(Client->GetReplicationBridge()->GetReplicatedObject(ServerObject->NetRefHandle), nullptr);
 
 	// Mark the object to be filtered out.
 	ServerObject->SetFilterOut(true);
@@ -1208,7 +1208,7 @@ UE_NET_TEST_FIXTURE(FTestFilteringFixture, ObjectGetsUpdatedFilterOutSetting)
 	Server->PostSendUpdate();
 
 	// Check that the object does not exist on the client.
-	UE_NET_ASSERT_EQ(Client->GetReplicationBridge()->GetReplicatedObject(ServerObject->NetHandle), nullptr);
+	UE_NET_ASSERT_EQ(Client->GetReplicationBridge()->GetReplicatedObject(ServerObject->NetRefHandle), nullptr);
 }
 
 UE_NET_TEST_FIXTURE(FTestFilteringFixture, DynamicFilteredOutSubObjectsAreResetWhenIndexIsReused)
@@ -1227,11 +1227,11 @@ UE_NET_TEST_FIXTURE(FTestFilteringFixture, DynamicFilteredOutSubObjectsAreResetW
 
 	// Spawn object on server and set filter
 	UReplicatedTestObject* ServerObject = Server->CreateObject(0, 0);
-	Server->ReplicationSystem->SetFilter(ServerObject->NetHandle, MockFilterHandle);
+	Server->ReplicationSystem->SetFilter(ServerObject->NetRefHandle, MockFilterHandle);
 
 	// Create and destroy subobject
 	{
-		UReplicatedTestObject* ServerSubObject = Server->CreateSubObject(ServerObject->NetHandle, 0, 0);
+		UReplicatedTestObject* ServerSubObject = Server->CreateSubObject(ServerObject->NetRefHandle, 0, 0);
 		Server->DestroyObject(ServerSubObject);
 	}
 
@@ -1241,7 +1241,7 @@ UE_NET_TEST_FIXTURE(FTestFilteringFixture, DynamicFilteredOutSubObjectsAreResetW
 	Server->PostSendUpdate();
 
 	// Object should not exist on client
-	UE_NET_ASSERT_EQ(Client->GetReplicationBridge()->GetReplicatedObject(ServerObject->NetHandle), nullptr);
+	UE_NET_ASSERT_EQ(Client->GetReplicationBridge()->GetReplicatedObject(ServerObject->NetRefHandle), nullptr);
 
 	// Create new object which should get the same internal index as the destroyed SubObject
 	UReplicatedTestObject* ServerObject2 = Server->CreateObject(0, 0);
@@ -1252,7 +1252,7 @@ UE_NET_TEST_FIXTURE(FTestFilteringFixture, DynamicFilteredOutSubObjectsAreResetW
 	Server->PostSendUpdate();
 
 	// Object should exist on client
-	UE_NET_ASSERT_NE(Client->GetReplicationBridge()->GetReplicatedObject(ServerObject2->NetHandle), nullptr);
+	UE_NET_ASSERT_NE(Client->GetReplicationBridge()->GetReplicatedObject(ServerObject2->NetRefHandle), nullptr);
 }
 
 
