Iris - Implemented globally unique NetHandles.

diff --git a/Engine/Plugins/Runtime/ReplicationSystemTestPlugin/Source/Private/Tests/ReplicationSystem/Filtering/TestGroups.cpp b/Engine/Plugins/Runtime/ReplicationSystemTestPlugin/Source/Private/Tests/ReplicationSystem/Filtering/TestGroups.cpp
--- a/Engine/Plugins/Runtime/ReplicationSystemTestPlugin/Source/Private/Tests/ReplicationSystem/Filtering/TestGroups.cpp
+++ b/Engine/Plugins/Runtime/ReplicationSystemTestPlugin/Source/Private/Tests/ReplicationSystem/Filtering/TestGroups.cpp
@@ -52,10 +52,10 @@ protected:
 		}
 	}
 
-	void VerifyObjectFilterStatus(FNetHandle Handle, uint32 ConnectionId, bool bExpectedReplicationAllowed)
+	void VerifyObjectFilterStatus(FNetRefHandle Handle, uint32 ConnectionId, bool bExpectedReplicationAllowed)
 	{
-		const uint32 InternalIndex = Server->GetReplicationSystem()->GetReplicationSystemInternal()->GetNetHandleManager().GetInternalIndex(Handle);
-		UE_NET_ASSERT_NE((uint32)FNetHandleManager::InvalidInternalIndex, InternalIndex);
+		const uint32 InternalIndex = Server->GetReplicationSystem()->GetReplicationSystemInternal()->GetNetRefHandleManager().GetInternalIndex(Handle);
+		UE_NET_ASSERT_NE((uint32)FNetRefHandleManager::InvalidInternalIndex, InternalIndex);
 
 		// Verify expected filter status
 		const FNetBitArrayView ObjectGroupFilter = ServerFiltering->GetGroupFilteredOutObjects(ConnectionId);
@@ -90,19 +90,19 @@ UE_NET_TEST_FIXTURE(FTestGroupsFixture, PublicGroupAPI)
 
 		FNetObjectGroupHandle GroupHandle = ReplicationSystem->CreateGroup();
 
-		UE_NET_ASSERT_FALSE(ReplicationSystem->IsInGroup(GroupHandle, ServerObject->NetHandle));
+		UE_NET_ASSERT_FALSE(ReplicationSystem->IsInGroup(GroupHandle, ServerObject->NetRefHandle));
 
-		ReplicationSystem->AddToGroup(GroupHandle, ServerObject->NetHandle);
+		ReplicationSystem->AddToGroup(GroupHandle, ServerObject->NetRefHandle);
 
-		UE_NET_ASSERT_TRUE(ReplicationSystem->IsInGroup(GroupHandle, ServerObject->NetHandle));
+		UE_NET_ASSERT_TRUE(ReplicationSystem->IsInGroup(GroupHandle, ServerObject->NetRefHandle));
 
-		ReplicationSystem->RemoveFromGroup(GroupHandle, ServerObject->NetHandle);
+		ReplicationSystem->RemoveFromGroup(GroupHandle, ServerObject->NetRefHandle);
 
-		UE_NET_ASSERT_FALSE(ReplicationSystem->IsInGroup(GroupHandle, ServerObject->NetHandle));
+		UE_NET_ASSERT_FALSE(ReplicationSystem->IsInGroup(GroupHandle, ServerObject->NetRefHandle));
 
 		ReplicationSystem->DestroyGroup(GroupHandle);
 
-		UE_NET_ASSERT_FALSE(ReplicationSystem->IsInGroup(GroupHandle, ServerObject->NetHandle));
+		UE_NET_ASSERT_FALSE(ReplicationSystem->IsInGroup(GroupHandle, ServerObject->NetRefHandle));
 	}
 }
 
@@ -118,28 +118,28 @@ UE_NET_TEST_FIXTURE(FTestGroupsFixture, PublicGroupAPIMemberOfMultipleGroups)
 
 	UReplicatedTestObject* ServerObject = Server->CreateObject(0, 0);
 
-	UE_NET_ASSERT_FALSE(ReplicationSystem->IsInGroup(GroupHandle0, ServerObject->NetHandle));
-	UE_NET_ASSERT_FALSE(ReplicationSystem->IsInGroup(GroupHandle1, ServerObject->NetHandle));
+	UE_NET_ASSERT_FALSE(ReplicationSystem->IsInGroup(GroupHandle0, ServerObject->NetRefHandle));
+	UE_NET_ASSERT_FALSE(ReplicationSystem->IsInGroup(GroupHandle1, ServerObject->NetRefHandle));
 
-	ReplicationSystem->AddToGroup(GroupHandle0, ServerObject->NetHandle);
+	ReplicationSystem->AddToGroup(GroupHandle0, ServerObject->NetRefHandle);
 
-	UE_NET_ASSERT_TRUE(ReplicationSystem->IsInGroup(GroupHandle0, ServerObject->NetHandle));
-	UE_NET_ASSERT_FALSE(ReplicationSystem->IsInGroup(GroupHandle1, ServerObject->NetHandle));
+	UE_NET_ASSERT_TRUE(ReplicationSystem->IsInGroup(GroupHandle0, ServerObject->NetRefHandle));
+	UE_NET_ASSERT_FALSE(ReplicationSystem->IsInGroup(GroupHandle1, ServerObject->NetRefHandle));
 
-	ReplicationSystem->AddToGroup(GroupHandle1, ServerObject->NetHandle);
+	ReplicationSystem->AddToGroup(GroupHandle1, ServerObject->NetRefHandle);
 
-	UE_NET_ASSERT_TRUE(ReplicationSystem->IsInGroup(GroupHandle0, ServerObject->NetHandle));
-	UE_NET_ASSERT_TRUE(ReplicationSystem->IsInGroup(GroupHandle1, ServerObject->NetHandle));
+	UE_NET_ASSERT_TRUE(ReplicationSystem->IsInGroup(GroupHandle0, ServerObject->NetRefHandle));
+	UE_NET_ASSERT_TRUE(ReplicationSystem->IsInGroup(GroupHandle1, ServerObject->NetRefHandle));
 
-	ReplicationSystem->RemoveFromGroup(GroupHandle0, ServerObject->NetHandle);
+	ReplicationSystem->RemoveFromGroup(GroupHandle0, ServerObject->NetRefHandle);
 
-	UE_NET_ASSERT_FALSE(ReplicationSystem->IsInGroup(GroupHandle0, ServerObject->NetHandle));
-	UE_NET_ASSERT_TRUE(ReplicationSystem->IsInGroup(GroupHandle1, ServerObject->NetHandle));
+	UE_NET_ASSERT_FALSE(ReplicationSystem->IsInGroup(GroupHandle0, ServerObject->NetRefHandle));
+	UE_NET_ASSERT_TRUE(ReplicationSystem->IsInGroup(GroupHandle1, ServerObject->NetRefHandle));
 
-	ReplicationSystem->RemoveFromGroup(GroupHandle1, ServerObject->NetHandle);
+	ReplicationSystem->RemoveFromGroup(GroupHandle1, ServerObject->NetRefHandle);
 
-	UE_NET_ASSERT_FALSE(ReplicationSystem->IsInGroup(GroupHandle0, ServerObject->NetHandle));
-	UE_NET_ASSERT_FALSE(ReplicationSystem->IsInGroup(GroupHandle1, ServerObject->NetHandle));
+	UE_NET_ASSERT_FALSE(ReplicationSystem->IsInGroup(GroupHandle0, ServerObject->NetRefHandle));
+	UE_NET_ASSERT_FALSE(ReplicationSystem->IsInGroup(GroupHandle1, ServerObject->NetRefHandle));
 }
 
 // Group filtering tests
@@ -155,9 +155,9 @@ UE_NET_TEST_FIXTURE(FTestGroupsFixture, GroupFilterAPI)
 	UE_NET_ASSERT_NE(0U, (uint32)GroupHandle);	
 
 	// Add Objects to group
-	Server->ReplicationSystem->AddToGroup(GroupHandle, ServerObject0->NetHandle);
-	Server->ReplicationSystem->AddToGroup(GroupHandle, ServerObject1->NetHandle);
-	Server->ReplicationSystem->AddToGroup(GroupHandle, ServerObject2->NetHandle);
+	Server->ReplicationSystem->AddToGroup(GroupHandle, ServerObject0->NetRefHandle);
+	Server->ReplicationSystem->AddToGroup(GroupHandle, ServerObject1->NetRefHandle);
+	Server->ReplicationSystem->AddToGroup(GroupHandle, ServerObject2->NetRefHandle);
 	
 	UReplicatedTestObject* ServerObject3 = Server->CreateObject(0, 0);
 	UReplicatedTestObject* ServerObject4 = Server->CreateObject(0, 0);
@@ -168,9 +168,9 @@ UE_NET_TEST_FIXTURE(FTestGroupsFixture, GroupFilterAPI)
 	UE_NET_ASSERT_NE(0U, (uint32)GroupHandle2);	
 
 	// Add Objects to group
-	Server->ReplicationSystem->AddToGroup(GroupHandle2, ServerObject3->NetHandle);
-	Server->ReplicationSystem->AddToGroup(GroupHandle2, ServerObject4->NetHandle);
-	Server->ReplicationSystem->AddToGroup(GroupHandle2, ServerObject5->NetHandle);
+	Server->ReplicationSystem->AddToGroup(GroupHandle2, ServerObject3->NetRefHandle);
+	Server->ReplicationSystem->AddToGroup(GroupHandle2, ServerObject4->NetRefHandle);
+	Server->ReplicationSystem->AddToGroup(GroupHandle2, ServerObject5->NetRefHandle);
 													 
 	UReplicatedTestObject* ServerObject6 = Server->CreateObject(0, 0);
 
@@ -227,7 +227,7 @@ UE_NET_TEST_FIXTURE(FTestGroupsFixture, GroupFilterAPI)
 	UE_NET_ASSERT_TRUE(ObjectGroupFilter.IsNoBitSet());
 
 	// Add Object to group 3 which is set, which should filter out the object
-	Server->ReplicationSystem->AddToGroup(GroupHandle3, ServerObject6->NetHandle);
+	Server->ReplicationSystem->AddToGroup(GroupHandle3, ServerObject6->NetRefHandle);
 
 	// Update filters
 	Filter();
@@ -239,7 +239,7 @@ UE_NET_TEST_FIXTURE(FTestGroupsFixture, GroupFilterAPI)
 	UE_NET_ASSERT_TRUE(ObjectGroupFilter.IsAnyBitSet());
 
 	// Add Object to group 3 which is set, which should filter out the object
-	Server->ReplicationSystem->RemoveFromGroup(GroupHandle3, ServerObject6->NetHandle);
+	Server->ReplicationSystem->RemoveFromGroup(GroupHandle3, ServerObject6->NetRefHandle);
 
 	// Update filters
 	Filter();
@@ -283,7 +283,7 @@ UE_NET_TEST_FIXTURE(FTestGroupsFixture, GroupFilterAPIObjectMemberOfMoreThanOneG
 {
 	// Spawn object on server
 	UReplicatedTestObject* ServerObject0 = Server->CreateObject(0, 0);
-	const FNetHandle ServerHandle = ServerObject0->NetHandle;
+	const FNetRefHandle ServerHandle = ServerObject0->NetRefHandle;
 
 	// Create groups
 	FNetObjectGroupHandle GroupHandle0 = Server->ReplicationSystem->CreateGroup();
@@ -348,7 +348,7 @@ UE_NET_TEST_FIXTURE(FTestGroupsFixture, GroupFilterAPINotFilteredGroup)
 	UReplicatedTestObject* ServerObject2 = Server->CreateObject(0, 0);
 
 	// Add Objects to NotFilteredGroup
-	Server->ReplicationSystem->AddToGroup(NotReplicatedNetObjectGroupHandle, ServerObject0->NetHandle);
+	Server->ReplicationSystem->AddToGroup(NotReplicatedNetObjectGroupHandle, ServerObject0->NetRefHandle);
 
 	// Add client
 	FReplicationSystemTestClient* Client = CreateClient();
@@ -362,7 +362,7 @@ UE_NET_TEST_FIXTURE(FTestGroupsFixture, GroupFilterAPINotFilteredGroup)
 	UE_NET_ASSERT_TRUE(ObjectGroupFilter.IsAnyBitSet());
 
 	// Remove object from NotReplicatedNetObjectGroup
-	Server->ReplicationSystem->RemoveFromGroup(NotReplicatedNetObjectGroupHandle, ServerObject0->NetHandle);
+	Server->ReplicationSystem->RemoveFromGroup(NotReplicatedNetObjectGroupHandle, ServerObject0->NetRefHandle);
 
 	// Update filters
 	Filter();
@@ -376,10 +376,10 @@ UE_NET_TEST_FIXTURE(FTestGroupsFixture, GroupFilterAPISubObjects)
 {
 	// Spawn object on server
 	UReplicatedTestObject* ServerObject = Server->CreateObject(0, 0);
-	const FNetHandle ServerHandle = ServerObject->NetHandle;
+	const FNetRefHandle ServerHandle = ServerObject->NetRefHandle;
 
 	UReplicatedTestObject* ServerSubObject = Server->CreateSubObject(ServerHandle, 0, 0);	
-	const FNetHandle ServerSubObjectHandle = ServerSubObject->NetHandle;
+	const FNetRefHandle ServerSubObjectHandle = ServerSubObject->NetRefHandle;
 	
 	// Create groups
 	FNetObjectGroupHandle GroupHandle = Server->ReplicationSystem->CreateGroup();
@@ -417,7 +417,7 @@ UE_NET_TEST_FIXTURE(FTestGroupsFixture, GroupFilterAPIRemovingSubObjectsRestores
 
 	// Spawn object on server
 	UReplicatedTestObject* ServerObject = Server->CreateObject(0, 0);
-	const FNetHandle ServerHandle = ServerObject->NetHandle;
+	const FNetRefHandle ServerHandle = ServerObject->NetRefHandle;
 
 	// Add Object to group
 	Server->ReplicationSystem->AddToGroup(GroupHandle, ServerHandle);
@@ -425,7 +425,7 @@ UE_NET_TEST_FIXTURE(FTestGroupsFixture, GroupFilterAPIRemovingSubObjectsRestores
 
 	// Create sub object
 	UReplicatedTestObject* ServerSubObject = Server->CreateSubObject(ServerHandle, 0, 0);
-	const FNetHandle ServerSubObjectHandle = ServerSubObject->NetHandle;
+	const FNetRefHandle ServerSubObjectHandle = ServerSubObject->NetRefHandle;
 
 	// destroy subobject in the same frame
 	Server->DestroyObject(ServerSubObject);
@@ -446,7 +446,7 @@ UE_NET_TEST_FIXTURE(FTestGroupsFixture, GroupFilterAPILateAddedSubObjectIsFilter
 {
 	// Spawn object on server
 	UReplicatedTestObject* ServerObject = Server->CreateObject(0, 0);
-	const FNetHandle ServerHandle = ServerObject->NetHandle;
+	const FNetRefHandle ServerHandle = ServerObject->NetRefHandle;
 	
 	// Create groups
 	FNetObjectGroupHandle GroupHandle = Server->ReplicationSystem->CreateGroup();
@@ -463,7 +463,7 @@ UE_NET_TEST_FIXTURE(FTestGroupsFixture, GroupFilterAPILateAddedSubObjectIsFilter
 
 	// Add subobject
 	UReplicatedTestObject* ServerSubObject = Server->CreateSubObject(ServerHandle, 0, 0);	
-	const FNetHandle ServerSubObjectHandle = ServerSubObject->NetHandle;
+	const FNetRefHandle ServerSubObjectHandle = ServerSubObject->NetRefHandle;
 
 	// Update filters
 	Filter();
@@ -480,11 +480,11 @@ UE_NET_TEST_FIXTURE(FTestGroupsFixture, GroupFilterAPISubObjectIsFilteredOutIfOw
 {
 	// Spawn object on server
 	UReplicatedTestObject* ServerObject = Server->CreateObject(0, 0);
-	const FNetHandle ServerHandle = ServerObject->NetHandle;
+	const FNetRefHandle ServerHandle = ServerObject->NetRefHandle;
 
 	// Add subobject
 	UReplicatedTestObject* ServerSubObject = Server->CreateSubObject(ServerHandle, 0, 0);	
-	const FNetHandle ServerSubObjectHandle = ServerSubObject->NetHandle;
+	const FNetRefHandle ServerSubObjectHandle = ServerSubObject->NetRefHandle;
 	
 	// Create groups
 	FNetObjectGroupHandle GroupHandle = Server->ReplicationSystem->CreateGroup();
@@ -527,11 +527,11 @@ UE_NET_TEST_FIXTURE(FTestGroupsFixture, GroupFilterAPISubObjectIsFilteredOutWith
 {
 	// Spawn object on server
 	UReplicatedTestObject* ServerObject = Server->CreateObject(0, 0);
-	const FNetHandle ServerHandle = ServerObject->NetHandle;
+	const FNetRefHandle ServerHandle = ServerObject->NetRefHandle;
 
 	// Add subobject
 	UReplicatedTestObject* ServerSubObject = Server->CreateSubObject(ServerHandle, 0, 0);	
-	const FNetHandle ServerSubObjectHandle = ServerSubObject->NetHandle;
+	const FNetRefHandle ServerSubObjectHandle = ServerSubObject->NetRefHandle;
 	
 	// Create groups
 	FNetObjectGroupHandle GroupHandle = Server->ReplicationSystem->CreateGroup();
diff --git a/Engine/Plugins/Runtime/ReplicationSystemTestPlugin/Source/Private/Tests/ReplicationSystem/NetBlob/MockNetBlob.cpp b/Engine/Plugins/Runtime/ReplicationSystemTestPlugin/Source/Private/Tests/ReplicationSystem/NetBlob/MockNetBlob.cpp
--- a/Engine/Plugins/Runtime/ReplicationSystemTestPlugin/Source/Private/Tests/ReplicationSystem/NetBlob/MockNetBlob.cpp
+++ b/Engine/Plugins/Runtime/ReplicationSystemTestPlugin/Source/Private/Tests/ReplicationSystem/NetBlob/MockNetBlob.cpp
@@ -19,7 +19,7 @@ FMockNetBlob::FMockNetBlob(const FNetBlobCreationInfo& CreationInfo)
 {
 }
 
-void FMockNetBlob::SerializeWithObject(FNetSerializationContext& Context, FNetHandle NetHandle) const
+void FMockNetBlob::SerializeWithObject(FNetSerializationContext& Context, FNetRefHandle RefHandle) const
 {
 	FNetBitStreamWriter* Writer = Context.GetBitStreamWriter();
 
@@ -29,7 +29,7 @@ void FMockNetBlob::SerializeWithObject(FNetSerializationContext& Context, FNetHa
 	Writer->WriteBits(UE_MAKEFOURCC('C', 'N', 'R', 'Y'), 32U);
 }
 
-void FMockNetBlob::DeserializeWithObject(FNetSerializationContext& Context, FNetHandle NetHandle)
+void FMockNetBlob::DeserializeWithObject(FNetSerializationContext& Context, FNetRefHandle RefHandle)
 {
 	FNetBitStreamReader* Reader = Context.GetBitStreamReader();
 	
@@ -162,7 +162,7 @@ TRefCountPtr<UE::Net::FNetBlob> UMockSequentialPartialNetBlobHandler::CreateNetB
 void UMockSequentialPartialNetBlobHandler::OnNetBlobReceived(UE::Net::FNetSerializationContext& Context, const TRefCountPtr<FNetBlob>& NetBlob)
 {
 	++CallCounts.OnNetBlobReceived;
-	Assembler.AddPartialNetBlob(Context, UE::Net::FNetHandle(), reinterpret_cast<const TRefCountPtr<UE::Net::FPartialNetBlob>&>(NetBlob));
+	Assembler.AddPartialNetBlob(Context, UE::Net::FNetRefHandle(), reinterpret_cast<const TRefCountPtr<UE::Net::FPartialNetBlob>&>(NetBlob));
 	if (Assembler.IsReadyToAssemble())
 	{
 		const TRefCountPtr<FNetBlob>& AssembledBlob = Assembler.Assemble(Context);
diff --git a/Engine/Plugins/Runtime/ReplicationSystemTestPlugin/Source/Private/Tests/ReplicationSystem/NetBlob/MockNetBlob.h b/Engine/Plugins/Runtime/ReplicationSystemTestPlugin/Source/Private/Tests/ReplicationSystem/NetBlob/MockNetBlob.h
--- a/Engine/Plugins/Runtime/ReplicationSystemTestPlugin/Source/Private/Tests/ReplicationSystem/NetBlob/MockNetBlob.h
+++ b/Engine/Plugins/Runtime/ReplicationSystemTestPlugin/Source/Private/Tests/ReplicationSystem/NetBlob/MockNetBlob.h
@@ -20,8 +20,8 @@ public:
 	void SetPayloadBitCount(uint32 BitCount) { BlobBitCount = BitCount; }
 
 private:
-	virtual void SerializeWithObject(FNetSerializationContext& Context, FNetHandle NetHandle) const override;
-	virtual void DeserializeWithObject(FNetSerializationContext& Context, FNetHandle NetHandle) override;
+	virtual void SerializeWithObject(FNetSerializationContext& Context, FNetRefHandle RefHandle) const override;
+	virtual void DeserializeWithObject(FNetSerializationContext& Context, FNetRefHandle RefHandle) override;
 
 	virtual void Serialize(FNetSerializationContext& Context) const override;
 	virtual void Deserialize(FNetSerializationContext& Context) override;
diff --git a/Engine/Plugins/Runtime/ReplicationSystemTestPlugin/Source/Private/Tests/ReplicationSystem/NetBlob/MockNetObjectAttachment.cpp b/Engine/Plugins/Runtime/ReplicationSystemTestPlugin/Source/Private/Tests/ReplicationSystem/NetBlob/MockNetObjectAttachment.cpp
--- a/Engine/Plugins/Runtime/ReplicationSystemTestPlugin/Source/Private/Tests/ReplicationSystem/NetBlob/MockNetObjectAttachment.cpp
+++ b/Engine/Plugins/Runtime/ReplicationSystemTestPlugin/Source/Private/Tests/ReplicationSystem/NetBlob/MockNetObjectAttachment.cpp
@@ -18,11 +18,11 @@ FMockNetObjectAttachment::FMockNetObjectAttachment(const FNetBlobCreationInfo& C
 {
 }
 
-void FMockNetObjectAttachment::SerializeWithObject(FNetSerializationContext& Context, FNetHandle NetHandle) const
+void FMockNetObjectAttachment::SerializeWithObject(FNetSerializationContext& Context, FNetRefHandle RefHandle) const
 {
 	FNetBitStreamWriter* Writer = Context.GetBitStreamWriter();
 
-	SerializeSubObjectReference(Context, NetHandle);
+	SerializeSubObjectReference(Context, RefHandle);
 
 	Writer->WriteBits(UE_MAKEFOURCC('O', 'B', 'J', ' '), 32U);
 	Writer->WriteBits(BlobBitCount, 32U);
@@ -30,11 +30,11 @@ void FMockNetObjectAttachment::SerializeWithObject(FNetSerializationContext& Con
 	Writer->WriteBits(UE_MAKEFOURCC('C', 'N', 'R', 'Y'), 32U);
 }
 
-void FMockNetObjectAttachment::DeserializeWithObject(FNetSerializationContext& Context, FNetHandle NetHandle)
+void FMockNetObjectAttachment::DeserializeWithObject(FNetSerializationContext& Context, FNetRefHandle RefHandle)
 {
 	FNetBitStreamReader* Reader = Context.GetBitStreamReader();
 	
-	DeserializeSubObjectReference(Context, NetHandle);
+	DeserializeSubObjectReference(Context, RefHandle);
 
 	const uint32 Header = Reader->ReadBits(32U);
 	if (Header != UE_MAKEFOURCC('O', 'B', 'J', ' '))
diff --git a/Engine/Plugins/Runtime/ReplicationSystemTestPlugin/Source/Private/Tests/ReplicationSystem/NetBlob/MockNetObjectAttachment.h b/Engine/Plugins/Runtime/ReplicationSystemTestPlugin/Source/Private/Tests/ReplicationSystem/NetBlob/MockNetObjectAttachment.h
--- a/Engine/Plugins/Runtime/ReplicationSystemTestPlugin/Source/Private/Tests/ReplicationSystem/NetBlob/MockNetObjectAttachment.h
+++ b/Engine/Plugins/Runtime/ReplicationSystemTestPlugin/Source/Private/Tests/ReplicationSystem/NetBlob/MockNetObjectAttachment.h
@@ -19,8 +19,8 @@ public:
 	void SetPayloadBitCount(uint32 BitCount) { BlobBitCount = BitCount; }
 
 private:
-	virtual void SerializeWithObject(FNetSerializationContext& Context, FNetHandle NetHandle) const override;
-	virtual void DeserializeWithObject(FNetSerializationContext& Context, FNetHandle NetHandle) override;
+	virtual void SerializeWithObject(FNetSerializationContext& Context, FNetRefHandle RefHandle) const override;
+	virtual void DeserializeWithObject(FNetSerializationContext& Context, FNetRefHandle RefHandle) override;
 
 	virtual void Serialize(FNetSerializationContext& Context) const override;
 	virtual void Deserialize(FNetSerializationContext& Context) override;
diff --git a/Engine/Plugins/Runtime/ReplicationSystemTestPlugin/Source/Private/Tests/ReplicationSystem/NetBlob/TestPartialNetBlob.cpp b/Engine/Plugins/Runtime/ReplicationSystemTestPlugin/Source/Private/Tests/ReplicationSystem/NetBlob/TestPartialNetBlob.cpp
--- a/Engine/Plugins/Runtime/ReplicationSystemTestPlugin/Source/Private/Tests/ReplicationSystem/NetBlob/TestPartialNetBlob.cpp
+++ b/Engine/Plugins/Runtime/ReplicationSystemTestPlugin/Source/Private/Tests/ReplicationSystem/NetBlob/TestPartialNetBlob.cpp
@@ -101,13 +101,13 @@ UE_NET_TEST_FIXTURE(FPartialNetBlobTestFixture, TestSplitBlobIsReceivedDespitePa
 	Server->SendAndDeliverTo(Client, DeliverPacket);
 	Server->PostSendUpdate();
 
-	UE_NET_ASSERT_NE(Client->GetReplicationBridge()->GetReplicatedObject(ServerObject->NetHandle), nullptr);
+	UE_NET_ASSERT_NE(Client->GetReplicationBridge()->GetReplicatedObject(ServerObject->NetRefHandle), nullptr);
 
 	// Create attachment
 	{
 		constexpr uint32 PayloadBitCount = 16000;
 		const TRefCountPtr<FNetObjectAttachment>& Attachment = MockNetObjectAttachmentHandler->CreateUnreliableNetObjectAttachment(PayloadBitCount);
-		FNetObjectReference AttachmentTarget = FObjectReferenceCache::MakeNetObjectReference(ServerObject->NetHandle);
+		FNetObjectReference AttachmentTarget = FObjectReferenceCache::MakeNetObjectReference(ServerObject->NetRefHandle);
 		Server->GetReplicationSystem()->QueueNetObjectAttachment(Client->ConnectionIdOnServer, AttachmentTarget, Attachment);
 	}
 
diff --git a/Engine/Plugins/Runtime/ReplicationSystemTestPlugin/Source/Private/Tests/ReplicationSystem/NetBlob/TestReliableNetBlobQueue.cpp b/Engine/Plugins/Runtime/ReplicationSystemTestPlugin/Source/Private/Tests/ReplicationSystem/NetBlob/TestReliableNetBlobQueue.cpp
--- a/Engine/Plugins/Runtime/ReplicationSystemTestPlugin/Source/Private/Tests/ReplicationSystem/NetBlob/TestReliableNetBlobQueue.cpp
+++ b/Engine/Plugins/Runtime/ReplicationSystemTestPlugin/Source/Private/Tests/ReplicationSystem/NetBlob/TestReliableNetBlobQueue.cpp
@@ -25,13 +25,13 @@ UE_NET_TEST_FIXTURE(FTestReliableNetBlobQueue, CanSendMoreReliableBlobsThanRelia
 	Server->SendAndDeliverTo(Client, DeliverPacket);
 	Server->PostSendUpdate();
 
-	UE_NET_ASSERT_NE(Client->GetReplicationBridge()->GetReplicatedObject(ServerObject->NetHandle), nullptr);
+	UE_NET_ASSERT_NE(Client->GetReplicationBridge()->GetReplicatedObject(ServerObject->NetRefHandle), nullptr);
 
 	// Create huge attachment
 	constexpr uint32 BlobPartCount = FReliableNetBlobQueue::MaxUnackedBlobCount + 3U;
 	{
 		const TRefCountPtr<FNetObjectAttachment>& Attachment = CreateReliableNetObjectAttachmentWithPartCount(BlobPartCount);
-		FNetObjectReference AttachmentTarget = FObjectReferenceCache::MakeNetObjectReference(ServerObject->NetHandle);
+		FNetObjectReference AttachmentTarget = FObjectReferenceCache::MakeNetObjectReference(ServerObject->NetRefHandle);
 		Server->GetReplicationSystem()->QueueNetObjectAttachment(Client->ConnectionIdOnServer, AttachmentTarget, Attachment);
 	}
 
@@ -62,7 +62,7 @@ UE_NET_TEST_FIXTURE(FTestReliableNetBlobQueue, CanSendMoreReliableBlobsThanRelia
 	Server->SendAndDeliverTo(Client, DeliverPacket);
 	Server->PostSendUpdate();
 
-	UE_NET_ASSERT_NE(Client->GetReplicationBridge()->GetReplicatedObject(ServerObject->NetHandle), nullptr);
+	UE_NET_ASSERT_NE(Client->GetReplicationBridge()->GetReplicatedObject(ServerObject->NetRefHandle), nullptr);
 
 	// Perform the test with different packets being lost.
 	for (uint32 PacketNumberToDrop : {0U, 1U})
@@ -74,7 +74,7 @@ UE_NET_TEST_FIXTURE(FTestReliableNetBlobQueue, CanSendMoreReliableBlobsThanRelia
 		constexpr uint32 BlobPartCount = FReliableNetBlobQueue::MaxUnackedBlobCount + 3U;
 		{
 			const TRefCountPtr<FNetObjectAttachment>& Attachment = CreateReliableNetObjectAttachmentWithPartCount(BlobPartCount);
-			FNetObjectReference AttachmentTarget = FObjectReferenceCache::MakeNetObjectReference(ServerObject->NetHandle);
+			FNetObjectReference AttachmentTarget = FObjectReferenceCache::MakeNetObjectReference(ServerObject->NetRefHandle);
 			Server->GetReplicationSystem()->QueueNetObjectAttachment(Client->ConnectionIdOnServer, AttachmentTarget, Attachment);
 		}
 
diff --git a/Engine/Plugins/Runtime/ReplicationSystemTestPlugin/Source/Private/Tests/ReplicationSystem/Prioritization/MockNetObjectPrioritizer.cpp b/Engine/Plugins/Runtime/ReplicationSystemTestPlugin/Source/Private/Tests/ReplicationSystem/Prioritization/MockNetObjectPrioritizer.cpp
--- a/Engine/Plugins/Runtime/ReplicationSystemTestPlugin/Source/Private/Tests/ReplicationSystem/Prioritization/MockNetObjectPrioritizer.cpp
+++ b/Engine/Plugins/Runtime/ReplicationSystemTestPlugin/Source/Private/Tests/ReplicationSystem/Prioritization/MockNetObjectPrioritizer.cpp
@@ -132,7 +132,7 @@ void UMockNetObjectPrioritizer::Prioritize(FNetObjectPrioritizationParams& Param
 	CallStatus.SuccessfulCallCounts.Prioritize += bIsProperCall;
 }
 
-float UMockNetObjectPrioritizer::GetPriority(UE::Net::Private::FInternalNetHandle ObjectIndex) const
+float UMockNetObjectPrioritizer::GetPriority(UE::Net::Private::FInternalNetRefIndex ObjectIndex) const
 {
 	check(AddedIndices[ObjectIndex]);
 	if (const float* Priority = ObjectToPriority.Find(ObjectIndex))
diff --git a/Engine/Plugins/Runtime/ReplicationSystemTestPlugin/Source/Private/Tests/ReplicationSystem/Prioritization/MockNetObjectPrioritizer.h b/Engine/Plugins/Runtime/ReplicationSystemTestPlugin/Source/Private/Tests/ReplicationSystem/Prioritization/MockNetObjectPrioritizer.h
--- a/Engine/Plugins/Runtime/ReplicationSystemTestPlugin/Source/Private/Tests/ReplicationSystem/Prioritization/MockNetObjectPrioritizer.h
+++ b/Engine/Plugins/Runtime/ReplicationSystemTestPlugin/Source/Private/Tests/ReplicationSystem/Prioritization/MockNetObjectPrioritizer.h
@@ -9,7 +9,7 @@
 
 namespace UE::Net::Private
 {
-	typedef uint32 FInternalNetHandle;
+	typedef uint32 FInternalNetRefIndex;
 }
 
 const UE::Net::FRepTag RepTag_NetTest_Priority = 0x22C6880A98E628EDULL;
@@ -56,7 +56,7 @@ public:
 	inline const FFunctionCallStatus& GetFunctionCallStatus() const { return CallStatus; }
 	inline void ResetFunctionCallStatus() { CallStatus = FFunctionCallStatus({}); }
 
-	float GetPriority(UE::Net::Private::FInternalNetHandle ObjectIndex) const;
+	float GetPriority(UE::Net::Private::FInternalNetRefIndex ObjectIndex) const;
 
 protected:
 	virtual void Init(FNetObjectPrioritizerInitParams& Params) override;
diff --git a/Engine/Plugins/Runtime/ReplicationSystemTestPlugin/Source/Private/Tests/ReplicationSystem/Prioritization/TestNetObjectPrioritizer.cpp b/Engine/Plugins/Runtime/ReplicationSystemTestPlugin/Source/Private/Tests/ReplicationSystem/Prioritization/TestNetObjectPrioritizer.cpp
--- a/Engine/Plugins/Runtime/ReplicationSystemTestPlugin/Source/Private/Tests/ReplicationSystem/Prioritization/TestNetObjectPrioritizer.cpp
+++ b/Engine/Plugins/Runtime/ReplicationSystemTestPlugin/Source/Private/Tests/ReplicationSystem/Prioritization/TestNetObjectPrioritizer.cpp
@@ -25,7 +25,7 @@ private:
 		FReplicationSystemTestFixture::SetUp();
 		InitMockNetObjectPrioritizer();
 
-		NetHandleManager = &ReplicationSystem->GetReplicationSystemInternal()->GetNetHandleManager();
+		NetRefHandleManager = &ReplicationSystem->GetReplicationSystemInternal()->GetNetRefHandleManager();
 	}
 
 	virtual void TearDown() override
@@ -33,7 +33,7 @@ private:
 		FReplicationSystemTestFixture::TearDown();
 		RestoreNetObjectPrioritizerDefinitions();
 
-		NetHandleManager = nullptr;
+		NetRefHandleManager = nullptr;
 	}
 
 	void InitNetObjectPrioritizerDefinitions()
@@ -77,7 +77,7 @@ private:
 protected:
 	UMockNetObjectPrioritizer* MockNetObjectPrioritizer = nullptr;
 	FNetObjectPrioritizerHandle MockPrioritizerHandle = InvalidNetObjectPrioritizerHandle;
-	FNetHandleManager* NetHandleManager = nullptr;
+	FNetRefHandleManager* NetRefHandleManager = nullptr;
 	static constexpr float FakeDeltaTime = 0.0334f;
 
 private:
@@ -112,8 +112,8 @@ UE_NET_TEST_FIXTURE(FTestNetPrioritizerFixture, PrioritizerAddObjectSucceedsAndR
 			MockNetObjectPrioritizer->SetFunctionCallSetup(CallSetup);
 		}
 
-		const FNetHandle NetHandle = ReplicationBridge->BeginReplication(TestObject);
-		ReplicationSystem->SetPrioritizer(NetHandle, MockPrioritizerHandle);
+		const FNetRefHandle NetRefHandle = ReplicationBridge->BeginReplication(TestObject);
+		ReplicationSystem->SetPrioritizer(NetRefHandle, MockPrioritizerHandle);
 
 		// We don't know if prioritizer changes are batched, but we assume everything is setup correctly in PreSendUpdate at least.
 		ReplicationSystem->PreSendUpdate(FakeDeltaTime);
@@ -158,8 +158,8 @@ UE_NET_TEST_FIXTURE(FTestNetPrioritizerFixture, PrioritizerAddObjectFailsAndRemo
 			MockNetObjectPrioritizer->SetFunctionCallSetup(CallSetup);
 		}
 
-		const FNetHandle NetHandle = ReplicationBridge->BeginReplication(TestObject);
-		ReplicationSystem->SetPrioritizer(NetHandle, MockPrioritizerHandle);
+		const FNetRefHandle NetRefHandle = ReplicationBridge->BeginReplication(TestObject);
+		ReplicationSystem->SetPrioritizer(NetRefHandle, MockPrioritizerHandle);
 
 		// We don't know if prioritizer changes are batched, but we assume everything is setup correctly in PreSendUpdate at least.
 		ReplicationSystem->PreSendUpdate(FakeDeltaTime);
@@ -202,8 +202,8 @@ UE_NET_TEST_FIXTURE(FTestNetPrioritizerFixture, SwitchingPrioritizersCallsRemove
 			MockNetObjectPrioritizer->SetFunctionCallSetup(CallSetup);
 		}
 
-		const FNetHandle NetHandle = ReplicationBridge->BeginReplication(TestObject);
-		ReplicationSystem->SetPrioritizer(NetHandle, MockPrioritizerHandle);
+		const FNetRefHandle NetRefHandle = ReplicationBridge->BeginReplication(TestObject);
+		ReplicationSystem->SetPrioritizer(NetRefHandle, MockPrioritizerHandle);
 
 		// We don't know if prioritizer changes are batched, but we assume everything is setup correctly in PreSendUpdate at least.
 		ReplicationSystem->PreSendUpdate(FakeDeltaTime);
@@ -219,7 +219,7 @@ UE_NET_TEST_FIXTURE(FTestNetPrioritizerFixture, SwitchingPrioritizersCallsRemove
 	{
 		MockNetObjectPrioritizer->ResetFunctionCallStatus();
 
-		ReplicationSystem->SetStaticPriority(TestObject->NetHandle, 1.0f);
+		ReplicationSystem->SetStaticPriority(TestObject->NetRefHandle, 1.0f);
 
 		ReplicationSystem->PreSendUpdate(FakeDeltaTime);
 
@@ -239,8 +239,8 @@ UE_NET_TEST_FIXTURE(FTestNetPrioritizerFixture, PrioritizeIsNotCalledWhenNoObjec
 	{
 		MockNetObjectPrioritizer->ResetFunctionCallStatus();
 
-		const FNetHandle NetHandle = ReplicationBridge->BeginReplication(TestObject);
-		ReplicationSystem->SetStaticPriority(NetHandle, 1.0f);
+		const FNetRefHandle NetRefHandle = ReplicationBridge->BeginReplication(TestObject);
+		ReplicationSystem->SetStaticPriority(NetRefHandle, 1.0f);
 
 		ReplicationSystem->PreSendUpdate(FakeDeltaTime);
 
@@ -267,8 +267,8 @@ UE_NET_TEST_FIXTURE(FTestNetPrioritizerFixture, PrioritizeIsNotCalledWhenThereAr
 			MockNetObjectPrioritizer->SetFunctionCallSetup(CallSetup);
 		}
 
-		const FNetHandle NetHandle = ReplicationBridge->BeginReplication(TestObject);
-		ReplicationSystem->SetPrioritizer(NetHandle, MockPrioritizerHandle);
+		const FNetRefHandle NetRefHandle = ReplicationBridge->BeginReplication(TestObject);
+		ReplicationSystem->SetPrioritizer(NetRefHandle, MockPrioritizerHandle);
 
 		ReplicationSystem->PreSendUpdate(FakeDeltaTime);
 
@@ -308,11 +308,11 @@ UE_NET_TEST_FIXTURE(FTestNetPrioritizerFixture, PrioritizeIsCalledWhenThereAreCo
 			MockNetObjectPrioritizer->SetFunctionCallSetup(CallSetup);
 		}
 
-		const FNetHandle NetHandle1 = ReplicationBridge->BeginReplication(TestObject1);
-		ReplicationSystem->SetPrioritizer(NetHandle1, MockPrioritizerHandle);
+		const FNetRefHandle NetRefHandle1 = ReplicationBridge->BeginReplication(TestObject1);
+		ReplicationSystem->SetPrioritizer(NetRefHandle1, MockPrioritizerHandle);
 
-		const FNetHandle NetHandle2 = ReplicationBridge->BeginReplication(TestObject2);
-		ReplicationSystem->SetPrioritizer(NetHandle2, MockPrioritizerHandle);
+		const FNetRefHandle NetRefHandle2 = ReplicationBridge->BeginReplication(TestObject2);
+		ReplicationSystem->SetPrioritizer(NetRefHandle2, MockPrioritizerHandle);
 
 		ReplicationSystem->PreSendUpdate(FakeDeltaTime);
 
@@ -344,8 +344,8 @@ UE_NET_TEST_FIXTURE(FTestNetPrioritizerFixture, UpdateObjectsIsCalledForDirtyObj
 			MockNetObjectPrioritizer->SetFunctionCallSetup(CallSetup);
 		}
 
-		const FNetHandle NetHandle1 = ReplicationBridge->BeginReplication(TestObject1);
-		ReplicationSystem->SetPrioritizer(NetHandle1, MockPrioritizerHandle);
+		const FNetRefHandle NetRefHandle1 = ReplicationBridge->BeginReplication(TestObject1);
+		ReplicationSystem->SetPrioritizer(NetRefHandle1, MockPrioritizerHandle);
 
 		ReplicationSystem->PreSendUpdate(FakeDeltaTime);
 		ReplicationSystem->PostSendUpdate();
@@ -382,8 +382,8 @@ UE_NET_TEST_FIXTURE(FTestNetPrioritizerFixture, UpdateObjectsIsNotCalledWhenNoOb
 			MockNetObjectPrioritizer->SetFunctionCallSetup(CallSetup);
 		}
 
-		const FNetHandle NetHandle1 = ReplicationBridge->BeginReplication(TestObject1);
-		ReplicationSystem->SetPrioritizer(NetHandle1, MockPrioritizerHandle);
+		const FNetRefHandle NetRefHandle1 = ReplicationBridge->BeginReplication(TestObject1);
+		ReplicationSystem->SetPrioritizer(NetRefHandle1, MockPrioritizerHandle);
 
 		ReplicationSystem->PreSendUpdate(FakeDeltaTime);
 		ReplicationSystem->PostSendUpdate();
@@ -416,13 +416,13 @@ UE_NET_TEST_FIXTURE(FTestNetPrioritizerFixture, NativeIrisObjectGetsPriorityFrom
 	constexpr float ObjectPriority = 1.3056640625E10f;
 	Object->SetPriority(ObjectPriority);
 
-	const FNetHandle NetHandle = ReplicationBridge->BeginReplication(Object);
-	ReplicationSystem->SetPrioritizer(NetHandle, MockPrioritizerHandle);
+	const FNetRefHandle NetRefHandle = ReplicationBridge->BeginReplication(Object);
+	ReplicationSystem->SetPrioritizer(NetRefHandle, MockPrioritizerHandle);
 
 	ReplicationSystem->PreSendUpdate(FakeDeltaTime);
 	ReplicationSystem->PostSendUpdate();
 
-	UE::Net::Private::FInternalNetHandle InternalIndex = NetHandleManager->GetInternalIndex(NetHandle);
+	UE::Net::Private::FInternalNetRefIndex InternalIndex = NetRefHandleManager->GetInternalIndex(NetRefHandle);
 	const float Priority = MockNetObjectPrioritizer->GetPriority(InternalIndex);
 	UE_NET_ASSERT_EQ(Priority, ObjectPriority);
 }
@@ -439,8 +439,8 @@ UE_NET_TEST_FIXTURE(FTestNetPrioritizerFixture, NativeIrisObjectGetsUpdatedPrior
 
 	Object->SetPriority(2.0f);
 
-	const FNetHandle NetHandle = ReplicationBridge->BeginReplication(Object);
-	ReplicationSystem->SetPrioritizer(NetHandle, MockPrioritizerHandle);
+	const FNetRefHandle NetRefHandle = ReplicationBridge->BeginReplication(Object);
+	ReplicationSystem->SetPrioritizer(NetRefHandle, MockPrioritizerHandle);
 
 	// Update priority
 	constexpr float UpdatedPriority = 4711;
@@ -449,7 +449,7 @@ UE_NET_TEST_FIXTURE(FTestNetPrioritizerFixture, NativeIrisObjectGetsUpdatedPrior
 	ReplicationSystem->PreSendUpdate(FakeDeltaTime);
 	ReplicationSystem->PostSendUpdate();
 
-	UE::Net::Private::FInternalNetHandle InternalIndex = NetHandleManager->GetInternalIndex(NetHandle);
+	UE::Net::Private::FInternalNetRefIndex InternalIndex = NetRefHandleManager->GetInternalIndex(NetRefHandle);
 	const float Priority = MockNetObjectPrioritizer->GetPriority(InternalIndex);
 	UE_NET_ASSERT_EQ(Priority, UpdatedPriority);
 }
@@ -467,13 +467,13 @@ UE_NET_TEST_FIXTURE(FTestNetPrioritizerFixture, ObjectGetsPriorityFromStart)
 	constexpr float ObjectPriority = 1.3056640625E10f;
 	Object->SetPriority(ObjectPriority);
 
-	const FNetHandle NetHandle = ReplicationBridge->BeginReplication(Object);
-	ReplicationSystem->SetPrioritizer(NetHandle, MockPrioritizerHandle);
+	const FNetRefHandle NetRefHandle = ReplicationBridge->BeginReplication(Object);
+	ReplicationSystem->SetPrioritizer(NetRefHandle, MockPrioritizerHandle);
 
 	ReplicationSystem->PreSendUpdate(FakeDeltaTime);
 	ReplicationSystem->PostSendUpdate();
 
-	UE::Net::Private::FInternalNetHandle InternalIndex = NetHandleManager->GetInternalIndex(NetHandle);
+	UE::Net::Private::FInternalNetRefIndex InternalIndex = NetRefHandleManager->GetInternalIndex(NetRefHandle);
 	const float Priority = MockNetObjectPrioritizer->GetPriority(InternalIndex);
 	UE_NET_ASSERT_EQ(Priority, ObjectPriority);
 }
@@ -490,8 +490,8 @@ UE_NET_TEST_FIXTURE(FTestNetPrioritizerFixture, ObjectGetsUpdatedPriority)
 
 	Object->SetPriority(2.0f);
 
-	const FNetHandle NetHandle = ReplicationBridge->BeginReplication(Object);
-	ReplicationSystem->SetPrioritizer(NetHandle, MockPrioritizerHandle);
+	const FNetRefHandle NetRefHandle = ReplicationBridge->BeginReplication(Object);
+	ReplicationSystem->SetPrioritizer(NetRefHandle, MockPrioritizerHandle);
 
 	// Update priority
 	constexpr float UpdatedPriority = 4711;
@@ -500,7 +500,7 @@ UE_NET_TEST_FIXTURE(FTestNetPrioritizerFixture, ObjectGetsUpdatedPriority)
 	ReplicationSystem->PreSendUpdate(FakeDeltaTime);
 	ReplicationSystem->PostSendUpdate();
 
-	UE::Net::Private::FInternalNetHandle InternalIndex = NetHandleManager->GetInternalIndex(NetHandle);
+	UE::Net::Private::FInternalNetRefIndex InternalIndex = NetRefHandleManager->GetInternalIndex(NetRefHandle);
 	const float Priority = MockNetObjectPrioritizer->GetPriority(InternalIndex);
 	UE_NET_ASSERT_EQ(Priority, UpdatedPriority);
 }
diff --git a/Engine/Plugins/Runtime/ReplicationSystemTestPlugin/Source/Private/Tests/ReplicationSystem/ReplicatedTestObject.cpp b/Engine/Plugins/Runtime/ReplicationSystemTestPlugin/Source/Private/Tests/ReplicationSystem/ReplicatedTestObject.cpp
--- a/Engine/Plugins/Runtime/ReplicationSystemTestPlugin/Source/Private/Tests/ReplicationSystem/ReplicatedTestObject.cpp
+++ b/Engine/Plugins/Runtime/ReplicationSystemTestPlugin/Source/Private/Tests/ReplicationSystem/ReplicatedTestObject.cpp
@@ -34,55 +34,55 @@ UReplicatedTestObjectBridge::UReplicatedTestObjectBridge()
 {
 }
 
-const UE::Net::FReplicationInstanceProtocol* UReplicatedTestObjectBridge::GetReplicationInstanceProtocol(FNetHandle Handle) const
+const UE::Net::FReplicationInstanceProtocol* UReplicatedTestObjectBridge::GetReplicationInstanceProtocol(FNetRefHandle Handle) const
 {
-	const UE::Net::Private::FNetHandleManager* LocalNetHandleManager = &GetReplicationSystem()->GetReplicationSystemInternal()->GetNetHandleManager();
+	const UE::Net::Private::FNetRefHandleManager* LocalNetRefHandleManager = &GetReplicationSystem()->GetReplicationSystemInternal()->GetNetRefHandleManager();
 
-	if (uint32 InternalObjectIndex = LocalNetHandleManager->GetInternalIndex(Handle))
+	if (uint32 InternalObjectIndex = LocalNetRefHandleManager->GetInternalIndex(Handle))
 	{
-		return LocalNetHandleManager->GetReplicatedObjectDataNoCheck(InternalObjectIndex).InstanceProtocol;
+		return LocalNetRefHandleManager->GetReplicatedObjectDataNoCheck(InternalObjectIndex).InstanceProtocol;
 	}
 
 	return nullptr;
 };
 
-UE::Net::FNetHandle UReplicatedTestObjectBridge::BeginReplication(UReplicatedTestObject* Instance)
+UE::Net::FNetRefHandle UReplicatedTestObjectBridge::BeginReplication(UReplicatedTestObject* Instance)
 {
-	// Create NetHandle for the registered fragments
-	Super::FCreateNetHandleParams Params = Super::DefaultCreateNetHandleParams;
+	// Create NetRefHandle for the registered fragments
+	Super::FCreateNetRefHandleParams Params = Super::DefaultCreateNetRefHandleParams;
 	Params.bCanReceive = true;
 
-	FNetHandle Handle = Super::BeginReplication(Instance, Params);
+	FNetRefHandle Handle = Super::BeginReplication(Instance, Params);
 
-	// This is optional but typically we want to cache at least the NetHandle in the game instance to avoid doing map lookups to find it
+	// This is optional but typically we want to cache at least the NetRefHandle in the game instance to avoid doing map lookups to find it
 	if (Handle.IsValid())
 	{
-		Instance->NetHandle = Handle;
+		Instance->NetRefHandle = Handle;
 	}
 
 	return Handle;
 }
 
-UE::Net::FNetHandle UReplicatedTestObjectBridge::BeginReplication(FNetHandle OwnerHandle, UReplicatedTestObject* Instance, FNetHandle InsertRelativeToSubObjectHandle, ESubObjectInsertionOrder InsertionOrder)
+UE::Net::FNetRefHandle UReplicatedTestObjectBridge::BeginReplication(FNetRefHandle OwnerHandle, UReplicatedTestObject* Instance, FNetRefHandle InsertRelativeToSubObjectHandle, ESubObjectInsertionOrder InsertionOrder)
 {
-	// Create NetHandle for the registered fragments
-	Super::FCreateNetHandleParams Params = Super::DefaultCreateNetHandleParams;
+	// Create NetRefHandle for the registered fragments
+	Super::FCreateNetRefHandleParams Params = Super::DefaultCreateNetRefHandleParams;
 	Params.bCanReceive = true;
 
-	// Create NetHandle for the registered fragments
-	FNetHandle Handle = Super::BeginReplication(OwnerHandle, Instance, InsertRelativeToSubObjectHandle, InsertionOrder, Params);
+	// Create NetRefHandle for the registered fragments
+	FNetRefHandle Handle = Super::BeginReplication(OwnerHandle, Instance, InsertRelativeToSubObjectHandle, InsertionOrder, Params);
 
 	if (Handle.IsValid())
 	{
-		// This is optional but typically we want to cache at least the NetHandle in the game instance to avoid doing map lookups to find it
-		Instance->NetHandle = Handle;
+		// This is optional but typically we want to cache at least the NetRefHandle in the game instance to avoid doing map lookups to find it
+		Instance->NetRefHandle = Handle;
 	}
 	
 	return Handle;
 }
 
 
-bool UReplicatedTestObjectBridge::WriteCreationHeader(UE::Net::FNetSerializationContext& Context, FNetHandle Handle)
+bool UReplicatedTestObjectBridge::WriteCreationHeader(UE::Net::FNetSerializationContext& Context, FNetRefHandle Handle)
 {
 	UE::Net::FNetBitStreamWriter& Writer = *Context.GetBitStreamWriter();
 
@@ -135,7 +135,7 @@ UObjectReplicationBridge::FCreationHeader* UReplicatedTestObjectBridge::ReadCrea
 	return Header.Release();
 }
 
-UObject* UReplicatedTestObjectBridge::BeginInstantiateFromRemote(FNetHandle SubObjectOwnerHandle, const UE::Net::FNetObjectResolveContext& ResolveContext, const FCreationHeader* InHeader)
+UObject* UReplicatedTestObjectBridge::BeginInstantiateFromRemote(FNetRefHandle SubObjectOwnerHandle, const UE::Net::FNetObjectResolveContext& ResolveContext, const FCreationHeader* InHeader)
 {
 	const FReplicationTestObjectCreationHeader* Header = static_cast<const FReplicationTestObjectCreationHeader*>(InHeader);
 
@@ -166,11 +166,11 @@ UObject* UReplicatedTestObjectBridge::BeginInstantiateFromRemote(FNetHandle SubO
 	return CreatedObject;
 }
 
-void UReplicatedTestObjectBridge::EndInstantiateFromRemote(FNetHandle Handle)
+void UReplicatedTestObjectBridge::EndInstantiateFromRemote(FNetRefHandle Handle)
 {
 	UReplicatedTestObject* Instance = CastChecked<UReplicatedTestObject>(GetReplicatedObject(Handle));
 
-	Instance->NetHandle = Handle;
+	Instance->NetRefHandle = Handle;
 }
 
 void UReplicatedTestObjectBridge::DestroyInstanceFromRemote(UObject* Instance, bool bTearOff)
@@ -197,9 +197,9 @@ void UReplicatedTestObjectBridge::SetPollFramePeriod(UReplicatedTestObject* Inst
 	{
 		return;
 	}
-	const FNetHandleManager* LocalNetHandleManager = &GetReplicationSystem()->GetReplicationSystemInternal()->GetNetHandleManager();
-	const uint32 InternalObjectIndex = LocalNetHandleManager->GetInternalIndex(Instance->NetHandle);
-	if (InternalObjectIndex != FNetHandleManager::InvalidInternalIndex)
+	const FNetRefHandleManager* LocalNetRefHandleManager = &GetReplicationSystem()->GetReplicationSystemInternal()->GetNetRefHandleManager();
+	const uint32 InternalObjectIndex = LocalNetRefHandleManager->GetInternalIndex(Instance->NetRefHandle);
+	if (InternalObjectIndex != FNetRefHandleManager::InvalidInternalIndex)
 	{
 		UObjectReplicationBridge::SetPollFramePeriod(InternalObjectIndex, FramePeriod);
 	}
@@ -298,7 +298,7 @@ UTestReplicatedIrisObject::UTestReplicatedIrisObject()
 
 void UTestReplicatedIrisObject::AddComponents(const UTestReplicatedIrisObject::FComponents& InComponents)
 {
-	check(!NetHandle.IsValid())
+	check(!NetRefHandle.IsValid())
 	for (uint32 It = 0; It < InComponents.PropertyComponentCount; ++It)
 	{
 		Components.Emplace(NewObject<UTestReplicatedIrisPropertyComponent>());
@@ -327,7 +327,7 @@ void UTestReplicatedIrisObject::AddComponents(const UTestReplicatedIrisObject::F
 
 void UTestReplicatedIrisObject::AddComponents(uint32 PropertyComponentCount, uint32 IrisComponentCount)
 {
-	check(!NetHandle.IsValid())
+	check(!NetRefHandle.IsValid())
 	// Setup a few components
 	for (uint32 It = 0; It < PropertyComponentCount; ++It)
 	{
@@ -343,7 +343,7 @@ void UTestReplicatedIrisObject::AddComponents(uint32 PropertyComponentCount, uin
 
 void UTestReplicatedIrisObject::AddDynamicStateComponents(uint32 DynamicStateComponentCount)
 {
-	check(!NetHandle.IsValid())
+	check(!NetRefHandle.IsValid())
 	for (uint32 It = 0; It < DynamicStateComponentCount; ++It)
 	{
 		DynamicStateComponents.Emplace(NewObject<UTestReplicatedIrisDynamicStatePropertyComponent>());
diff --git a/Engine/Plugins/Runtime/ReplicationSystemTestPlugin/Source/Private/Tests/ReplicationSystem/ReplicatedTestObject.h b/Engine/Plugins/Runtime/ReplicationSystemTestPlugin/Source/Private/Tests/ReplicationSystem/ReplicatedTestObject.h
--- a/Engine/Plugins/Runtime/ReplicationSystemTestPlugin/Source/Private/Tests/ReplicationSystem/ReplicatedTestObject.h
+++ b/Engine/Plugins/Runtime/ReplicationSystemTestPlugin/Source/Private/Tests/ReplicationSystem/ReplicatedTestObject.h
@@ -32,11 +32,10 @@ class UReplicatedTestObject : public UObject
 public:
 	// Register the fragments for this object
 	virtual void RegisterReplicationFragments(UE::Net::FFragmentRegistrationContext& Fragments, UE::Net::EFragmentRegistrationFlags RegistrationFlags) override {}
-	bool IsReplicating() const { return NetHandle.IsValid(); }
 	virtual bool IsSupportedForNetworking() const override { return true; }
 
-	// Currently we carry around some cached data, this could be stored in the bridge and be accessed through handle resolve.
-	UE::Net::FNetHandle NetHandle;
+	/** Cached NetRefHandle to simplify testing. DO NOT use this for multi-system tests. */
+	UE::Net::FNetRefHandle NetRefHandle;
 };
 
 USTRUCT()
@@ -408,11 +407,11 @@ public:
 
 	// This is the local Interface, it is up to each bridge implementation to define the interface for that type
 	// In this example we have methods that directly uses UReplicatedTestObject;
-	FNetHandle BeginReplication(UReplicatedTestObject* Instance);
-	FNetHandle BeginReplication(FNetHandle OwnerHandle, UReplicatedTestObject* Instance, FNetHandle InsertRelativeToSubObjectHandle = FNetHandle(), ESubObjectInsertionOrder InsertionOrder = UReplicationBridge::ESubObjectInsertionOrder::None);
+	FNetRefHandle BeginReplication(UReplicatedTestObject* Instance);
+	FNetRefHandle BeginReplication(FNetRefHandle OwnerHandle, UReplicatedTestObject* Instance, FNetRefHandle InsertRelativeToSubObjectHandle = FNetRefHandle(), ESubObjectInsertionOrder InsertionOrder = UReplicationBridge::ESubObjectInsertionOrder::None);
 
 	// For testing we expose some things that normally are not accessible
-	const UE::Net::FReplicationInstanceProtocol* GetReplicationInstanceProtocol(FNetHandle Handle) const;
+	const UE::Net::FReplicationInstanceProtocol* GetReplicationInstanceProtocol(FNetRefHandle Handle) const;
 
 	void SetPollFramePeriod(UReplicatedTestObject* Instance, uint8 FramePeriod);
 
@@ -429,11 +428,11 @@ protected:
 		uint32 NumConnectionFilteredComponentsToSpawn;
 	};
 
-	virtual bool WriteCreationHeader(UE::Net::FNetSerializationContext& Context, FNetHandle Handle) override;
+	virtual bool WriteCreationHeader(UE::Net::FNetSerializationContext& Context, FNetRefHandle Handle) override;
 	virtual FCreationHeader* ReadCreationHeader(UE::Net::FNetSerializationContext& Context) override;
 
-	virtual UObject* BeginInstantiateFromRemote(FNetHandle SubObjectOwnerHandle, const UE::Net::FNetObjectResolveContext& ResolveContext, const FCreationHeader* InHeader) override;
-	virtual void EndInstantiateFromRemote(FNetHandle Handle) override;
+	virtual UObject* BeginInstantiateFromRemote(FNetRefHandle SubObjectOwnerHandle, const UE::Net::FNetObjectResolveContext& ResolveContext, const FCreationHeader* InHeader) override;
+	virtual void EndInstantiateFromRemote(FNetRefHandle Handle) override;
 	virtual void DestroyInstanceFromRemote(UObject* Instance, bool bTearOff);
 
 	TArray<TStrongObjectPtr<UObject>>* CreatedObjectsOnNode;
diff --git a/Engine/Plugins/Runtime/ReplicationSystemTestPlugin/Source/Private/Tests/ReplicationSystem/ReplicationSystemServerClientTestFixture.cpp b/Engine/Plugins/Runtime/ReplicationSystemTestPlugin/Source/Private/Tests/ReplicationSystem/ReplicationSystemServerClientTestFixture.cpp
--- a/Engine/Plugins/Runtime/ReplicationSystemTestPlugin/Source/Private/Tests/ReplicationSystem/ReplicationSystemServerClientTestFixture.cpp
+++ b/Engine/Plugins/Runtime/ReplicationSystemTestPlugin/Source/Private/Tests/ReplicationSystem/ReplicationSystemServerClientTestFixture.cpp
@@ -85,6 +85,16 @@ FReplicationSystemTestNode::FReplicationSystemTestNode(bool bIsServer, const TCH
 	UE_NET_TRACE_UPDATE_INSTANCE(GetNetTraceId(), bIsServer, Name);
 }
 
+FReplicationSystemTestNode::~FReplicationSystemTestNode()
+{
+	const uint32 NetTraceId = ReplicationSystem->GetId();
+	FReplicationSystemFactory::DestroyReplicationSystem(ReplicationSystem);
+	CreatedObjects.Empty();
+
+	// End NetTrace session for this instance
+	UE_NET_TRACE_END_SESSION(NetTraceId);
+}
+
 uint32 FReplicationSystemTestNode::GetNetTraceId() const
 { 
 	return ReplicationSystem ? ReplicationSystem->GetId() : ~0U;
@@ -115,7 +125,7 @@ UTestReplicatedIrisObject* FReplicationSystemTestNode::CreateObject(const UTestR
 	return CreatedObject;
 }
 
-UTestReplicatedIrisObject* FReplicationSystemTestNode::CreateSubObject(FNetHandle Owner, const UTestReplicatedIrisObject::FComponents& Components)
+UTestReplicatedIrisObject* FReplicationSystemTestNode::CreateSubObject(FNetRefHandle Owner, const UTestReplicatedIrisObject::FComponents& Components)
 {
 	UTestReplicatedIrisObject* CreatedObject = NewObject<UTestReplicatedIrisObject>();
 	CreatedObjects.Add(TStrongObjectPtr<UObject>(CreatedObject));
@@ -127,7 +137,7 @@ UTestReplicatedIrisObject* FReplicationSystemTestNode::CreateSubObject(FNetHandl
 	return CreatedObject;
 }
 
-UTestReplicatedIrisObject* FReplicationSystemTestNode::CreateSubObject(FNetHandle Owner, uint32 NumComponents, uint32 NumIrisComponents)
+UTestReplicatedIrisObject* FReplicationSystemTestNode::CreateSubObject(FNetRefHandle Owner, uint32 NumComponents, uint32 NumIrisComponents)
 {
 	UTestReplicatedIrisObject* CreatedObject = NewObject<UTestReplicatedIrisObject>();
 	CreatedObjects.Add(TStrongObjectPtr<UObject>(CreatedObject));
@@ -153,7 +163,7 @@ UTestReplicatedIrisObject* FReplicationSystemTestNode::CreateObjectWithDynamicSt
 void FReplicationSystemTestNode::DestroyObject(UReplicatedTestObject* Object, EEndReplicationFlags EndReplicationFlags)
 {
 	// Destroy handle
-	check(Object && Object->NetHandle.IsValid());
+	check(Object && Object->NetRefHandle.IsValid());
 
 	// End replication for the handle
 	ReplicationBridge->EndReplication(Object, EndReplicationFlags);
@@ -296,18 +306,13 @@ void FReplicationSystemTestNode::RecvUpdate(uint32 ConnectionId, FNetSerializati
 	check(Context.GetBitStreamReader()->GetBitsLeft() == 0U);
 }
 
-FReplicationSystemTestNode::~FReplicationSystemTestNode()
+uint32 FReplicationSystemTestNode::GetReplicationSystemId() const
 {
-	const uint32 NetTraceId = ReplicationSystem->GetId();
-	FReplicationSystemFactory::DestroyReplicationSystem(ReplicationSystem);
-	CreatedObjects.Empty();
-
-	// End NetTrace session for this instance
-	UE_NET_TRACE_END_SESSION(NetTraceId);
+	return ReplicationSystem ? ReplicationSystem->GetId() : uint32(~0U);
 }
 
 // FReplicationSystemTestClient implementation
-FReplicationSystemTestClient::FReplicationSystemTestClient(uint32 ReplicationSystemId, const TCHAR* Name)
+FReplicationSystemTestClient::FReplicationSystemTestClient(const TCHAR* Name)
 : FReplicationSystemTestNode(false, Name)
 , ConnectionIdOnServer(~0U)
 {
@@ -358,9 +363,22 @@ void FReplicationSystemServerClientTestFixture::SetUp()
 	Server = new FReplicationSystemTestServer(GetName());
 }
 
+void FReplicationSystemServerClientTestFixture::TearDown()
+{
+	delete Server;
+	for (FReplicationSystemTestClient* Client : Clients)
+	{
+		delete Client;
+	}
+	Clients.Empty();
+	DataStreamUtil.TearDown();
+
+	FNetworkAutomationTestSuiteFixture::TearDown();
+}
+
 FReplicationSystemTestClient* FReplicationSystemServerClientTestFixture::CreateClient()
 {
-	FReplicationSystemTestClient* Client = new FReplicationSystemTestClient(Clients.Num() + 1U, GetName());
+	FReplicationSystemTestClient* Client = new FReplicationSystemTestClient(GetName());
 	Clients.Add(Client);
 
 	// The client needs a connection
@@ -372,17 +390,4 @@ FReplicationSystemTestClient* FReplicationSystemServerClientTestFixture::CreateC
 	return Client;
 }
 
- void FReplicationSystemServerClientTestFixture::TearDown()
-{
-	delete Server;
-	for (FReplicationSystemTestClient* Client : Clients)
-	{
-		delete Client;
-	}
-	Clients.Empty();
-	DataStreamUtil.TearDown();
-
-	FNetworkAutomationTestSuiteFixture::TearDown();
-}
-
 }
diff --git a/Engine/Plugins/Runtime/ReplicationSystemTestPlugin/Source/Private/Tests/ReplicationSystem/ReplicationSystemServerClientTestFixture.h b/Engine/Plugins/Runtime/ReplicationSystemTestPlugin/Source/Private/Tests/ReplicationSystem/ReplicationSystemServerClientTestFixture.h
--- a/Engine/Plugins/Runtime/ReplicationSystemTestPlugin/Source/Private/Tests/ReplicationSystem/ReplicationSystemServerClientTestFixture.h
+++ b/Engine/Plugins/Runtime/ReplicationSystemTestPlugin/Source/Private/Tests/ReplicationSystem/ReplicationSystemServerClientTestFixture.h
@@ -73,7 +73,7 @@ public:
 	};
 
 public:
-	explicit FReplicationSystemTestNode(bool bIsServer, const TCHAR* Name);
+	FReplicationSystemTestNode(bool bIsServer, const TCHAR* Name);
 	~FReplicationSystemTestNode();
 
 	template<typename T>
@@ -93,7 +93,7 @@ public:
 	}
 
 	template<typename T>
-	T* CreateSubObject(FNetHandle OwnerHandle, FNetHandle InsertRelativeToSubObjectHandle = FNetHandle(), UReplicationBridge::ESubObjectInsertionOrder InsertionOrder = UReplicationBridge::ESubObjectInsertionOrder::None)
+	T* CreateSubObject(FNetRefHandle OwnerHandle, FNetRefHandle InsertRelativeToSubObjectHandle = FNetRefHandle(), UReplicationBridge::ESubObjectInsertionOrder InsertionOrder = UReplicationBridge::ESubObjectInsertionOrder::None)
 	{
 		T* CreatedObject = NewObject<T>();
 		if (Cast<UReplicatedTestObject>(CreatedObject))
@@ -109,15 +109,15 @@ public:
 	}
 
 	template <typename T>
-	T* GetObjectAs(FNetHandle Handle)
+	T* GetObjectAs(FNetRefHandle Handle)
 	{
 		return static_cast<T*>(ReplicationBridge->GetReplicatedObject(Handle));
 	}
 
 	UTestReplicatedIrisObject* CreateObject(uint32 NumComponents, uint32 NumIrisComponents);
-	UTestReplicatedIrisObject* CreateSubObject(FNetHandle Owner, uint32 NumComponents, uint32 NumIrisComponents);
+	UTestReplicatedIrisObject* CreateSubObject(FNetRefHandle Owner, uint32 NumComponents, uint32 NumIrisComponents);
 	UTestReplicatedIrisObject* CreateObject(const UTestReplicatedIrisObject::FComponents& Components);
-	UTestReplicatedIrisObject* CreateSubObject(FNetHandle Owner, const UTestReplicatedIrisObject::FComponents& Components);
+	UTestReplicatedIrisObject* CreateSubObject(FNetRefHandle Owner, const UTestReplicatedIrisObject::FComponents& Components);
 	UTestReplicatedIrisObject* CreateObjectWithDynamicState(uint32 NumComponents, uint32 NumIrisComponents, uint32 NumDynamicStateComponents);
 	void DestroyObject(UReplicatedTestObject*, EEndReplicationFlags EndReplicationFlags = EEndReplicationFlags::Destroy);
 
@@ -138,6 +138,7 @@ public:
 
 	UReplicatedTestObjectBridge* GetReplicationBridge() { return ReplicationBridge; }
 	UReplicationSystem* GetReplicationSystem() { return ReplicationSystem; }
+	uint32 GetReplicationSystemId() const;
 
 	void SetMaxSendPacketSize(uint32 InMaxSendPacketSize) { MaxSendPacketSize = InMaxSendPacketSize; }
 
@@ -158,7 +159,7 @@ private:
 class FReplicationSystemTestClient : public FReplicationSystemTestNode
 {
 public:
-	FReplicationSystemTestClient(uint32 ReplicationSystemId, const TCHAR* Name);
+	FReplicationSystemTestClient(const TCHAR* Name);
 
 	uint32 ConnectionIdOnServer;
 	uint32 LocalConnectionId;
diff --git a/Engine/Plugins/Runtime/ReplicationSystemTestPlugin/Source/Private/Tests/ReplicationSystem/TestCancelPendingDestroy.cpp b/Engine/Plugins/Runtime/ReplicationSystemTestPlugin/Source/Private/Tests/ReplicationSystem/TestCancelPendingDestroy.cpp
--- a/Engine/Plugins/Runtime/ReplicationSystemTestPlugin/Source/Private/Tests/ReplicationSystem/TestCancelPendingDestroy.cpp
+++ b/Engine/Plugins/Runtime/ReplicationSystemTestPlugin/Source/Private/Tests/ReplicationSystem/TestCancelPendingDestroy.cpp
@@ -26,12 +26,12 @@ UE_NET_TEST_FIXTURE(FTestCancelPendingDestroyFixture, TestCancelPendingDestroyDu
 	Server->PostSendUpdate();
 
 	// Filter out object to cause a PendingDestroy
-	Server->GetReplicationSystem()->AddToGroup(NotReplicatedNetObjectGroupHandle, ServerObject->NetHandle);
+	Server->GetReplicationSystem()->AddToGroup(NotReplicatedNetObjectGroupHandle, ServerObject->NetRefHandle);
 	Server->PreSendUpdate();
 	Server->PostSendUpdate();
 
 	// Remove object from filter to cause object to end up in CancelPendingDestroy
-	Server->GetReplicationSystem()->RemoveFromGroup(NotReplicatedNetObjectGroupHandle, ServerObject->NetHandle);
+	Server->GetReplicationSystem()->RemoveFromGroup(NotReplicatedNetObjectGroupHandle, ServerObject->NetRefHandle);
 	Server->PreSendUpdate();
 	Server->PostSendUpdate();
 
@@ -39,7 +39,7 @@ UE_NET_TEST_FIXTURE(FTestCancelPendingDestroyFixture, TestCancelPendingDestroyDu
 	Server->DeliverTo(Client, DeliverPacket);
 
 	// Verify that the object now exists on client
-	UE_NET_ASSERT_NE(Client->GetReplicationBridge()->GetReplicatedObject(ServerObject->NetHandle), nullptr);
+	UE_NET_ASSERT_NE(Client->GetReplicationBridge()->GetReplicatedObject(ServerObject->NetRefHandle), nullptr);
 
 	// Modify a property on the object and make sure it's replicated as the object should now be confirmed created
 	ServerObject->IntA ^= 1;
@@ -48,7 +48,7 @@ UE_NET_TEST_FIXTURE(FTestCancelPendingDestroyFixture, TestCancelPendingDestroyDu
 	Server->SendAndDeliverTo(Client, DeliverPacket);
 	Server->PostSendUpdate();
 
-	UTestReplicatedIrisObject* ClientObject = Cast<UTestReplicatedIrisObject>(Client->GetReplicationBridge()->GetReplicatedObject(ServerObject->NetHandle));
+	UTestReplicatedIrisObject* ClientObject = Cast<UTestReplicatedIrisObject>(Client->GetReplicationBridge()->GetReplicatedObject(ServerObject->NetRefHandle));
 	UE_NET_ASSERT_NE(ClientObject, nullptr);
 	UE_NET_ASSERT_EQ(ClientObject->IntA, ServerObject->IntA);
 }
@@ -67,12 +67,12 @@ UE_NET_TEST_FIXTURE(FTestCancelPendingDestroyFixture, TestCancelPendingDestroyDu
 	Server->PostSendUpdate();
 
 	// Filter out object to cause a PendingDestroy
-	Server->GetReplicationSystem()->AddToGroup(NotReplicatedNetObjectGroupHandle, ServerObject->NetHandle);
+	Server->GetReplicationSystem()->AddToGroup(NotReplicatedNetObjectGroupHandle, ServerObject->NetRefHandle);
 	Server->PreSendUpdate();
 	Server->PostSendUpdate();
 
 	// Remove object from filter to cause object to end up in CancelPendingDestroy
-	Server->GetReplicationSystem()->RemoveFromGroup(NotReplicatedNetObjectGroupHandle, ServerObject->NetHandle);
+	Server->GetReplicationSystem()->RemoveFromGroup(NotReplicatedNetObjectGroupHandle, ServerObject->NetRefHandle);
 	Server->PreSendUpdate();
 	Server->PostSendUpdate();
 
@@ -84,7 +84,7 @@ UE_NET_TEST_FIXTURE(FTestCancelPendingDestroyFixture, TestCancelPendingDestroyDu
 	Server->SendAndDeliverTo(Client, DeliverPacket);
 	Server->PostSendUpdate();
 
-	UE_NET_ASSERT_NE(Client->GetReplicationBridge()->GetReplicatedObject(ServerObject->NetHandle), nullptr);
+	UE_NET_ASSERT_NE(Client->GetReplicationBridge()->GetReplicatedObject(ServerObject->NetRefHandle), nullptr);
 }
 
 UE_NET_TEST_FIXTURE(FTestCancelPendingDestroyFixture, TestCancelPendingDestroyDuringCreated)
@@ -101,15 +101,15 @@ UE_NET_TEST_FIXTURE(FTestCancelPendingDestroyFixture, TestCancelPendingDestroyDu
 	Server->PostSendUpdate();
 
 	// Verify that the object now also exists on the client
-	UE_NET_ASSERT_NE(Client->GetReplicationBridge()->GetReplicatedObject(ServerObject->NetHandle), nullptr);
+	UE_NET_ASSERT_NE(Client->GetReplicationBridge()->GetReplicatedObject(ServerObject->NetRefHandle), nullptr);
 
 	// Filter out object to cause a PendingDestroy
-	Server->GetReplicationSystem()->AddToGroup(NotReplicatedNetObjectGroupHandle, ServerObject->NetHandle);
+	Server->GetReplicationSystem()->AddToGroup(NotReplicatedNetObjectGroupHandle, ServerObject->NetRefHandle);
 	Server->PreSendUpdate();
 	Server->PostSendUpdate();
 
 	// Remove object from filter to cause object to end up in CancelPendingDestroy
-	Server->GetReplicationSystem()->RemoveFromGroup(NotReplicatedNetObjectGroupHandle, ServerObject->NetHandle);
+	Server->GetReplicationSystem()->RemoveFromGroup(NotReplicatedNetObjectGroupHandle, ServerObject->NetRefHandle);
 	Server->PreSendUpdate();
 	Server->PostSendUpdate();
 
@@ -120,7 +120,7 @@ UE_NET_TEST_FIXTURE(FTestCancelPendingDestroyFixture, TestCancelPendingDestroyDu
 	Server->SendAndDeliverTo(Client, DeliverPacket);
 	Server->PostSendUpdate();
 
-	UTestReplicatedIrisObject* ClientObject = Cast<UTestReplicatedIrisObject>(Client->GetReplicationBridge()->GetReplicatedObject(ServerObject->NetHandle));
+	UTestReplicatedIrisObject* ClientObject = Cast<UTestReplicatedIrisObject>(Client->GetReplicationBridge()->GetReplicatedObject(ServerObject->NetRefHandle));
 	UE_NET_ASSERT_NE(ClientObject, nullptr);
 	UE_NET_ASSERT_EQ(ClientObject->IntA, ServerObject->IntA);
 }
@@ -139,13 +139,13 @@ UE_NET_TEST_FIXTURE(FTestCancelPendingDestroyFixture, TestCancelPendingDestroyDu
 	Server->PostSendUpdate();
 
 	// Filter out object and write packet to cause object to end up in WaitOnDestroyConfirmation
-	Server->GetReplicationSystem()->AddToGroup(NotReplicatedNetObjectGroupHandle, ServerObject->NetHandle);
+	Server->GetReplicationSystem()->AddToGroup(NotReplicatedNetObjectGroupHandle, ServerObject->NetRefHandle);
 	Server->PreSendUpdate();
 	Server->SendUpdate(Client->ConnectionIdOnServer);
 	Server->PostSendUpdate();
 
 	// Remove object from filter to cause object to end up in CancelPendingDestroy
-	Server->GetReplicationSystem()->RemoveFromGroup(NotReplicatedNetObjectGroupHandle, ServerObject->NetHandle);
+	Server->GetReplicationSystem()->RemoveFromGroup(NotReplicatedNetObjectGroupHandle, ServerObject->NetRefHandle);
 	Server->PreSendUpdate();
 	Server->PostSendUpdate();
 
@@ -153,7 +153,7 @@ UE_NET_TEST_FIXTURE(FTestCancelPendingDestroyFixture, TestCancelPendingDestroyDu
 	Server->DeliverTo(Client, DoNotDeliverPacket);
 
 	// Verify that object still exists on client
-	UE_NET_ASSERT_NE(Client->GetReplicationBridge()->GetReplicatedObject(ServerObject->NetHandle), nullptr);
+	UE_NET_ASSERT_NE(Client->GetReplicationBridge()->GetReplicatedObject(ServerObject->NetRefHandle), nullptr);
 
 	// Modify a property on the object and make sure it's replicated as the object should still be created
 	ServerObject->IntA ^= 1;
@@ -162,7 +162,7 @@ UE_NET_TEST_FIXTURE(FTestCancelPendingDestroyFixture, TestCancelPendingDestroyDu
 	Server->SendAndDeliverTo(Client, DeliverPacket);
 	Server->PostSendUpdate();
 
-	UTestReplicatedIrisObject* ClientObject = Cast<UTestReplicatedIrisObject>(Client->GetReplicationBridge()->GetReplicatedObject(ServerObject->NetHandle));
+	UTestReplicatedIrisObject* ClientObject = Cast<UTestReplicatedIrisObject>(Client->GetReplicationBridge()->GetReplicatedObject(ServerObject->NetRefHandle));
 	UE_NET_ASSERT_NE(ClientObject, nullptr);
 	UE_NET_ASSERT_EQ(ClientObject->IntA, ServerObject->IntA);
 }
@@ -181,13 +181,13 @@ UE_NET_TEST_FIXTURE(FTestCancelPendingDestroyFixture, TestCancelPendingDestroyDu
 	Server->PostSendUpdate();
 
 	// Filter out object and write packet to cause object to end up in WaitOnDestroyConfirmation
-	Server->GetReplicationSystem()->AddToGroup(NotReplicatedNetObjectGroupHandle, ServerObject->NetHandle);
+	Server->GetReplicationSystem()->AddToGroup(NotReplicatedNetObjectGroupHandle, ServerObject->NetRefHandle);
 	Server->PreSendUpdate();
 	Server->SendUpdate(Client->ConnectionIdOnServer);
 	Server->PostSendUpdate();
 
 	// Remove object from filter to cause object to end up in CancelPendingDestroy, also modify a property 
-	Server->GetReplicationSystem()->RemoveFromGroup(NotReplicatedNetObjectGroupHandle, ServerObject->NetHandle);
+	Server->GetReplicationSystem()->RemoveFromGroup(NotReplicatedNetObjectGroupHandle, ServerObject->NetRefHandle);
 	ServerObject->IntA ^= 1;
 	Server->PreSendUpdate();
 	Server->PostSendUpdate();
@@ -196,14 +196,14 @@ UE_NET_TEST_FIXTURE(FTestCancelPendingDestroyFixture, TestCancelPendingDestroyDu
 	Server->DeliverTo(Client, DeliverPacket);
 
 	// The destroy was delivered so the client object should not exist
-	UE_NET_ASSERT_EQ(Client->GetReplicationBridge()->GetReplicatedObject(ServerObject->NetHandle), nullptr);
+	UE_NET_ASSERT_EQ(Client->GetReplicationBridge()->GetReplicatedObject(ServerObject->NetRefHandle), nullptr);
 
 	// Write and send packet
 	Server->PreSendUpdate();
 	Server->SendAndDeliverTo(Client, DeliverPacket);
 	Server->PostSendUpdate();
 
-	UTestReplicatedIrisObject* ClientObject = Cast<UTestReplicatedIrisObject>(Client->GetReplicationBridge()->GetReplicatedObject(ServerObject->NetHandle));
+	UTestReplicatedIrisObject* ClientObject = Cast<UTestReplicatedIrisObject>(Client->GetReplicationBridge()->GetReplicatedObject(ServerObject->NetRefHandle));
 	UE_NET_ASSERT_NE(ClientObject, nullptr);
 	UE_NET_ASSERT_EQ(ClientObject->IntA, ServerObject->IntA);
 }
@@ -223,24 +223,24 @@ UE_NET_TEST_FIXTURE(FTestCancelPendingDestroyFixture, TestCancelingCancelPending
 	Server->PostSendUpdate();
 
 	// Filter out object to cause a PendingDestroy
-	Server->GetReplicationSystem()->AddToGroup(NotReplicatedNetObjectGroupHandle, ServerObject->NetHandle);
+	Server->GetReplicationSystem()->AddToGroup(NotReplicatedNetObjectGroupHandle, ServerObject->NetRefHandle);
 	Server->PreSendUpdate();
 	Server->PostSendUpdate();
 
 	// Remove object from filter to cause object to end up in CancelPendingDestroy...
-	Server->GetReplicationSystem()->RemoveFromGroup(NotReplicatedNetObjectGroupHandle, ServerObject->NetHandle);
+	Server->GetReplicationSystem()->RemoveFromGroup(NotReplicatedNetObjectGroupHandle, ServerObject->NetRefHandle);
 	Server->PreSendUpdate();
 	Server->PostSendUpdate();
 
 	// ... and cancel that thought! We want the object to be destroyed after all.
-	Server->GetReplicationSystem()->AddToGroup(NotReplicatedNetObjectGroupHandle, ServerObject->NetHandle);
+	Server->GetReplicationSystem()->AddToGroup(NotReplicatedNetObjectGroupHandle, ServerObject->NetRefHandle);
 	Server->PreSendUpdate();
 	Server->PostSendUpdate();
 
 	Server->DeliverTo(Client, DeliverPacket);
 
 	// Verify that the object now exists on client
-	UE_NET_ASSERT_NE(Client->GetReplicationBridge()->GetReplicatedObject(ServerObject->NetHandle), nullptr);
+	UE_NET_ASSERT_NE(Client->GetReplicationBridge()->GetReplicatedObject(ServerObject->NetRefHandle), nullptr);
 
 	// Write and send a packet
 	Server->PreSendUpdate();
@@ -248,7 +248,7 @@ UE_NET_TEST_FIXTURE(FTestCancelPendingDestroyFixture, TestCancelingCancelPending
 	Server->PostSendUpdate();
 
 	// Verify that the object was destroyed on the client
-	UE_NET_ASSERT_EQ(Client->GetReplicationBridge()->GetReplicatedObject(ServerObject->NetHandle), nullptr);
+	UE_NET_ASSERT_EQ(Client->GetReplicationBridge()->GetReplicatedObject(ServerObject->NetRefHandle), nullptr);
 }
 
 UE_NET_TEST_FIXTURE(FTestCancelPendingDestroyFixture, TestCancelingCancelPendingDestroyDuringWaitOnCreateConfirmationWithPacketLoss)
@@ -265,17 +265,17 @@ UE_NET_TEST_FIXTURE(FTestCancelPendingDestroyFixture, TestCancelingCancelPending
 	Server->PostSendUpdate();
 
 	// Filter out object to cause a PendingDestroy
-	Server->GetReplicationSystem()->AddToGroup(NotReplicatedNetObjectGroupHandle, ServerObject->NetHandle);
+	Server->GetReplicationSystem()->AddToGroup(NotReplicatedNetObjectGroupHandle, ServerObject->NetRefHandle);
 	Server->PreSendUpdate();
 	Server->PostSendUpdate();
 
 	// Remove object from filter to cause object to end up in CancelPendingDestroy...
-	Server->GetReplicationSystem()->RemoveFromGroup(NotReplicatedNetObjectGroupHandle, ServerObject->NetHandle);
+	Server->GetReplicationSystem()->RemoveFromGroup(NotReplicatedNetObjectGroupHandle, ServerObject->NetRefHandle);
 	Server->PreSendUpdate();
 	Server->PostSendUpdate();
 
 	// ... and cancel that thought! We want the object to be destroyed after all.
-	Server->GetReplicationSystem()->AddToGroup(NotReplicatedNetObjectGroupHandle, ServerObject->NetHandle);
+	Server->GetReplicationSystem()->AddToGroup(NotReplicatedNetObjectGroupHandle, ServerObject->NetRefHandle);
 	Server->PreSendUpdate();
 	Server->PostSendUpdate();
 
@@ -283,7 +283,7 @@ UE_NET_TEST_FIXTURE(FTestCancelPendingDestroyFixture, TestCancelingCancelPending
 	Server->DeliverTo(Client, DoNotDeliverPacket);
 
 	// Verify that the object doesn't exist on the client
-	UE_NET_ASSERT_EQ(Client->GetReplicationBridge()->GetReplicatedObject(ServerObject->NetHandle), nullptr);
+	UE_NET_ASSERT_EQ(Client->GetReplicationBridge()->GetReplicatedObject(ServerObject->NetRefHandle), nullptr);
 
 	// Write and send a packet
 	Server->PreSendUpdate();
@@ -291,7 +291,7 @@ UE_NET_TEST_FIXTURE(FTestCancelPendingDestroyFixture, TestCancelingCancelPending
 	Server->PostSendUpdate();
 
 	// Verify that the object still doesn't exist on the client
-	UE_NET_ASSERT_EQ(Client->GetReplicationBridge()->GetReplicatedObject(ServerObject->NetHandle), nullptr);
+	UE_NET_ASSERT_EQ(Client->GetReplicationBridge()->GetReplicatedObject(ServerObject->NetRefHandle), nullptr);
 }
 
 UE_NET_TEST_FIXTURE(FTestCancelPendingDestroyFixture, TestCancelingCancelPendingDestroyDuringCreated)
@@ -308,20 +308,20 @@ UE_NET_TEST_FIXTURE(FTestCancelPendingDestroyFixture, TestCancelingCancelPending
 	Server->PostSendUpdate();
 
 	// Verify that the object now also exists on the client
-	UE_NET_ASSERT_NE(Client->GetReplicationBridge()->GetReplicatedObject(ServerObject->NetHandle), nullptr);
+	UE_NET_ASSERT_NE(Client->GetReplicationBridge()->GetReplicatedObject(ServerObject->NetRefHandle), nullptr);
 
 	// Filter out object to cause a PendingDestroy
-	Server->GetReplicationSystem()->AddToGroup(NotReplicatedNetObjectGroupHandle, ServerObject->NetHandle);
+	Server->GetReplicationSystem()->AddToGroup(NotReplicatedNetObjectGroupHandle, ServerObject->NetRefHandle);
 	Server->PreSendUpdate();
 	Server->PostSendUpdate();
 
 	// Remove object from filter to cause object to end up in CancelPendingDestroy...
-	Server->GetReplicationSystem()->RemoveFromGroup(NotReplicatedNetObjectGroupHandle, ServerObject->NetHandle);
+	Server->GetReplicationSystem()->RemoveFromGroup(NotReplicatedNetObjectGroupHandle, ServerObject->NetRefHandle);
 	Server->PreSendUpdate();
 	Server->PostSendUpdate();
 
 	// ... and cancel that thought! We want the object to be destroyed after all.
-	Server->GetReplicationSystem()->AddToGroup(NotReplicatedNetObjectGroupHandle, ServerObject->NetHandle);
+	Server->GetReplicationSystem()->AddToGroup(NotReplicatedNetObjectGroupHandle, ServerObject->NetRefHandle);
 	Server->PreSendUpdate();
 	Server->PostSendUpdate();
 
@@ -331,7 +331,7 @@ UE_NET_TEST_FIXTURE(FTestCancelPendingDestroyFixture, TestCancelingCancelPending
 	Server->PostSendUpdate();
 
 	// Verify that the object was destroyed on the client
-	UE_NET_ASSERT_EQ(Client->GetReplicationBridge()->GetReplicatedObject(ServerObject->NetHandle), nullptr);
+	UE_NET_ASSERT_EQ(Client->GetReplicationBridge()->GetReplicatedObject(ServerObject->NetRefHandle), nullptr);
 }
 
 UE_NET_TEST_FIXTURE(FTestCancelPendingDestroyFixture, TestCancelingCancelPendingDestroyDuringWaitOnDestroyConfirmationWithPacketLoss)
@@ -348,18 +348,18 @@ UE_NET_TEST_FIXTURE(FTestCancelPendingDestroyFixture, TestCancelingCancelPending
 	Server->PostSendUpdate();
 
 	// Filter out object and write packet to cause object to end up in WaitOnDestroyConfirmation
-	Server->GetReplicationSystem()->AddToGroup(NotReplicatedNetObjectGroupHandle, ServerObject->NetHandle);
+	Server->GetReplicationSystem()->AddToGroup(NotReplicatedNetObjectGroupHandle, ServerObject->NetRefHandle);
 	Server->PreSendUpdate();
 	Server->SendUpdate(Client->ConnectionIdOnServer);
 	Server->PostSendUpdate();
 
 	// Remove object from filter to cause object to end up in CancelPendingDestroy...
-	Server->GetReplicationSystem()->RemoveFromGroup(NotReplicatedNetObjectGroupHandle, ServerObject->NetHandle);
+	Server->GetReplicationSystem()->RemoveFromGroup(NotReplicatedNetObjectGroupHandle, ServerObject->NetRefHandle);
 	Server->PreSendUpdate();
 	Server->PostSendUpdate();
 
 	// ... and cancel that thought! We want the object to be destroyed after all.
-	Server->GetReplicationSystem()->AddToGroup(NotReplicatedNetObjectGroupHandle, ServerObject->NetHandle);
+	Server->GetReplicationSystem()->AddToGroup(NotReplicatedNetObjectGroupHandle, ServerObject->NetRefHandle);
 	Server->PreSendUpdate();
 	Server->PostSendUpdate();
 
@@ -367,7 +367,7 @@ UE_NET_TEST_FIXTURE(FTestCancelPendingDestroyFixture, TestCancelingCancelPending
 	Server->DeliverTo(Client, DoNotDeliverPacket);
 
 	// Verify that the object still exists on client
-	UE_NET_ASSERT_NE(Client->GetReplicationBridge()->GetReplicatedObject(ServerObject->NetHandle), nullptr);
+	UE_NET_ASSERT_NE(Client->GetReplicationBridge()->GetReplicatedObject(ServerObject->NetRefHandle), nullptr);
 
 	// Write and send a packet
 	Server->PreSendUpdate();
@@ -375,7 +375,7 @@ UE_NET_TEST_FIXTURE(FTestCancelPendingDestroyFixture, TestCancelingCancelPending
 	Server->PostSendUpdate();
 
 	// Verify that the object was destroyed on the client
-	UE_NET_ASSERT_EQ(Client->GetReplicationBridge()->GetReplicatedObject(ServerObject->NetHandle), nullptr);
+	UE_NET_ASSERT_EQ(Client->GetReplicationBridge()->GetReplicatedObject(ServerObject->NetRefHandle), nullptr);
 }
 
 UE_NET_TEST_FIXTURE(FTestCancelPendingDestroyFixture, TestCancelingCancelPendingDestroyDuringWaitOnDestroyConfirmationWithoutPacketLoss)
@@ -392,18 +392,18 @@ UE_NET_TEST_FIXTURE(FTestCancelPendingDestroyFixture, TestCancelingCancelPending
 	Server->PostSendUpdate();
 
 	// Filter out object and write packet to cause object to end up in WaitOnDestroyConfirmation
-	Server->GetReplicationSystem()->AddToGroup(NotReplicatedNetObjectGroupHandle, ServerObject->NetHandle);
+	Server->GetReplicationSystem()->AddToGroup(NotReplicatedNetObjectGroupHandle, ServerObject->NetRefHandle);
 	Server->PreSendUpdate();
 	Server->SendUpdate(Client->ConnectionIdOnServer);
 	Server->PostSendUpdate();
 
 	// Remove object from filter to cause object to end up in CancelPendingDestroy...
-	Server->GetReplicationSystem()->RemoveFromGroup(NotReplicatedNetObjectGroupHandle, ServerObject->NetHandle);
+	Server->GetReplicationSystem()->RemoveFromGroup(NotReplicatedNetObjectGroupHandle, ServerObject->NetRefHandle);
 	Server->PreSendUpdate();
 	Server->PostSendUpdate();
 
 	// ... and cancel that thought! We want the object to be destroyed after all.
-	Server->GetReplicationSystem()->AddToGroup(NotReplicatedNetObjectGroupHandle, ServerObject->NetHandle);
+	Server->GetReplicationSystem()->AddToGroup(NotReplicatedNetObjectGroupHandle, ServerObject->NetRefHandle);
 	Server->PreSendUpdate();
 	Server->PostSendUpdate();
 
@@ -411,7 +411,7 @@ UE_NET_TEST_FIXTURE(FTestCancelPendingDestroyFixture, TestCancelingCancelPending
 	Server->DeliverTo(Client, DeliverPacket);
 
 	// Verify that the object doesn't exist on the client
-	UE_NET_ASSERT_EQ(Client->GetReplicationBridge()->GetReplicatedObject(ServerObject->NetHandle), nullptr);
+	UE_NET_ASSERT_EQ(Client->GetReplicationBridge()->GetReplicatedObject(ServerObject->NetRefHandle), nullptr);
 
 	// Write and send a packet
 	Server->PreSendUpdate();
@@ -419,7 +419,7 @@ UE_NET_TEST_FIXTURE(FTestCancelPendingDestroyFixture, TestCancelingCancelPending
 	Server->PostSendUpdate();
 
 	// Verify that the object still doesn't exist on the client
-	UE_NET_ASSERT_EQ(Client->GetReplicationBridge()->GetReplicatedObject(ServerObject->NetHandle), nullptr);
+	UE_NET_ASSERT_EQ(Client->GetReplicationBridge()->GetReplicatedObject(ServerObject->NetRefHandle), nullptr);
 }
 
 
@@ -431,7 +431,7 @@ UE_NET_TEST_FIXTURE(FTestCancelPendingDestroyFixture, TestCancelPendingSubObject
 
 	// Spawn object with subobject on server
 	UTestReplicatedIrisObject* ServerObject = Server->CreateObject(UTestReplicatedIrisObject::FComponents());
-	UTestReplicatedIrisObject* ServerSubObject = Server->CreateSubObject(ServerObject->NetHandle, UTestReplicatedIrisObject::FComponents());
+	UTestReplicatedIrisObject* ServerSubObject = Server->CreateSubObject(ServerObject->NetRefHandle, UTestReplicatedIrisObject::FComponents());
 
 	// Write and send packet
 	Server->PreSendUpdate();
@@ -439,15 +439,15 @@ UE_NET_TEST_FIXTURE(FTestCancelPendingDestroyFixture, TestCancelPendingSubObject
 	Server->PostSendUpdate();
 
 	// Verify that the subobject exists on the client
-	UE_NET_ASSERT_NE(Client->GetReplicationBridge()->GetReplicatedObject(ServerSubObject->NetHandle), nullptr);
+	UE_NET_ASSERT_NE(Client->GetReplicationBridge()->GetReplicatedObject(ServerSubObject->NetRefHandle), nullptr);
 
 	// Filter out subobject to cause a SubObjectPendingDestroy
-	Server->GetReplicationSystem()->AddToGroup(NotReplicatedNetObjectGroupHandle, ServerSubObject->NetHandle);
+	Server->GetReplicationSystem()->AddToGroup(NotReplicatedNetObjectGroupHandle, ServerSubObject->NetRefHandle);
 	Server->PreSendUpdate();
 	Server->PostSendUpdate();
 
 	// Remove subobject from filter to cause object to end up in CancelPendingDestroy
-	Server->GetReplicationSystem()->RemoveFromGroup(NotReplicatedNetObjectGroupHandle, ServerSubObject->NetHandle);
+	Server->GetReplicationSystem()->RemoveFromGroup(NotReplicatedNetObjectGroupHandle, ServerSubObject->NetRefHandle);
 	Server->PreSendUpdate();
 	Server->PostSendUpdate();
 
@@ -458,7 +458,7 @@ UE_NET_TEST_FIXTURE(FTestCancelPendingDestroyFixture, TestCancelPendingSubObject
 	Server->SendAndDeliverTo(Client, DeliverPacket);
 	Server->PostSendUpdate();
 
-	UTestReplicatedIrisObject* ClientSubObject = Cast<UTestReplicatedIrisObject>(Client->GetReplicationBridge()->GetReplicatedObject(ServerSubObject->NetHandle));
+	UTestReplicatedIrisObject* ClientSubObject = Cast<UTestReplicatedIrisObject>(Client->GetReplicationBridge()->GetReplicatedObject(ServerSubObject->NetRefHandle));
 	UE_NET_ASSERT_NE(ClientSubObject, nullptr);
 	UE_NET_ASSERT_EQ(ClientSubObject->IntA, ServerSubObject->IntA);
 }
@@ -470,7 +470,7 @@ UE_NET_TEST_FIXTURE(FTestCancelPendingDestroyFixture, TestCancelPendingSubObject
 
 	// Spawn object on server
 	UTestReplicatedIrisObject* ServerObject = Server->CreateObject(UTestReplicatedIrisObject::FComponents());
-	UTestReplicatedIrisObject* ServerSubObject = Server->CreateSubObject(ServerObject->NetHandle, UTestReplicatedIrisObject::FComponents());
+	UTestReplicatedIrisObject* ServerSubObject = Server->CreateSubObject(ServerObject->NetRefHandle, UTestReplicatedIrisObject::FComponents());
 
 	// Write and send packet
 	Server->PreSendUpdate();
@@ -478,13 +478,13 @@ UE_NET_TEST_FIXTURE(FTestCancelPendingDestroyFixture, TestCancelPendingSubObject
 	Server->PostSendUpdate();
 
 	// Filter out object and write packet to cause subobject to end up in WaitOnDestroyConfirmation
-	Server->GetReplicationSystem()->AddToGroup(NotReplicatedNetObjectGroupHandle, ServerSubObject->NetHandle);
+	Server->GetReplicationSystem()->AddToGroup(NotReplicatedNetObjectGroupHandle, ServerSubObject->NetRefHandle);
 	Server->PreSendUpdate();
 	Server->SendUpdate(Client->ConnectionIdOnServer);
 	Server->PostSendUpdate();
 
 	// Remove subobject from filter to cause object to end up in CancelPendingDestroy
-	Server->GetReplicationSystem()->RemoveFromGroup(NotReplicatedNetObjectGroupHandle, ServerSubObject->NetHandle);
+	Server->GetReplicationSystem()->RemoveFromGroup(NotReplicatedNetObjectGroupHandle, ServerSubObject->NetRefHandle);
 	Server->PreSendUpdate();
 	Server->PostSendUpdate();
 
@@ -492,7 +492,7 @@ UE_NET_TEST_FIXTURE(FTestCancelPendingDestroyFixture, TestCancelPendingSubObject
 	Server->DeliverTo(Client, DoNotDeliverPacket);
 
 	// Verify that object still exists on client
-	UE_NET_ASSERT_NE(Client->GetReplicationBridge()->GetReplicatedObject(ServerSubObject->NetHandle), nullptr);
+	UE_NET_ASSERT_NE(Client->GetReplicationBridge()->GetReplicatedObject(ServerSubObject->NetRefHandle), nullptr);
 
 	// Modify a property on the subobject and make sure it's replicated as the subobject should still exist
 	ServerSubObject->IntA ^= 1;
@@ -501,7 +501,7 @@ UE_NET_TEST_FIXTURE(FTestCancelPendingDestroyFixture, TestCancelPendingSubObject
 	Server->SendAndDeliverTo(Client, DeliverPacket);
 	Server->PostSendUpdate();
 
-	UTestReplicatedIrisObject* ClientSubObject = Cast<UTestReplicatedIrisObject>(Client->GetReplicationBridge()->GetReplicatedObject(ServerSubObject->NetHandle));
+	UTestReplicatedIrisObject* ClientSubObject = Cast<UTestReplicatedIrisObject>(Client->GetReplicationBridge()->GetReplicatedObject(ServerSubObject->NetRefHandle));
 	UE_NET_ASSERT_NE(ClientSubObject, nullptr);
 	UE_NET_ASSERT_EQ(ClientSubObject->IntA, ServerSubObject->IntA);
 }
@@ -513,7 +513,7 @@ UE_NET_TEST_FIXTURE(FTestCancelPendingDestroyFixture, TestCancelPendingSubObject
 
 	// Spawn object on server
 	UTestReplicatedIrisObject* ServerObject = Server->CreateObject(UTestReplicatedIrisObject::FComponents());
-	UTestReplicatedIrisObject* ServerSubObject = Server->CreateSubObject(ServerObject->NetHandle, UTestReplicatedIrisObject::FComponents());
+	UTestReplicatedIrisObject* ServerSubObject = Server->CreateSubObject(ServerObject->NetRefHandle, UTestReplicatedIrisObject::FComponents());
 
 	// Write and send packet
 	Server->PreSendUpdate();
@@ -521,13 +521,13 @@ UE_NET_TEST_FIXTURE(FTestCancelPendingDestroyFixture, TestCancelPendingSubObject
 	Server->PostSendUpdate();
 
 	// Filter out object and write packet to cause subobject to end up in WaitOnDestroyConfirmation
-	Server->GetReplicationSystem()->AddToGroup(NotReplicatedNetObjectGroupHandle, ServerSubObject->NetHandle);
+	Server->GetReplicationSystem()->AddToGroup(NotReplicatedNetObjectGroupHandle, ServerSubObject->NetRefHandle);
 	Server->PreSendUpdate();
 	Server->SendUpdate(Client->ConnectionIdOnServer);
 	Server->PostSendUpdate();
 
 	// Remove object from filter to cause object to end up in CancelPendingDestroy, also modify a property 
-	Server->GetReplicationSystem()->RemoveFromGroup(NotReplicatedNetObjectGroupHandle, ServerSubObject->NetHandle);
+	Server->GetReplicationSystem()->RemoveFromGroup(NotReplicatedNetObjectGroupHandle, ServerSubObject->NetRefHandle);
 	ServerSubObject->IntA ^= 1;
 	Server->PreSendUpdate();
 	Server->PostSendUpdate();
@@ -536,7 +536,7 @@ UE_NET_TEST_FIXTURE(FTestCancelPendingDestroyFixture, TestCancelPendingSubObject
 	Server->DeliverTo(Client, DeliverPacket);
 
 	// The destroy was delivered so the client subobject should not exist
-	UE_NET_ASSERT_EQ(Client->GetReplicationBridge()->GetReplicatedObject(ServerSubObject->NetHandle), nullptr);
+	UE_NET_ASSERT_EQ(Client->GetReplicationBridge()->GetReplicatedObject(ServerSubObject->NetRefHandle), nullptr);
 
 	// Write and send a packet
 	Server->PreSendUpdate();
@@ -544,7 +544,7 @@ UE_NET_TEST_FIXTURE(FTestCancelPendingDestroyFixture, TestCancelPendingSubObject
 	Server->PostSendUpdate();
 
 	// Client subobject should have been created with the latest state
-	UTestReplicatedIrisObject* ClientSubObject = Cast<UTestReplicatedIrisObject>(Client->GetReplicationBridge()->GetReplicatedObject(ServerSubObject->NetHandle));
+	UTestReplicatedIrisObject* ClientSubObject = Cast<UTestReplicatedIrisObject>(Client->GetReplicationBridge()->GetReplicatedObject(ServerSubObject->NetRefHandle));
 	UE_NET_ASSERT_NE(ClientSubObject, nullptr);
 	UE_NET_ASSERT_EQ(ClientSubObject->IntA, ServerSubObject->IntA);
 }
@@ -556,7 +556,7 @@ UE_NET_TEST_FIXTURE(FTestCancelPendingDestroyFixture, TestCancelingCancelPending
 
 	// Spawn object on server
 	UTestReplicatedIrisObject* ServerObject = Server->CreateObject(UTestReplicatedIrisObject::FComponents());
-	UTestReplicatedIrisObject* ServerSubObject = Server->CreateSubObject(ServerObject->NetHandle, UTestReplicatedIrisObject::FComponents());
+	UTestReplicatedIrisObject* ServerSubObject = Server->CreateSubObject(ServerObject->NetRefHandle, UTestReplicatedIrisObject::FComponents());
 
 	// Write and send packet
 	Server->PreSendUpdate();
@@ -564,20 +564,20 @@ UE_NET_TEST_FIXTURE(FTestCancelPendingDestroyFixture, TestCancelingCancelPending
 	Server->PostSendUpdate();
 
 	// Verify that the subobject now also exists on the client
-	UE_NET_ASSERT_NE(Client->GetReplicationBridge()->GetReplicatedObject(ServerSubObject->NetHandle), nullptr);
+	UE_NET_ASSERT_NE(Client->GetReplicationBridge()->GetReplicatedObject(ServerSubObject->NetRefHandle), nullptr);
 
 	// Filter out object to cause a SubObjectPendingDestroy
-	Server->GetReplicationSystem()->AddToGroup(NotReplicatedNetObjectGroupHandle, ServerSubObject->NetHandle);
+	Server->GetReplicationSystem()->AddToGroup(NotReplicatedNetObjectGroupHandle, ServerSubObject->NetRefHandle);
 	Server->PreSendUpdate();
 	Server->PostSendUpdate();
 
 	// Remove object from filter to cause object to end up in CancelPendingDestroy...
-	Server->GetReplicationSystem()->RemoveFromGroup(NotReplicatedNetObjectGroupHandle, ServerSubObject->NetHandle);
+	Server->GetReplicationSystem()->RemoveFromGroup(NotReplicatedNetObjectGroupHandle, ServerSubObject->NetRefHandle);
 	Server->PreSendUpdate();
 	Server->PostSendUpdate();
 
 	// ... and cancel that thought! We want the subobject to be destroyed after all.
-	Server->GetReplicationSystem()->AddToGroup(NotReplicatedNetObjectGroupHandle, ServerSubObject->NetHandle);
+	Server->GetReplicationSystem()->AddToGroup(NotReplicatedNetObjectGroupHandle, ServerSubObject->NetRefHandle);
 	Server->PreSendUpdate();
 	Server->PostSendUpdate();
 
@@ -587,7 +587,7 @@ UE_NET_TEST_FIXTURE(FTestCancelPendingDestroyFixture, TestCancelingCancelPending
 	Server->PostSendUpdate();
 
 	// Verify that the subobject was destroyed on the client
-	UE_NET_ASSERT_EQ(Client->GetReplicationBridge()->GetReplicatedObject(ServerSubObject->NetHandle), nullptr);
+	UE_NET_ASSERT_EQ(Client->GetReplicationBridge()->GetReplicatedObject(ServerSubObject->NetRefHandle), nullptr);
 }
 
 UE_NET_TEST_FIXTURE(FTestCancelPendingDestroyFixture, TestCancelingCancelPendingSubObjectDestroyDuringWaitOnDestroyConfirmationWithPacketLoss)
@@ -597,7 +597,7 @@ UE_NET_TEST_FIXTURE(FTestCancelPendingDestroyFixture, TestCancelingCancelPending
 
 	// Spawn object on server
 	UTestReplicatedIrisObject* ServerObject = Server->CreateObject(UTestReplicatedIrisObject::FComponents());
-	UTestReplicatedIrisObject* ServerSubObject = Server->CreateSubObject(ServerObject->NetHandle, UTestReplicatedIrisObject::FComponents());
+	UTestReplicatedIrisObject* ServerSubObject = Server->CreateSubObject(ServerObject->NetRefHandle, UTestReplicatedIrisObject::FComponents());
 
 	// Write and send packet
 	Server->PreSendUpdate();
@@ -605,18 +605,18 @@ UE_NET_TEST_FIXTURE(FTestCancelPendingDestroyFixture, TestCancelingCancelPending
 	Server->PostSendUpdate();
 
 	// Filter out object and write packet to cause subobject to end up in WaitOnDestroyConfirmation
-	Server->GetReplicationSystem()->AddToGroup(NotReplicatedNetObjectGroupHandle, ServerSubObject->NetHandle);
+	Server->GetReplicationSystem()->AddToGroup(NotReplicatedNetObjectGroupHandle, ServerSubObject->NetRefHandle);
 	Server->PreSendUpdate();
 	Server->SendUpdate(Client->ConnectionIdOnServer);
 	Server->PostSendUpdate();
 
 	// Remove object from filter to cause subobject to end up in CancelPendingDestroy...
-	Server->GetReplicationSystem()->RemoveFromGroup(NotReplicatedNetObjectGroupHandle, ServerSubObject->NetHandle);
+	Server->GetReplicationSystem()->RemoveFromGroup(NotReplicatedNetObjectGroupHandle, ServerSubObject->NetRefHandle);
 	Server->PreSendUpdate();
 	Server->PostSendUpdate();
 
 	// ... and cancel that thought! We want the subobject to be destroyed after all.
-	Server->GetReplicationSystem()->AddToGroup(NotReplicatedNetObjectGroupHandle, ServerSubObject->NetHandle);
+	Server->GetReplicationSystem()->AddToGroup(NotReplicatedNetObjectGroupHandle, ServerSubObject->NetRefHandle);
 	Server->PreSendUpdate();
 	Server->PostSendUpdate();
 
@@ -624,7 +624,7 @@ UE_NET_TEST_FIXTURE(FTestCancelPendingDestroyFixture, TestCancelingCancelPending
 	Server->DeliverTo(Client, DoNotDeliverPacket);
 
 	// Verify that the object still exists on client
-	UE_NET_ASSERT_NE(Client->GetReplicationBridge()->GetReplicatedObject(ServerSubObject->NetHandle), nullptr);
+	UE_NET_ASSERT_NE(Client->GetReplicationBridge()->GetReplicatedObject(ServerSubObject->NetRefHandle), nullptr);
 
 	// Write and send a packet
 	Server->PreSendUpdate();
@@ -632,7 +632,7 @@ UE_NET_TEST_FIXTURE(FTestCancelPendingDestroyFixture, TestCancelingCancelPending
 	Server->PostSendUpdate();
 
 	// Verify that the subobject was destroyed on the client
-	UE_NET_ASSERT_EQ(Client->GetReplicationBridge()->GetReplicatedObject(ServerSubObject->NetHandle), nullptr);
+	UE_NET_ASSERT_EQ(Client->GetReplicationBridge()->GetReplicatedObject(ServerSubObject->NetRefHandle), nullptr);
 }
 
 UE_NET_TEST_FIXTURE(FTestCancelPendingDestroyFixture, TestCancelingCancelPendingSubObjectDestroyDuringWaitOnDestroyConfirmationWithoutPacketLoss)
@@ -642,7 +642,7 @@ UE_NET_TEST_FIXTURE(FTestCancelPendingDestroyFixture, TestCancelingCancelPending
 
 	// Spawn object on server
 	UTestReplicatedIrisObject* ServerObject = Server->CreateObject(UTestReplicatedIrisObject::FComponents());
-	UTestReplicatedIrisObject* ServerSubObject = Server->CreateSubObject(ServerObject->NetHandle, UTestReplicatedIrisObject::FComponents());
+	UTestReplicatedIrisObject* ServerSubObject = Server->CreateSubObject(ServerObject->NetRefHandle, UTestReplicatedIrisObject::FComponents());
 
 	// Write and send packet
 	Server->PreSendUpdate();
@@ -650,18 +650,18 @@ UE_NET_TEST_FIXTURE(FTestCancelPendingDestroyFixture, TestCancelingCancelPending
 	Server->PostSendUpdate();
 
 	// Filter out subobject and write packet to cause subobject to end up in WaitOnDestroyConfirmation
-	Server->GetReplicationSystem()->AddToGroup(NotReplicatedNetObjectGroupHandle, ServerSubObject->NetHandle);
+	Server->GetReplicationSystem()->AddToGroup(NotReplicatedNetObjectGroupHandle, ServerSubObject->NetRefHandle);
 	Server->PreSendUpdate();
 	Server->SendUpdate(Client->ConnectionIdOnServer);
 	Server->PostSendUpdate();
 
 	// Remove subobject from filter to cause subobject to end up in CancelPendingDestroy...
-	Server->GetReplicationSystem()->RemoveFromGroup(NotReplicatedNetObjectGroupHandle, ServerSubObject->NetHandle);
+	Server->GetReplicationSystem()->RemoveFromGroup(NotReplicatedNetObjectGroupHandle, ServerSubObject->NetRefHandle);
 	Server->PreSendUpdate();
 	Server->PostSendUpdate();
 
 	// ... and cancel that thought! We want the subobject to be destroyed after all.
-	Server->GetReplicationSystem()->AddToGroup(NotReplicatedNetObjectGroupHandle, ServerSubObject->NetHandle);
+	Server->GetReplicationSystem()->AddToGroup(NotReplicatedNetObjectGroupHandle, ServerSubObject->NetRefHandle);
 	Server->PreSendUpdate();
 	Server->PostSendUpdate();
 
@@ -669,7 +669,7 @@ UE_NET_TEST_FIXTURE(FTestCancelPendingDestroyFixture, TestCancelingCancelPending
 	Server->DeliverTo(Client, DeliverPacket);
 
 	// Verify that the subobject doesn't exist on the client
-	UE_NET_ASSERT_EQ(Client->GetReplicationBridge()->GetReplicatedObject(ServerSubObject->NetHandle), nullptr);
+	UE_NET_ASSERT_EQ(Client->GetReplicationBridge()->GetReplicatedObject(ServerSubObject->NetRefHandle), nullptr);
 
 	// Write and send a packet
 	Server->PreSendUpdate();
@@ -677,7 +677,7 @@ UE_NET_TEST_FIXTURE(FTestCancelPendingDestroyFixture, TestCancelingCancelPending
 	Server->PostSendUpdate();
 
 	// Verify that the subobject still doesn't exist on the client
-	UE_NET_ASSERT_EQ(Client->GetReplicationBridge()->GetReplicatedObject(ServerSubObject->NetHandle), nullptr);
+	UE_NET_ASSERT_EQ(Client->GetReplicationBridge()->GetReplicatedObject(ServerSubObject->NetRefHandle), nullptr);
 }
 
 UE_NET_TEST_FIXTURE(FTestCancelPendingDestroyFixture, TestCancelingSubObjectDestroyAfterParentIsMarkedForDestroy)
@@ -687,8 +687,8 @@ UE_NET_TEST_FIXTURE(FTestCancelPendingDestroyFixture, TestCancelingSubObjectDest
 
 	// Spawn object on server
 	UTestReplicatedIrisObject* ServerObject = Server->CreateObject(UTestReplicatedIrisObject::FComponents());
-	UTestReplicatedIrisObject* ServerSubObject1 = Server->CreateSubObject(ServerObject->NetHandle, UTestReplicatedIrisObject::FComponents());
-	UTestReplicatedIrisObject* ServerSubObject2 = Server->CreateSubObject(ServerObject->NetHandle, UTestReplicatedIrisObject::FComponents());
+	UTestReplicatedIrisObject* ServerSubObject1 = Server->CreateSubObject(ServerObject->NetRefHandle, UTestReplicatedIrisObject::FComponents());
+	UTestReplicatedIrisObject* ServerSubObject2 = Server->CreateSubObject(ServerObject->NetRefHandle, UTestReplicatedIrisObject::FComponents());
 
 	// Write and send packet
 	Server->PreSendUpdate();
@@ -696,19 +696,19 @@ UE_NET_TEST_FIXTURE(FTestCancelPendingDestroyFixture, TestCancelingSubObjectDest
 	Server->PostSendUpdate();
 
 	// Filter out subobject and write packet to cause subobject to end up in WaitOnDestroyConfirmation
-	Server->GetReplicationSystem()->AddToGroup(NotReplicatedNetObjectGroupHandle, ServerSubObject2->NetHandle);
+	Server->GetReplicationSystem()->AddToGroup(NotReplicatedNetObjectGroupHandle, ServerSubObject2->NetRefHandle);
 	Server->PreSendUpdate();
 	Server->SendUpdate(Client->ConnectionIdOnServer);
 	Server->PostSendUpdate();
 
 	// Filter out parent and write packet
-	Server->GetReplicationSystem()->AddToGroup(NotReplicatedNetObjectGroupHandle, ServerObject->NetHandle);
+	Server->GetReplicationSystem()->AddToGroup(NotReplicatedNetObjectGroupHandle, ServerObject->NetRefHandle);
 	Server->PreSendUpdate();
 	Server->SendUpdate(Client->ConnectionIdOnServer);
 	Server->PostSendUpdate();
 
 	// Remove subobject from filter to cause subobject to end up in CancelPendingDestroy
-	Server->GetReplicationSystem()->RemoveFromGroup(NotReplicatedNetObjectGroupHandle, ServerSubObject2->NetHandle);
+	Server->GetReplicationSystem()->RemoveFromGroup(NotReplicatedNetObjectGroupHandle, ServerSubObject2->NetRefHandle);
 	Server->PreSendUpdate();
 	Server->PostSendUpdate();
 
@@ -716,17 +716,17 @@ UE_NET_TEST_FIXTURE(FTestCancelPendingDestroyFixture, TestCancelingSubObjectDest
 	Server->DeliverTo(Client, DoNotDeliverPacket);
 
 	// Verify all objects and subobjects exist on the client
-	UE_NET_ASSERT_NE(Client->GetReplicationBridge()->GetReplicatedObject(ServerObject->NetHandle), nullptr);
-	UE_NET_ASSERT_NE(Client->GetReplicationBridge()->GetReplicatedObject(ServerSubObject1->NetHandle), nullptr);
-	UE_NET_ASSERT_NE(Client->GetReplicationBridge()->GetReplicatedObject(ServerSubObject2->NetHandle), nullptr);
+	UE_NET_ASSERT_NE(Client->GetReplicationBridge()->GetReplicatedObject(ServerObject->NetRefHandle), nullptr);
+	UE_NET_ASSERT_NE(Client->GetReplicationBridge()->GetReplicatedObject(ServerSubObject1->NetRefHandle), nullptr);
+	UE_NET_ASSERT_NE(Client->GetReplicationBridge()->GetReplicatedObject(ServerSubObject2->NetRefHandle), nullptr);
 
 	// Deliver owner destroy packet
 	Server->DeliverTo(Client, DeliverPacket);
 
 	// Verify that no objects exist on the client
-	UE_NET_ASSERT_EQ(Client->GetReplicationBridge()->GetReplicatedObject(ServerObject->NetHandle), nullptr);
-	UE_NET_ASSERT_EQ(Client->GetReplicationBridge()->GetReplicatedObject(ServerSubObject1->NetHandle), nullptr);
-	UE_NET_ASSERT_EQ(Client->GetReplicationBridge()->GetReplicatedObject(ServerSubObject2->NetHandle), nullptr);
+	UE_NET_ASSERT_EQ(Client->GetReplicationBridge()->GetReplicatedObject(ServerObject->NetRefHandle), nullptr);
+	UE_NET_ASSERT_EQ(Client->GetReplicationBridge()->GetReplicatedObject(ServerSubObject1->NetRefHandle), nullptr);
+	UE_NET_ASSERT_EQ(Client->GetReplicationBridge()->GetReplicatedObject(ServerSubObject2->NetRefHandle), nullptr);
 
 	// Write and send a packet
 	Server->PreSendUpdate();
@@ -734,9 +734,9 @@ UE_NET_TEST_FIXTURE(FTestCancelPendingDestroyFixture, TestCancelingSubObjectDest
 	Server->PostSendUpdate();
 
 	// Verify that all objects are still gone on the client
-	UE_NET_ASSERT_EQ(Client->GetReplicationBridge()->GetReplicatedObject(ServerObject->NetHandle), nullptr);
-	UE_NET_ASSERT_EQ(Client->GetReplicationBridge()->GetReplicatedObject(ServerSubObject1->NetHandle), nullptr);
-	UE_NET_ASSERT_EQ(Client->GetReplicationBridge()->GetReplicatedObject(ServerSubObject2->NetHandle), nullptr);
+	UE_NET_ASSERT_EQ(Client->GetReplicationBridge()->GetReplicatedObject(ServerObject->NetRefHandle), nullptr);
+	UE_NET_ASSERT_EQ(Client->GetReplicationBridge()->GetReplicatedObject(ServerSubObject1->NetRefHandle), nullptr);
+	UE_NET_ASSERT_EQ(Client->GetReplicationBridge()->GetReplicatedObject(ServerSubObject2->NetRefHandle), nullptr);
 }
 
 // See TestObjectSplitting.cpp for cancel pending destroy on huge object
