Iris - Implemented globally unique NetHandles.

diff --git a/Engine/Plugins/Runtime/ReplicationSystemTestPlugin/Source/Private/Tests/ReplicationSystem/TestDependentObjects.cpp b/Engine/Plugins/Runtime/ReplicationSystemTestPlugin/Source/Private/Tests/ReplicationSystem/TestDependentObjects.cpp
--- a/Engine/Plugins/Runtime/ReplicationSystemTestPlugin/Source/Private/Tests/ReplicationSystem/TestDependentObjects.cpp
+++ b/Engine/Plugins/Runtime/ReplicationSystemTestPlugin/Source/Private/Tests/ReplicationSystem/TestDependentObjects.cpp
@@ -34,7 +34,7 @@ UE_NET_TEST_FIXTURE(FReplicationSystemServerClientTestFixture, TestDependentObje
 	Server->SendAndDeliverTo(Client, true);
 	Server->PostSendUpdate();
 
-	Bridge->AddDependentObject(ServerObject->NetHandle, ServerDependentObject->NetHandle);
+	Bridge->AddDependentObject(ServerObject->NetRefHandle, ServerDependentObject->NetRefHandle);
 	ServerDependentObject->IntA = 1;
 
 	// Send and deliver packet
@@ -43,8 +43,8 @@ UE_NET_TEST_FIXTURE(FReplicationSystemServerClientTestFixture, TestDependentObje
 	Server->PostSendUpdate();
 
 	// Store Pointer to objects
-	UTestReplicatedIrisObject* ClientObject = Cast<UTestReplicatedIrisObject>(Client->GetReplicationBridge()->GetReplicatedObject(ServerObject->NetHandle));
-	UTestReplicatedIrisObject* ClientDependentObject = Cast<UTestReplicatedIrisObject>(Client->GetReplicationBridge()->GetReplicatedObject(ServerDependentObject->NetHandle));
+	UTestReplicatedIrisObject* ClientObject = Cast<UTestReplicatedIrisObject>(Client->GetReplicationBridge()->GetReplicatedObject(ServerObject->NetRefHandle));
+	UTestReplicatedIrisObject* ClientDependentObject = Cast<UTestReplicatedIrisObject>(Client->GetReplicationBridge()->GetReplicatedObject(ServerDependentObject->NetRefHandle));
 
 	UE_NET_ASSERT_NE(ClientDependentObject, nullptr);
 	UE_NET_ASSERT_EQ(ServerDependentObject->IntA, ClientDependentObject->IntA);
@@ -84,12 +84,12 @@ UE_NET_TEST_FIXTURE(FReplicationSystemServerClientTestFixture, TestDependentObje
 	// Filter out Server object to start with
 	FNetObjectGroupHandle FilterGroup = ReplicationSystem->CreateGroup();
 	ReplicationSystem->AddGroupFilter(FilterGroup);
-	ReplicationSystem->AddToGroup(FilterGroup, ServerObject->NetHandle);
+	ReplicationSystem->AddToGroup(FilterGroup, ServerObject->NetRefHandle);
 
 	// Setup dependent object to only replicate with ServerObject
 	UTestReplicatedIrisObject* ServerDependentObject = Server->CreateObject(0, 0);
-	ReplicationSystem->SetStaticPriority(ServerDependentObject->NetHandle, 0.f);
-	Bridge->AddDependentObject(ServerObject->NetHandle, ServerDependentObject->NetHandle);
+	ReplicationSystem->SetStaticPriority(ServerDependentObject->NetRefHandle, 0.f);
+	Bridge->AddDependentObject(ServerObject->NetRefHandle, ServerDependentObject->NetRefHandle);
 
 	// Send and deliver packet
 	Server->PreSendUpdate();
@@ -97,8 +97,8 @@ UE_NET_TEST_FIXTURE(FReplicationSystemServerClientTestFixture, TestDependentObje
 	Server->PostSendUpdate();
 
 	// Objects should not be created on client
-	UTestReplicatedIrisObject* ClientObject = Cast<UTestReplicatedIrisObject>(Client->GetReplicationBridge()->GetReplicatedObject(ServerObject->NetHandle));
-	UTestReplicatedIrisObject* ClientDependentObject = Cast<UTestReplicatedIrisObject>(Client->GetReplicationBridge()->GetReplicatedObject(ServerDependentObject->NetHandle));
+	UTestReplicatedIrisObject* ClientObject = Cast<UTestReplicatedIrisObject>(Client->GetReplicationBridge()->GetReplicatedObject(ServerObject->NetRefHandle));
+	UTestReplicatedIrisObject* ClientDependentObject = Cast<UTestReplicatedIrisObject>(Client->GetReplicationBridge()->GetReplicatedObject(ServerDependentObject->NetRefHandle));
 	UE_NET_ASSERT_EQ(ClientObject, nullptr);
 	UE_NET_ASSERT_EQ(ClientDependentObject, nullptr);
 
@@ -111,8 +111,8 @@ UE_NET_TEST_FIXTURE(FReplicationSystemServerClientTestFixture, TestDependentObje
 	Server->PostSendUpdate();
 
 	// Objects should now exist on client
-	ClientObject = Cast<UTestReplicatedIrisObject>(Client->GetReplicationBridge()->GetReplicatedObject(ServerObject->NetHandle));
-	ClientDependentObject = Cast<UTestReplicatedIrisObject>(Client->GetReplicationBridge()->GetReplicatedObject(ServerDependentObject->NetHandle));
+	ClientObject = Cast<UTestReplicatedIrisObject>(Client->GetReplicationBridge()->GetReplicatedObject(ServerObject->NetRefHandle));
+	ClientDependentObject = Cast<UTestReplicatedIrisObject>(Client->GetReplicationBridge()->GetReplicatedObject(ServerDependentObject->NetRefHandle));
 	UE_NET_ASSERT_NE(ClientObject, nullptr);
 	UE_NET_ASSERT_NE(ClientDependentObject, nullptr);
 }
@@ -132,18 +132,18 @@ UE_NET_TEST_FIXTURE(FReplicationSystemServerClientTestFixture, TestChainedDepend
 	// Filter out Server object to start with
 	FNetObjectGroupHandle FilterGroup = ReplicationSystem->CreateGroup();
 	ReplicationSystem->AddGroupFilter(FilterGroup);
-	ReplicationSystem->AddToGroup(FilterGroup, ServerObject->NetHandle);
+	ReplicationSystem->AddToGroup(FilterGroup, ServerObject->NetRefHandle);
 
 	// Setup dependent object to only replicate with ServerObject
 	UTestReplicatedIrisObject* ServerDependentObject0 = Server->CreateObject(0, 0);
 	UTestReplicatedIrisObject* ServerDependentObject1 = Server->CreateObject(0, 0);
 	
-	ReplicationSystem->SetStaticPriority(ServerDependentObject0->NetHandle, 0.f);
-	ReplicationSystem->SetStaticPriority(ServerDependentObject1->NetHandle, 0.f);
+	ReplicationSystem->SetStaticPriority(ServerDependentObject0->NetRefHandle, 0.f);
+	ReplicationSystem->SetStaticPriority(ServerDependentObject1->NetRefHandle, 0.f);
 	
 	// Setup dependency chain
-	Bridge->AddDependentObject(ServerObject->NetHandle, ServerDependentObject0->NetHandle);
-	Bridge->AddDependentObject(ServerDependentObject0->NetHandle, ServerDependentObject1->NetHandle);
+	Bridge->AddDependentObject(ServerObject->NetRefHandle, ServerDependentObject0->NetRefHandle);
+	Bridge->AddDependentObject(ServerDependentObject0->NetRefHandle, ServerDependentObject1->NetRefHandle);
 
 	// Send and deliver packet
 	Server->PreSendUpdate();
@@ -151,9 +151,9 @@ UE_NET_TEST_FIXTURE(FReplicationSystemServerClientTestFixture, TestChainedDepend
 	Server->PostSendUpdate();
 
 	// Objects should not be created on client
-	UTestReplicatedIrisObject* ClientObject = Cast<UTestReplicatedIrisObject>(Client->GetReplicationBridge()->GetReplicatedObject(ServerObject->NetHandle));
-	UTestReplicatedIrisObject* ClientDependentObject0 = Cast<UTestReplicatedIrisObject>(Client->GetReplicationBridge()->GetReplicatedObject(ServerDependentObject0->NetHandle));
-	UTestReplicatedIrisObject* ClientDependentObject1 = Cast<UTestReplicatedIrisObject>(Client->GetReplicationBridge()->GetReplicatedObject(ServerDependentObject1->NetHandle));
+	UTestReplicatedIrisObject* ClientObject = Cast<UTestReplicatedIrisObject>(Client->GetReplicationBridge()->GetReplicatedObject(ServerObject->NetRefHandle));
+	UTestReplicatedIrisObject* ClientDependentObject0 = Cast<UTestReplicatedIrisObject>(Client->GetReplicationBridge()->GetReplicatedObject(ServerDependentObject0->NetRefHandle));
+	UTestReplicatedIrisObject* ClientDependentObject1 = Cast<UTestReplicatedIrisObject>(Client->GetReplicationBridge()->GetReplicatedObject(ServerDependentObject1->NetRefHandle));
 
 	UE_NET_ASSERT_EQ(ClientObject, nullptr);
 	UE_NET_ASSERT_EQ(ClientDependentObject0, nullptr);
@@ -168,9 +168,9 @@ UE_NET_TEST_FIXTURE(FReplicationSystemServerClientTestFixture, TestChainedDepend
 	Server->PostSendUpdate();
 
 	// Verify that dependent object now is created
-	ClientObject = Cast<UTestReplicatedIrisObject>(Client->GetReplicationBridge()->GetReplicatedObject(ServerObject->NetHandle));
-	ClientDependentObject0 = Cast<UTestReplicatedIrisObject>(Client->GetReplicationBridge()->GetReplicatedObject(ServerDependentObject0->NetHandle));
-	ClientDependentObject1 = Cast<UTestReplicatedIrisObject>(Client->GetReplicationBridge()->GetReplicatedObject(ServerDependentObject1->NetHandle));
+	ClientObject = Cast<UTestReplicatedIrisObject>(Client->GetReplicationBridge()->GetReplicatedObject(ServerObject->NetRefHandle));
+	ClientDependentObject0 = Cast<UTestReplicatedIrisObject>(Client->GetReplicationBridge()->GetReplicatedObject(ServerDependentObject0->NetRefHandle));
+	ClientDependentObject1 = Cast<UTestReplicatedIrisObject>(Client->GetReplicationBridge()->GetReplicatedObject(ServerDependentObject1->NetRefHandle));
 	UE_NET_ASSERT_NE(ClientObject, nullptr);
 	UE_NET_ASSERT_NE(ClientDependentObject0, nullptr);
 	UE_NET_ASSERT_NE(ClientDependentObject1, nullptr);
@@ -197,8 +197,8 @@ UE_NET_TEST_FIXTURE(FReplicationSystemServerClientTestFixture, TestDependentObje
 	Server->PostSendUpdate();
 
 	// Store Pointer to objects and verify state after initial replication
-	UTestReplicatedIrisObject* ClientObject = Cast<UTestReplicatedIrisObject>(Client->GetReplicationBridge()->GetReplicatedObject(ServerObject->NetHandle));
-	UTestReplicatedIrisObject* ClientDependentObject = Cast<UTestReplicatedIrisObject>(Client->GetReplicationBridge()->GetReplicatedObject(ServerDependentObject->NetHandle));
+	UTestReplicatedIrisObject* ClientObject = Cast<UTestReplicatedIrisObject>(Client->GetReplicationBridge()->GetReplicatedObject(ServerObject->NetRefHandle));
+	UTestReplicatedIrisObject* ClientDependentObject = Cast<UTestReplicatedIrisObject>(Client->GetReplicationBridge()->GetReplicatedObject(ServerDependentObject->NetRefHandle));
 	
 	UE_NET_ASSERT_NE(ClientDependentObject, nullptr);
 	UE_NET_ASSERT_NE(ClientObject, nullptr);
@@ -219,7 +219,7 @@ UE_NET_TEST_FIXTURE(FReplicationSystemServerClientTestFixture, TestDependentObje
 	UE_NET_ASSERT_NE(ClientDependentObject->IntA, ServerDependentObject->IntA);
 
 	// Add dependency
-	Server->ReplicationBridge->AddDependentObject(ServerObject->NetHandle, ServerDependentObject->NetHandle);
+	Server->ReplicationBridge->AddDependentObject(ServerObject->NetRefHandle, ServerDependentObject->NetRefHandle);
 
 	// Change a value on owner
 	ServerObject->IntA = 2;
@@ -254,8 +254,8 @@ UE_NET_TEST_FIXTURE(FReplicationSystemServerClientTestFixture, TestDependentObje
 	Server->ReplicationBridge->SetPollFramePeriod(ServerDependentObject, 40U);
 
 	// Add dependent object to both server objects
-	Server->ReplicationBridge->AddDependentObject(ServerObject0->NetHandle, ServerDependentObject->NetHandle);
-	Server->ReplicationBridge->AddDependentObject(ServerObject1->NetHandle, ServerDependentObject->NetHandle);
+	Server->ReplicationBridge->AddDependentObject(ServerObject0->NetRefHandle, ServerDependentObject->NetRefHandle);
+	Server->ReplicationBridge->AddDependentObject(ServerObject1->NetRefHandle, ServerDependentObject->NetRefHandle);
 
 	// Send and deliver packet, All objects are polled and replicated
 	Server->PreSendUpdate();
@@ -263,9 +263,9 @@ UE_NET_TEST_FIXTURE(FReplicationSystemServerClientTestFixture, TestDependentObje
 	Server->PostSendUpdate();
 
 	// Verify state after initial replication
-	UTestReplicatedIrisObject* ClientDependentObject = Cast<UTestReplicatedIrisObject>(Client->GetReplicationBridge()->GetReplicatedObject(ServerDependentObject->NetHandle));
-	UTestReplicatedIrisObject* ClientObject0 = Cast<UTestReplicatedIrisObject>(Client->GetReplicationBridge()->GetReplicatedObject(ServerObject0->NetHandle));
-	UTestReplicatedIrisObject* ClientObject1 = Cast<UTestReplicatedIrisObject>(Client->GetReplicationBridge()->GetReplicatedObject(ServerObject1->NetHandle));
+	UTestReplicatedIrisObject* ClientDependentObject = Cast<UTestReplicatedIrisObject>(Client->GetReplicationBridge()->GetReplicatedObject(ServerDependentObject->NetRefHandle));
+	UTestReplicatedIrisObject* ClientObject0 = Cast<UTestReplicatedIrisObject>(Client->GetReplicationBridge()->GetReplicatedObject(ServerObject0->NetRefHandle));
+	UTestReplicatedIrisObject* ClientObject1 = Cast<UTestReplicatedIrisObject>(Client->GetReplicationBridge()->GetReplicatedObject(ServerObject1->NetRefHandle));
 
 	UE_NET_ASSERT_NE(ClientDependentObject, nullptr);
 	UE_NET_ASSERT_NE(ClientObject0, nullptr);
@@ -323,7 +323,7 @@ UE_NET_TEST_FIXTURE(FReplicationSystemServerClientTestFixture, TestDependentObje
 	// Spawn second object add it as a dependency and bump poll period
 	UTestReplicatedIrisObject* ServerDependentObject = Server->CreateObject(0, 0);
 	Server->ReplicationBridge->SetPollFramePeriod(ServerDependentObject, 10U);
-	Server->ReplicationBridge->AddDependentObject(ServerObject->NetHandle, ServerDependentObject->NetHandle);
+	Server->ReplicationBridge->AddDependentObject(ServerObject->NetRefHandle, ServerDependentObject->NetRefHandle);
 
 	// Send and deliver packet
 	Server->PreSendUpdate();
@@ -331,8 +331,8 @@ UE_NET_TEST_FIXTURE(FReplicationSystemServerClientTestFixture, TestDependentObje
 	Server->PostSendUpdate();
 
 	// Store Pointer to objects and verify state after initial replication
-	UTestReplicatedIrisObject* ClientObject = Cast<UTestReplicatedIrisObject>(Client->GetReplicationBridge()->GetReplicatedObject(ServerObject->NetHandle));
-	UTestReplicatedIrisObject* ClientDependentObject = Cast<UTestReplicatedIrisObject>(Client->GetReplicationBridge()->GetReplicatedObject(ServerDependentObject->NetHandle));
+	UTestReplicatedIrisObject* ClientObject = Cast<UTestReplicatedIrisObject>(Client->GetReplicationBridge()->GetReplicatedObject(ServerObject->NetRefHandle));
+	UTestReplicatedIrisObject* ClientDependentObject = Cast<UTestReplicatedIrisObject>(Client->GetReplicationBridge()->GetReplicatedObject(ServerDependentObject->NetRefHandle));
 
 	// Modify data
 	ServerObject->IntA = 1;
@@ -348,7 +348,7 @@ UE_NET_TEST_FIXTURE(FReplicationSystemServerClientTestFixture, TestDependentObje
 	UE_NET_ASSERT_NE(ClientDependentObject->IntA, ServerDependentObject->IntA);
 
 	// Mark dependent dirty
-	ReplicationSystem->MarkDirty(ServerDependentObject->NetHandle);
+	ReplicationSystem->MarkDirty(ServerDependentObject->NetRefHandle);
 
 	// Send and deliver packet, we expect dependent object to have replicated 
 	Server->PreSendUpdate();
@@ -364,7 +364,7 @@ UE_NET_TEST_FIXTURE(FReplicationSystemServerClientTestFixture, TestDependentObje
 	ServerDependentObject->IntA = 2;
 
 	// Mark parent dirty
-	ReplicationSystem->MarkDirty(ServerObject->NetHandle);
+	ReplicationSystem->MarkDirty(ServerObject->NetRefHandle);
 
 	// Send and deliver packet, we expect both objects to have replicated
 	Server->PreSendUpdate();
diff --git a/Engine/Plugins/Runtime/ReplicationSystemTestPlugin/Source/Private/Tests/ReplicationSystem/TestDirtyNetObjectTracker.cpp b/Engine/Plugins/Runtime/ReplicationSystemTestPlugin/Source/Private/Tests/ReplicationSystem/TestDirtyNetObjectTracker.cpp
--- a/Engine/Plugins/Runtime/ReplicationSystemTestPlugin/Source/Private/Tests/ReplicationSystem/TestDirtyNetObjectTracker.cpp
+++ b/Engine/Plugins/Runtime/ReplicationSystemTestPlugin/Source/Private/Tests/ReplicationSystem/TestDirtyNetObjectTracker.cpp
@@ -3,80 +3,60 @@
 #include "NetworkAutomationTest.h"
 #include "NetworkAutomationTestMacros.h"
 #include "Iris/ReplicationSystem/DirtyNetObjectTracker.h"
-#include "Iris/ReplicationSystem/NetHandle.h"
-#include "Iris/ReplicationSystem/ReplicationSystem.h"
+#include "Iris/ReplicationSystem/ReplicationSystemInternal.h"
+#include "Tests/ReplicationSystem/ReplicationSystemTestFixture.h"
 
 namespace UE::Net::Private
 {
 
-class FDirtyNetObjectTrackerTestFixture : public FNetworkAutomationTestSuiteFixture
+class FDirtyNetObjectTrackerTestFixture : public FReplicationSystemTestFixture
 {
-public:
-	FDirtyNetObjectTrackerTestFixture()
-	: FNetworkAutomationTestSuiteFixture()
-	, DirtyNetObjectTracker(nullptr)
-	{
-	}
+	using Super = FReplicationSystemTestFixture;
 
+public:
 	virtual void SetUp() override
 	{
-		DirtyNetObjectTracker = new FDirtyNetObjectTracker();
+		Super::SetUp();
+		DirtyNetObjectTracker = &ReplicationSystem->GetReplicationSystemInternal()->GetDirtyNetObjectTracker();
+		ReplicationSystemId = ReplicationSystem->GetId();
+		NetObjectIndexRangeEnd = DirtyNetObjectTracker->GetDirtyNetObjects().GetNumBits() - 1U;
 	}
 
 	virtual void TearDown() override
 	{
-		delete DirtyNetObjectTracker;
+		DirtyNetObjectTracker = nullptr;
+		ReplicationSystemId = ~0U;
+		Super::TearDown();
 	}
 
 protected:
-	static constexpr uint32 ReplicationSystemId = 0;
 	static constexpr uint32 NetObjectIndexRangeStart = 1;
-	static constexpr uint32 NetObjectIndexRangeEnd = 1000;
-
-	FDirtyNetObjectTracker* DirtyNetObjectTracker;
-
-};
-
-class FDirtyNetObjectTrackerWithInitTestFixture : public FDirtyNetObjectTrackerTestFixture
-{
-protected:
-	virtual void SetUp() override
-	{
-		FDirtyNetObjectTrackerTestFixture::SetUp();
-		Init();
-	}
 
-	void Init();
+	FDirtyNetObjectTracker* DirtyNetObjectTracker = nullptr;
+	uint32 ReplicationSystemId = ~0U;
+	uint32 NetObjectIndexRangeEnd = 1;
 };
 
-UE_NET_TEST_FIXTURE(FDirtyNetObjectTrackerTestFixture, TestInit)
+UE_NET_TEST_FIXTURE(FDirtyNetObjectTrackerTestFixture, TestNoObjectIsDirtyFromStart)
 {
-	const FDirtyNetObjectTrackerInitParams InitParams 
-	{
-		ReplicationSystemId,
-		NetObjectIndexRangeStart,
-		NetObjectIndexRangeEnd,
-	};
-
-	DirtyNetObjectTracker->Init(InitParams);
 	const FNetBitArrayView& DirtyObjects = DirtyNetObjectTracker->GetDirtyNetObjects();
 	UE_NET_ASSERT_EQ(DirtyObjects.GetNumBits(), NetObjectIndexRangeEnd + 1U);
 
 	UE_NET_ASSERT_FALSE(DirtyObjects.IsAnyBitSet()) << "Objects are marked as dirty before marking any object state as dirty";
 }
 
-UE_NET_TEST_FIXTURE(FDirtyNetObjectTrackerWithInitTestFixture, CannotMarkInvalidObjectAsDirty)
+UE_NET_TEST_FIXTURE(FDirtyNetObjectTrackerTestFixture, CannotMarkInvalidObjectAsDirty)
 {
-	MarkNetObjectStateDirty(ReplicationSystemId, FNetHandle().GetId());
+	MarkNetObjectStateDirty(ReplicationSystemId, FNetRefHandle().GetId());
 	MarkNetObjectStateDirty(ReplicationSystemId, NetObjectIndexRangeEnd + 1);
 
 	const FNetBitArrayView& DirtyObjects = DirtyNetObjectTracker->GetDirtyNetObjects();
 	UE_NET_ASSERT_FALSE(DirtyObjects.IsAnyBitSet());
 }
 
-UE_NET_TEST_FIXTURE(FDirtyNetObjectTrackerWithInitTestFixture, CanMarkValidObjectAsDirty)
+UE_NET_TEST_FIXTURE(FDirtyNetObjectTrackerTestFixture, CanMarkValidObjectAsDirty)
 {
-	constexpr uint32 IndexInRange = NetObjectIndexRangeStart + (NetObjectIndexRangeEnd - NetObjectIndexRangeStart)/2;
+	const uint32 IndexInRange = NetObjectIndexRangeStart + (NetObjectIndexRangeEnd - NetObjectIndexRangeStart)/2;
 	MarkNetObjectStateDirty(ReplicationSystemId, IndexInRange);
 	MarkNetObjectStateDirty(ReplicationSystemId, NetObjectIndexRangeStart);
 	MarkNetObjectStateDirty(ReplicationSystemId, NetObjectIndexRangeEnd);
@@ -87,7 +67,7 @@ UE_NET_TEST_FIXTURE(FDirtyNetObjectTrackerWithInitTestFixture, CanMarkValidObjec
 	UE_NET_ASSERT_TRUE(DirtyObjects.GetBit(NetObjectIndexRangeEnd));
 }
 
-UE_NET_TEST_FIXTURE(FDirtyNetObjectTrackerWithInitTestFixture, CanClearDirtyObjects)
+UE_NET_TEST_FIXTURE(FDirtyNetObjectTrackerTestFixture, CanClearDirtyObjects)
 {
 	MarkNetObjectStateDirty(ReplicationSystemId, NetObjectIndexRangeStart);
 	MarkNetObjectStateDirty(ReplicationSystemId, NetObjectIndexRangeEnd);
@@ -100,21 +80,8 @@ UE_NET_TEST_FIXTURE(FDirtyNetObjectTrackerWithInitTestFixture, CanClearDirtyObje
 UE_NET_TEST(DirtyNetObjectTracker, MarkingObjectAsDirtyInNonExistingSystemDoesNotCrash)
 {
 	constexpr uint32 NonExistingReplicationSystemId = 4711;
-	constexpr uint32 ArbitraryNetObjectIndex = 1;
+	constexpr uint32 ArbitraryNetObjectIndex = 1174;
 	MarkNetObjectStateDirty(NonExistingReplicationSystemId, ArbitraryNetObjectIndex);
 }
 
-//
-void FDirtyNetObjectTrackerWithInitTestFixture::Init()
-{
-	const FDirtyNetObjectTrackerInitParams InitParams 
-	{
-		ReplicationSystemId,
-		NetObjectIndexRangeStart,
-		NetObjectIndexRangeEnd,
-	};
-
-	DirtyNetObjectTracker->Init(InitParams);
-}
-
 }
diff --git a/Engine/Plugins/Runtime/ReplicationSystemTestPlugin/Source/Private/Tests/ReplicationSystem/TestFlushBeforeDestroy.cpp b/Engine/Plugins/Runtime/ReplicationSystemTestPlugin/Source/Private/Tests/ReplicationSystem/TestFlushBeforeDestroy.cpp
--- a/Engine/Plugins/Runtime/ReplicationSystemTestPlugin/Source/Private/Tests/ReplicationSystem/TestFlushBeforeDestroy.cpp
+++ b/Engine/Plugins/Runtime/ReplicationSystemTestPlugin/Source/Private/Tests/ReplicationSystem/TestFlushBeforeDestroy.cpp
@@ -19,19 +19,19 @@ UE_NET_TEST_FIXTURE(FTestFlushBeforeDestroyFixture, TestReliableAttachmentFlushe
 	RegisterNetBlobHandlers(Client);
 
 	UReplicatedTestObject* ServerObject = Server->CreateObject(0, 0);
-	FNetHandle ObjectHandle = ServerObject->NetHandle;
+	FNetRefHandle ObjectHandle = ServerObject->NetRefHandle;
 
 	Server->PreSendUpdate();
 	Server->SendAndDeliverTo(Client, DeliverPacket);
 	Server->PostSendUpdate();
 
-	UE_NET_ASSERT_NE(Client->GetReplicationBridge()->GetReplicatedObject(ServerObject->NetHandle), nullptr);
+	UE_NET_ASSERT_NE(Client->GetReplicationBridge()->GetReplicatedObject(ServerObject->NetRefHandle), nullptr);
 
 	// Create attachment
 	{
 		constexpr uint32 PayloadBitCount = 24;
 		const TRefCountPtr<FNetObjectAttachment>& Attachment = MockNetObjectAttachmentHandler->CreateReliableNetObjectAttachment(PayloadBitCount);
-		FNetObjectReference AttachmentTarget = FObjectReferenceCache::MakeNetObjectReference(ServerObject->NetHandle);
+		FNetObjectReference AttachmentTarget = FObjectReferenceCache::MakeNetObjectReference(ServerObject->NetRefHandle);
 		Server->GetReplicationSystem()->QueueNetObjectAttachment(Client->ConnectionIdOnServer, AttachmentTarget, Attachment);
 	}
 
@@ -62,19 +62,19 @@ UE_NET_TEST_FIXTURE(FTestFlushBeforeDestroyFixture, TestReliableAttachmentFlushe
 	RegisterNetBlobHandlers(Client);
 
 	UReplicatedTestObject* ServerObject = Server->CreateObject(0, 0);
-	FNetHandle ObjectHandle = ServerObject->NetHandle;
+	FNetRefHandle ObjectHandle = ServerObject->NetRefHandle;
 
 	Server->PreSendUpdate();
 	Server->SendAndDeliverTo(Client, DeliverPacket);
 	Server->PostSendUpdate();
 
-	UE_NET_ASSERT_NE(Client->GetReplicationBridge()->GetReplicatedObject(ServerObject->NetHandle), nullptr);
+	UE_NET_ASSERT_NE(Client->GetReplicationBridge()->GetReplicatedObject(ServerObject->NetRefHandle), nullptr);
 
 	// Create attachment
 	{
 		constexpr uint32 PayloadBitCount = 24;
 		const TRefCountPtr<FNetObjectAttachment>& Attachment = MockNetObjectAttachmentHandler->CreateReliableNetObjectAttachment(PayloadBitCount);
-		FNetObjectReference AttachmentTarget = FObjectReferenceCache::MakeNetObjectReference(ServerObject->NetHandle);
+		FNetObjectReference AttachmentTarget = FObjectReferenceCache::MakeNetObjectReference(ServerObject->NetRefHandle);
 		Server->GetReplicationSystem()->QueueNetObjectAttachment(Client->ConnectionIdOnServer, AttachmentTarget, Attachment);
 	}
 
@@ -112,21 +112,21 @@ UE_NET_TEST_FIXTURE(FTestFlushBeforeDestroyFixture, TestReliableAttachmentSubObj
 	RegisterNetBlobHandlers(Client);
 
 	UReplicatedTestObject* ServerObject = Server->CreateObject(0, 0);
-	FNetHandle ObjectHandle = ServerObject->NetHandle;
+	FNetRefHandle ObjectHandle = ServerObject->NetRefHandle;
 	UReplicatedTestObject* ServerSubObject = Server->CreateSubObject(ObjectHandle, 0, 0);
-	FNetHandle SubObjectHandle = ServerSubObject->NetHandle;
+	FNetRefHandle SubObjectHandle = ServerSubObject->NetRefHandle;
 	
 	Server->PreSendUpdate();
 	Server->SendAndDeliverTo(Client, DeliverPacket);
 	Server->PostSendUpdate();
 
-	UE_NET_ASSERT_NE(Client->GetReplicationBridge()->GetReplicatedObject(ServerObject->NetHandle), nullptr);
+	UE_NET_ASSERT_NE(Client->GetReplicationBridge()->GetReplicatedObject(ServerObject->NetRefHandle), nullptr);
 
 	// Create attachment
 	{
 		constexpr uint32 PayloadBitCount = 24;
 		const TRefCountPtr<FNetObjectAttachment>& Attachment = MockNetObjectAttachmentHandler->CreateReliableNetObjectAttachment(PayloadBitCount);
-		FNetObjectReference AttachmentTarget = FObjectReferenceCache::MakeNetObjectReference(ServerSubObject->NetHandle);
+		FNetObjectReference AttachmentTarget = FObjectReferenceCache::MakeNetObjectReference(ServerSubObject->NetRefHandle);
 		Server->GetReplicationSystem()->QueueNetObjectAttachment(Client->ConnectionIdOnServer, AttachmentTarget, Attachment);
 	}
 
@@ -156,21 +156,21 @@ UE_NET_TEST_FIXTURE(FTestFlushBeforeDestroyFixture, TestReliableAttachmentSubObj
 	RegisterNetBlobHandlers(Client);
 
 	UReplicatedTestObject* ServerObject = Server->CreateObject(0, 0);
-	FNetHandle ObjectHandle = ServerObject->NetHandle;
+	FNetRefHandle ObjectHandle = ServerObject->NetRefHandle;
 	UReplicatedTestObject* ServerSubObject = Server->CreateSubObject(ObjectHandle, 0, 0);
-	FNetHandle SubObjectHandle = ServerSubObject->NetHandle;
+	FNetRefHandle SubObjectHandle = ServerSubObject->NetRefHandle;
 	
 	Server->PreSendUpdate();
 	Server->SendAndDeliverTo(Client, DeliverPacket);
 	Server->PostSendUpdate();
 
-	UE_NET_ASSERT_NE(Client->GetReplicationBridge()->GetReplicatedObject(ServerObject->NetHandle), nullptr);
+	UE_NET_ASSERT_NE(Client->GetReplicationBridge()->GetReplicatedObject(ServerObject->NetRefHandle), nullptr);
 
 	// Create attachment
 	{
 		constexpr uint32 PayloadBitCount = 24;
 		const TRefCountPtr<FNetObjectAttachment>& Attachment = MockNetObjectAttachmentHandler->CreateReliableNetObjectAttachment(PayloadBitCount);
-		FNetObjectReference AttachmentTarget = FObjectReferenceCache::MakeNetObjectReference(ServerSubObject->NetHandle);
+		FNetObjectReference AttachmentTarget = FObjectReferenceCache::MakeNetObjectReference(ServerSubObject->NetRefHandle);
 		Server->GetReplicationSystem()->QueueNetObjectAttachment(Client->ConnectionIdOnServer, AttachmentTarget, Attachment);
 	}
 
@@ -201,7 +201,7 @@ UE_NET_TEST_FIXTURE(FTestFlushBeforeDestroyFixture, TestStateFlushedBeforeDestro
 	RegisterNetBlobHandlers(Client);
 
 	UTestReplicatedIrisObject* ServerObject = Server->CreateObject(0, 0);
-	FNetHandle ObjectHandle = ServerObject->NetHandle;
+	FNetRefHandle ObjectHandle = ServerObject->NetRefHandle;
 	
 	Server->PreSendUpdate();
 	Server->SendAndDeliverTo(Client, DeliverPacket);
@@ -244,7 +244,7 @@ UE_NET_TEST_FIXTURE(FTestFlushBeforeDestroyFixture, TestStateInFlightFlushedBefo
 	RegisterNetBlobHandlers(Client);
 
 	UTestReplicatedIrisObject* ServerObject = Server->CreateObject(0, 0);
-	FNetHandle ObjectHandle = ServerObject->NetHandle;
+	FNetRefHandle ObjectHandle = ServerObject->NetRefHandle;
 	
 	Server->PreSendUpdate();
 	Server->SendAndDeliverTo(Client, DeliverPacket);
@@ -298,9 +298,9 @@ UE_NET_TEST_FIXTURE(FTestFlushBeforeDestroyFixture, TestSubObjectStateFlushedBef
 	RegisterNetBlobHandlers(Client);
 
 	UTestReplicatedIrisObject* ServerObject = Server->CreateObject(0, 0);
-	FNetHandle ObjectHandle = ServerObject->NetHandle;
+	FNetRefHandle ObjectHandle = ServerObject->NetRefHandle;
 	UTestReplicatedIrisObject* ServerSubObject = Server->CreateSubObject(ObjectHandle, 0, 0);
-	FNetHandle SubObjectHandle = ServerSubObject->NetHandle;
+	FNetRefHandle SubObjectHandle = ServerSubObject->NetRefHandle;
 
 	Server->PreSendUpdate();
 	Server->SendAndDeliverTo(Client, DeliverPacket);
@@ -347,9 +347,9 @@ UE_NET_TEST_FIXTURE(FTestFlushBeforeDestroyFixture, TestSubObjectStateFlushedBef
 	RegisterNetBlobHandlers(Client);
 
 	UTestReplicatedIrisObject* ServerObject = Server->CreateObject(0, 0);
-	FNetHandle ObjectHandle = ServerObject->NetHandle;
+	FNetRefHandle ObjectHandle = ServerObject->NetRefHandle;
 	UTestReplicatedIrisObject* ServerSubObject = Server->CreateSubObject(ObjectHandle, 0, 0);
-	FNetHandle SubObjectHandle = ServerSubObject->NetHandle;
+	FNetRefHandle SubObjectHandle = ServerSubObject->NetRefHandle;
 
 	Server->PreSendUpdate();
 	Server->SendAndDeliverTo(Client, DeliverPacket);
@@ -407,7 +407,7 @@ UE_NET_TEST_FIXTURE(FTestFlushBeforeDestroyFixture, TestTearOffNewObjectWithReli
 	{
 		constexpr uint32 PayloadBitCount = 24;
 		const TRefCountPtr<FNetObjectAttachment>& Attachment = MockNetObjectAttachmentHandler->CreateReliableNetObjectAttachment(PayloadBitCount);
-		FNetObjectReference AttachmentTarget = FObjectReferenceCache::MakeNetObjectReference(ServerObject->NetHandle);
+		FNetObjectReference AttachmentTarget = FObjectReferenceCache::MakeNetObjectReference(ServerObject->NetRefHandle);
 		Server->GetReplicationSystem()->QueueNetObjectAttachment(Client->ConnectionIdOnServer, AttachmentTarget, Attachment);
 	}
 
@@ -420,7 +420,7 @@ UE_NET_TEST_FIXTURE(FTestFlushBeforeDestroyFixture, TestTearOffNewObjectWithReli
 	Server->PostSendUpdate();
 
 	// Verify that object got created
-	UTestReplicatedIrisObject* ClientObjectThatWillBeTornOff = Cast<UTestReplicatedIrisObject>(Client->GetReplicationBridge()->GetReplicatedObject(ServerObject->NetHandle));
+	UTestReplicatedIrisObject* ClientObjectThatWillBeTornOff = Cast<UTestReplicatedIrisObject>(Client->GetReplicationBridge()->GetReplicatedObject(ServerObject->NetRefHandle));
 
 	UE_NET_ASSERT_TRUE(ClientObjectThatWillBeTornOff != nullptr);
 
@@ -434,7 +434,7 @@ UE_NET_TEST_FIXTURE(FTestFlushBeforeDestroyFixture, TestTearOffNewObjectWithReli
 	Server->PostSendUpdate();
 
 	// Verify that ClientObject now has been teared off
-	UE_NET_ASSERT_TRUE(Client->GetReplicationBridge()->GetReplicatedObject(ServerObject->NetHandle) == nullptr);
+	UE_NET_ASSERT_TRUE(Client->GetReplicationBridge()->GetReplicatedObject(ServerObject->NetRefHandle) == nullptr);
 }
 
 // Test TearOff for existing confirmed object
@@ -459,12 +459,12 @@ UE_NET_TEST_FIXTURE(FTestFlushBeforeDestroyFixture, TestTearOffExistingObjectWit
 	{
 		constexpr uint32 PayloadBitCount = 24;
 		const TRefCountPtr<FNetObjectAttachment>& Attachment = MockNetObjectAttachmentHandler->CreateReliableNetObjectAttachment(PayloadBitCount);
-		FNetObjectReference AttachmentTarget = FObjectReferenceCache::MakeNetObjectReference(ServerObject->NetHandle);
+		FNetObjectReference AttachmentTarget = FObjectReferenceCache::MakeNetObjectReference(ServerObject->NetRefHandle);
 		Server->GetReplicationSystem()->QueueNetObjectAttachment(Client->ConnectionIdOnServer, AttachmentTarget, Attachment);
 	}
 
 	// Store Pointer to object 
-	UTestReplicatedIrisObject* ClientObjectThatWillBeTornOff = Cast<UTestReplicatedIrisObject>(Client->GetReplicationBridge()->GetReplicatedObject(ServerObject->NetHandle));
+	UTestReplicatedIrisObject* ClientObjectThatWillBeTornOff = Cast<UTestReplicatedIrisObject>(Client->GetReplicationBridge()->GetReplicatedObject(ServerObject->NetRefHandle));
 
 	UE_NET_ASSERT_TRUE(ClientObjectThatWillBeTornOff != nullptr);
 	UE_NET_ASSERT_EQ(ServerObject->IntA, ClientObjectThatWillBeTornOff->IntA);
@@ -485,7 +485,7 @@ UE_NET_TEST_FIXTURE(FTestFlushBeforeDestroyFixture, TestTearOffExistingObjectWit
 	UE_NET_ASSERT_EQ(ClientMockNetObjectAttachmentHandler->GetFunctionCallCounts().OnNetBlobReceived, 1U);
 
 	// Verify that ClientObject still is around (from a network perspective)
-	UE_NET_ASSERT_TRUE(Cast<UTestReplicatedIrisObject>(Client->GetReplicationBridge()->GetReplicatedObject(ServerObject->NetHandle)) != nullptr);
+	UE_NET_ASSERT_TRUE(Cast<UTestReplicatedIrisObject>(Client->GetReplicationBridge()->GetReplicatedObject(ServerObject->NetRefHandle)) != nullptr);
 
 	// Send and deliver packet
 	Server->PreSendUpdate();
@@ -493,7 +493,7 @@ UE_NET_TEST_FIXTURE(FTestFlushBeforeDestroyFixture, TestTearOffExistingObjectWit
 	Server->PostSendUpdate();
 
 	// Verify that ClientObject now has been teared off
-	UE_NET_ASSERT_TRUE(Cast<UTestReplicatedIrisObject>(Client->GetReplicationBridge()->GetReplicatedObject(ServerObject->NetHandle)) == nullptr);
+	UE_NET_ASSERT_TRUE(Cast<UTestReplicatedIrisObject>(Client->GetReplicationBridge()->GetReplicatedObject(ServerObject->NetRefHandle)) == nullptr);
 }
 
 // Test TearOff and SubObjects, SubObjects must apply state?
@@ -509,7 +509,7 @@ UE_NET_TEST_FIXTURE(FTestFlushBeforeDestroyFixture, TestImmediateTearOffExisting
 	UTestReplicatedIrisObject* ServerObject = Server->CreateObject(0,0);
 
 	// Spawn second object on server as a subobject
-	UTestReplicatedIrisObject* ServerSubObject = Server->CreateSubObject(ServerObject->NetHandle, 0, 0);
+	UTestReplicatedIrisObject* ServerSubObject = Server->CreateSubObject(ServerObject->NetRefHandle, 0, 0);
 
 	// Trigger replication
 	ServerObject->IntA = 1;
@@ -521,11 +521,11 @@ UE_NET_TEST_FIXTURE(FTestFlushBeforeDestroyFixture, TestImmediateTearOffExisting
 	Server->PostSendUpdate();
 
 	// Store Pointer to objects
-	UTestReplicatedIrisObject* ClientObjectThatWillBeTornOff = Cast<UTestReplicatedIrisObject>(Client->GetReplicationBridge()->GetReplicatedObject(ServerObject->NetHandle));
+	UTestReplicatedIrisObject* ClientObjectThatWillBeTornOff = Cast<UTestReplicatedIrisObject>(Client->GetReplicationBridge()->GetReplicatedObject(ServerObject->NetRefHandle));
 	UE_NET_ASSERT_TRUE(ClientObjectThatWillBeTornOff != nullptr);
 	UE_NET_ASSERT_EQ(ServerObject->IntA, ClientObjectThatWillBeTornOff->IntA);
 
-	UTestReplicatedIrisObject* ClientSubObjectThatWillBeTornOff = Cast<UTestReplicatedIrisObject>(Client->GetReplicationBridge()->GetReplicatedObject(ServerSubObject->NetHandle));
+	UTestReplicatedIrisObject* ClientSubObjectThatWillBeTornOff = Cast<UTestReplicatedIrisObject>(Client->GetReplicationBridge()->GetReplicatedObject(ServerSubObject->NetRefHandle));
 	UE_NET_ASSERT_TRUE(ClientSubObjectThatWillBeTornOff != nullptr);
 	UE_NET_ASSERT_EQ(ServerSubObject->IntA, ClientSubObjectThatWillBeTornOff->IntA);
 
@@ -536,7 +536,7 @@ UE_NET_TEST_FIXTURE(FTestFlushBeforeDestroyFixture, TestImmediateTearOffExisting
 	{
 		constexpr uint32 PayloadBitCount = 24;
 		const TRefCountPtr<FNetObjectAttachment>& Attachment = MockNetObjectAttachmentHandler->CreateReliableNetObjectAttachment(PayloadBitCount);
-		FNetObjectReference AttachmentTarget = FObjectReferenceCache::MakeNetObjectReference(ServerSubObject->NetHandle);
+		FNetObjectReference AttachmentTarget = FObjectReferenceCache::MakeNetObjectReference(ServerSubObject->NetRefHandle);
 		Server->GetReplicationSystem()->QueueNetObjectAttachment(Client->ConnectionIdOnServer, AttachmentTarget, Attachment);
 	}
 
@@ -558,7 +558,7 @@ UE_NET_TEST_FIXTURE(FTestFlushBeforeDestroyFixture, TestImmediateTearOffExisting
 	Server->PostSendUpdate();
 
 	// Verify that ClientObject is torn-off
-	UE_NET_ASSERT_TRUE(Cast<UTestReplicatedIrisObject>(Client->GetReplicationBridge()->GetReplicatedObject(ServerSubObject->NetHandle)) == nullptr);
+	UE_NET_ASSERT_TRUE(Cast<UTestReplicatedIrisObject>(Client->GetReplicationBridge()->GetReplicatedObject(ServerSubObject->NetRefHandle)) == nullptr);
 }
 
 
diff --git a/Engine/Plugins/Runtime/ReplicationSystemTestPlugin/Source/Private/Tests/ReplicationSystem/TestObjectDeltaSerialization.cpp b/Engine/Plugins/Runtime/ReplicationSystemTestPlugin/Source/Private/Tests/ReplicationSystem/TestObjectDeltaSerialization.cpp
--- a/Engine/Plugins/Runtime/ReplicationSystemTestPlugin/Source/Private/Tests/ReplicationSystem/TestObjectDeltaSerialization.cpp
+++ b/Engine/Plugins/Runtime/ReplicationSystemTestPlugin/Source/Private/Tests/ReplicationSystem/TestObjectDeltaSerialization.cpp
@@ -38,14 +38,14 @@ UE_NET_TEST_FIXTURE(FTestObjectDeltaSerialization, ClientReceivesLatestValuesWit
 	UTestReplicatedIrisObject::FComponents ObjectComponents;
 	ObjectComponents.PropertyComponentCount = PropertyComponentCount;
 	UTestReplicatedIrisObject* ServerObject = Server->CreateObject(ObjectComponents);
-	Server->ReplicationSystem->SetDeltaCompressionStatus(ServerObject->NetHandle, ENetObjectDeltaCompressionStatus::Allow);
+	Server->ReplicationSystem->SetDeltaCompressionStatus(ServerObject->NetRefHandle, ENetObjectDeltaCompressionStatus::Allow);
 
 	// Send and deliver packet
 	Server->PreSendUpdate();
 	Server->SendAndDeliverTo(Client, DeliverPacket);
 	Server->PostSendUpdate();
 
-	const UTestReplicatedIrisObject* ClientObject = Cast<UTestReplicatedIrisObject>(Client->GetReplicationBridge()->GetReplicatedObject(ServerObject->NetHandle));
+	const UTestReplicatedIrisObject* ClientObject = Cast<UTestReplicatedIrisObject>(Client->GetReplicationBridge()->GetReplicatedObject(ServerObject->NetRefHandle));
 	UE_NET_ASSERT_NE(ClientObject, nullptr);
 
 	// Set value on object and component
@@ -74,14 +74,14 @@ UE_NET_TEST_FIXTURE(FTestObjectDeltaSerialization, ClientReceivesLatestValuesWit
 	UTestReplicatedIrisObject::FComponents ObjectComponents;
 	ObjectComponents.PropertyComponentCount = PropertyComponentCount;
 	UTestReplicatedIrisObject* ServerObject = Server->CreateObject(ObjectComponents);
-	Server->ReplicationSystem->SetDeltaCompressionStatus(ServerObject->NetHandle, ENetObjectDeltaCompressionStatus::Allow);
+	Server->ReplicationSystem->SetDeltaCompressionStatus(ServerObject->NetRefHandle, ENetObjectDeltaCompressionStatus::Allow);
 
 	// Send and deliver packet
 	Server->PreSendUpdate();
 	Server->SendAndDeliverTo(Client, DeliverPacket);
 	Server->PostSendUpdate();
 
-	const UTestReplicatedIrisObject* ClientObject = Cast<UTestReplicatedIrisObject>(Client->GetReplicationBridge()->GetReplicatedObject(ServerObject->NetHandle));
+	const UTestReplicatedIrisObject* ClientObject = Cast<UTestReplicatedIrisObject>(Client->GetReplicationBridge()->GetReplicatedObject(ServerObject->NetRefHandle));
 	UE_NET_ASSERT_NE(ClientObject, nullptr);
 
 	// Set value on object and component
@@ -124,7 +124,7 @@ UE_NET_TEST_FIXTURE(FTestObjectDeltaSerialization, AllClientsReceivesLatestValue
 	UTestReplicatedIrisObject::FComponents ObjectComponents;
 	ObjectComponents.PropertyComponentCount = PropertyComponentCount;
 	UTestReplicatedIrisObject* ServerObject = Server->CreateObject(ObjectComponents);
-	Server->ReplicationSystem->SetDeltaCompressionStatus(ServerObject->NetHandle, ENetObjectDeltaCompressionStatus::Allow);
+	Server->ReplicationSystem->SetDeltaCompressionStatus(ServerObject->NetRefHandle, ENetObjectDeltaCompressionStatus::Allow);
 
 	// Send and deliver packets
 	Server->PreSendUpdate();
@@ -151,7 +151,7 @@ UE_NET_TEST_FIXTURE(FTestObjectDeltaSerialization, AllClientsReceivesLatestValue
 	// Check that the server modified members have their final values on all clients.
 	for (FReplicationSystemTestClient* Client : Clients)
 	{
-		const UTestReplicatedIrisObject* ClientObject = Cast<UTestReplicatedIrisObject>(Client->GetReplicationBridge()->GetReplicatedObject(ServerObject->NetHandle));
+		const UTestReplicatedIrisObject* ClientObject = Cast<UTestReplicatedIrisObject>(Client->GetReplicationBridge()->GetReplicatedObject(ServerObject->NetRefHandle));
 		UE_NET_ASSERT_NE(ClientObject, nullptr);
 
 		UE_NET_ASSERT_EQ(ClientObject->IntA, ServerObject->IntA);
@@ -173,10 +173,10 @@ UE_NET_TEST_FIXTURE(FTestObjectDeltaSerialization, SimulatedIsProperlyReplicated
 	UTestReplicatedIrisObject::FComponents ObjectComponents;
 	ObjectComponents.ConnectionFilteredComponentCount = ConnectionFilteredComponentCount;
 	UTestReplicatedIrisObject* ServerObject = Server->CreateObject(ObjectComponents);
-	Server->ReplicationSystem->SetDeltaCompressionStatus(ServerObject->NetHandle, ENetObjectDeltaCompressionStatus::Allow);
+	Server->ReplicationSystem->SetDeltaCompressionStatus(ServerObject->NetRefHandle, ENetObjectDeltaCompressionStatus::Allow);
 
 	// Set the client to be autonomous.
-	Server->ReplicationSystem->SetReplicationConditionConnectionFilter(ServerObject->NetHandle, EReplicationCondition::RoleAutonomous, Clients[0]->ConnectionIdOnServer, true);
+	Server->ReplicationSystem->SetReplicationConditionConnectionFilter(ServerObject->NetRefHandle, EReplicationCondition::RoleAutonomous, Clients[0]->ConnectionIdOnServer, true);
 
 	// Send and deliver packet
 	Server->PreSendUpdate();
@@ -186,7 +186,7 @@ UE_NET_TEST_FIXTURE(FTestObjectDeltaSerialization, SimulatedIsProperlyReplicated
 	}
 	Server->PostSendUpdate();
 
-	const UTestReplicatedIrisObject* ClientObject = Cast<UTestReplicatedIrisObject>(Clients[0]->GetReplicationBridge()->GetReplicatedObject(ServerObject->NetHandle));
+	const UTestReplicatedIrisObject* ClientObject = Cast<UTestReplicatedIrisObject>(Clients[0]->GetReplicationBridge()->GetReplicatedObject(ServerObject->NetRefHandle));
 	UE_NET_ASSERT_NE(ClientObject, nullptr);
 
 	// Set some values that should be replicated for simulated objects
@@ -210,7 +210,7 @@ UE_NET_TEST_FIXTURE(FTestObjectDeltaSerialization, SimulatedIsProperlyReplicated
 	UE_NET_ASSERT_EQ(ClientObject->ConnectionFilteredComponents[0]->SimulatedOrPhysicsNoReplayInt, 0);
 	
 	// Set the client to no longer be autonomous, meaning it should be "simulated".
-	Server->ReplicationSystem->SetReplicationConditionConnectionFilter(ServerObject->NetHandle, EReplicationCondition::RoleAutonomous, Clients[0]->ConnectionIdOnServer, false);
+	Server->ReplicationSystem->SetReplicationConditionConnectionFilter(ServerObject->NetRefHandle, EReplicationCondition::RoleAutonomous, Clients[0]->ConnectionIdOnServer, false);
 
 	// Change arbitrary property to trigger replication
 	ServerObject->IntA ^= 1;
@@ -243,17 +243,17 @@ UE_NET_TEST_FIXTURE(FTestObjectDeltaSerialization, SimulatedIsProperlyReplicated
 	UTestReplicatedIrisObject::FComponents ObjectComponents;
 	ObjectComponents.ConnectionFilteredComponentCount = ConnectionFilteredComponentCount;
 	UTestReplicatedIrisObject* ServerObject = Server->CreateObject(ObjectComponents);
-	Server->ReplicationSystem->SetDeltaCompressionStatus(ServerObject->NetHandle, ENetObjectDeltaCompressionStatus::Allow);
+	Server->ReplicationSystem->SetDeltaCompressionStatus(ServerObject->NetRefHandle, ENetObjectDeltaCompressionStatus::Allow);
 
 	// Set the client to be autonomous.
-	Server->ReplicationSystem->SetReplicationConditionConnectionFilter(ServerObject->NetHandle, EReplicationCondition::RoleAutonomous, Client->ConnectionIdOnServer, true);
+	Server->ReplicationSystem->SetReplicationConditionConnectionFilter(ServerObject->NetRefHandle, EReplicationCondition::RoleAutonomous, Client->ConnectionIdOnServer, true);
 
 	// Send and deliver packet
 	Server->PreSendUpdate();
 	Server->SendAndDeliverTo(Client, DeliverPacket);
 	Server->PostSendUpdate();
 
-	const UTestReplicatedIrisObject* ClientObject = Cast<UTestReplicatedIrisObject>(Client->GetReplicationBridge()->GetReplicatedObject(ServerObject->NetHandle));
+	const UTestReplicatedIrisObject* ClientObject = Cast<UTestReplicatedIrisObject>(Client->GetReplicationBridge()->GetReplicatedObject(ServerObject->NetRefHandle));
 	UE_NET_ASSERT_NE(ClientObject, nullptr);
 
 	// Set some values that should be replicated for simulated objects
@@ -274,7 +274,7 @@ UE_NET_TEST_FIXTURE(FTestObjectDeltaSerialization, SimulatedIsProperlyReplicated
 	UE_NET_ASSERT_EQ(ClientObject->ConnectionFilteredComponents[0]->SimulatedOrPhysicsNoReplayInt, 0);
 	
 	// Set the client to no longer be autonomous, meaning it should be "simulated".
-	Server->ReplicationSystem->SetReplicationConditionConnectionFilter(ServerObject->NetHandle, EReplicationCondition::RoleAutonomous, Client->ConnectionIdOnServer, false);
+	Server->ReplicationSystem->SetReplicationConditionConnectionFilter(ServerObject->NetRefHandle, EReplicationCondition::RoleAutonomous, Client->ConnectionIdOnServer, false);
 
 	// Change arbitrary property to trigger replication
 	ServerObject->IntA ^= 1;
@@ -314,7 +314,7 @@ UE_NET_TEST_FIXTURE(FTestObjectDeltaSerialization, ToOwnerStateIsReplicatedToOwn
 	UTestReplicatedIrisObject::FComponents ObjectComponents;
 	ObjectComponents.ConnectionFilteredComponentCount = ConnectionFilteredComponentCount;
 	UTestReplicatedIrisObject* ServerObject = Server->CreateObject(ObjectComponents);
-	Server->ReplicationSystem->SetDeltaCompressionStatus(ServerObject->NetHandle, ENetObjectDeltaCompressionStatus::Allow);
+	Server->ReplicationSystem->SetDeltaCompressionStatus(ServerObject->NetRefHandle, ENetObjectDeltaCompressionStatus::Allow);
 
 	// Set some values in ToOwner only state
 	ServerObject->ConnectionFilteredComponents[0]->ToOwnerA = 13;
@@ -325,7 +325,7 @@ UE_NET_TEST_FIXTURE(FTestObjectDeltaSerialization, ToOwnerStateIsReplicatedToOwn
 	Server->SendAndDeliverTo(Client, DeliverPacket);
 	Server->PostSendUpdate();
 	
-	const UTestReplicatedIrisObject* ClientObject = Cast<UTestReplicatedIrisObject>(Client->GetReplicationBridge()->GetReplicatedObject(ServerObject->NetHandle));
+	const UTestReplicatedIrisObject* ClientObject = Cast<UTestReplicatedIrisObject>(Client->GetReplicationBridge()->GetReplicatedObject(ServerObject->NetRefHandle));
 	UE_NET_ASSERT_NE(ClientObject, nullptr);
 
 	// Check that the server modified ToOwner members are still in the default state
@@ -333,7 +333,7 @@ UE_NET_TEST_FIXTURE(FTestObjectDeltaSerialization, ToOwnerStateIsReplicatedToOwn
 	UE_NET_ASSERT_EQ(ClientObject->ConnectionFilteredComponents[0]->ToOwnerB, 0);
 
 	// Set owner
-	Server->ReplicationSystem->SetOwningNetConnection(ServerObject->NetHandle, Client->ConnectionIdOnServer);
+	Server->ReplicationSystem->SetOwningNetConnection(ServerObject->NetRefHandle, Client->ConnectionIdOnServer);
 
 	// Change arbitrary property to trigger replication
 	ServerObject->IntA ^= 1;
@@ -357,17 +357,17 @@ UE_NET_TEST_FIXTURE(FTestObjectDeltaSerialization, InFlightChangesForDisabledCon
 	UTestReplicatedIrisObject::FComponents ObjectComponents;
 	ObjectComponents.ConnectionFilteredComponentCount = ConnectionFilteredComponentCount;
 	UTestReplicatedIrisObject* ServerObject = Server->CreateObject(ObjectComponents);
-	Server->ReplicationSystem->SetDeltaCompressionStatus(ServerObject->NetHandle, ENetObjectDeltaCompressionStatus::Allow);
+	Server->ReplicationSystem->SetDeltaCompressionStatus(ServerObject->NetRefHandle, ENetObjectDeltaCompressionStatus::Allow);
 
 	// Set autonomous role for our client. This is due to there not being a COND_Physics condition, but there is a COND_SimulatedOrPhysics.
 	{
-		Server->ReplicationSystem->SetReplicationConditionConnectionFilter(ServerObject->NetHandle, EReplicationCondition::RoleAutonomous, Client->ConnectionIdOnServer, true);
+		Server->ReplicationSystem->SetReplicationConditionConnectionFilter(ServerObject->NetRefHandle, EReplicationCondition::RoleAutonomous, Client->ConnectionIdOnServer, true);
 	}
 
 	// Enable replicating physics
 	{
 		constexpr bool bIsPhysicsReplicationAllowed = true;
-		Server->ReplicationSystem->SetReplicationCondition(ServerObject->NetHandle, EReplicationCondition::ReplicatePhysics, bIsPhysicsReplicationAllowed);
+		Server->ReplicationSystem->SetReplicationCondition(ServerObject->NetRefHandle, EReplicationCondition::ReplicatePhysics, bIsPhysicsReplicationAllowed);
 	}
 
 	// Set value for SimulatedOrPhysics condition.
@@ -379,7 +379,7 @@ UE_NET_TEST_FIXTURE(FTestObjectDeltaSerialization, InFlightChangesForDisabledCon
 	Server->SendAndDeliverTo(Client, DeliverPacket);
 	Server->PostSendUpdate();
 	
-	const UTestReplicatedIrisObject* ClientObject = Cast<UTestReplicatedIrisObject>(Client->GetReplicationBridge()->GetReplicatedObject(ServerObject->NetHandle));
+	const UTestReplicatedIrisObject* ClientObject = Cast<UTestReplicatedIrisObject>(Client->GetReplicationBridge()->GetReplicatedObject(ServerObject->NetRefHandle));
 	UE_NET_ASSERT_NE(ClientObject, nullptr);
 
 	// Check that the server modified SimulatedOrPhysics has been replicated properly.
@@ -401,7 +401,7 @@ UE_NET_TEST_FIXTURE(FTestObjectDeltaSerialization, InFlightChangesForDisabledCon
 	// Disable replicating physics
 	{
 		constexpr bool bIsPhysicsReplicationAllowed = false;
-		Server->ReplicationSystem->SetReplicationCondition(ServerObject->NetHandle, EReplicationCondition::ReplicatePhysics, bIsPhysicsReplicationAllowed);
+		Server->ReplicationSystem->SetReplicationCondition(ServerObject->NetRefHandle, EReplicationCondition::ReplicatePhysics, bIsPhysicsReplicationAllowed);
 	}
 
 	// Allow a new baseline to be created
diff --git a/Engine/Plugins/Runtime/ReplicationSystemTestPlugin/Source/Private/Tests/ReplicationSystem/TestObjectReferences.cpp b/Engine/Plugins/Runtime/ReplicationSystemTestPlugin/Source/Private/Tests/ReplicationSystem/TestObjectReferences.cpp
--- a/Engine/Plugins/Runtime/ReplicationSystemTestPlugin/Source/Private/Tests/ReplicationSystem/TestObjectReferences.cpp
+++ b/Engine/Plugins/Runtime/ReplicationSystemTestPlugin/Source/Private/Tests/ReplicationSystem/TestObjectReferences.cpp
@@ -150,9 +150,9 @@ struct FTestObjectReferencesFixture : public FReplicationSystemServerClientTestF
 		// Create test object
 		TestObject = Server->CreateObject<UTestObjectReferences_TestClassWithReferences>();
 		
-		uint32 ObjectIndex = Server->GetReplicationSystem()->GetReplicationSystemInternal()->GetNetHandleManager().GetInternalIndex(TestObject->NetHandle);
-		TestObjectStateBuffer = Server->GetReplicationSystem()->GetReplicationSystemInternal()->GetNetHandleManager().GetReplicatedObjectStateBufferNoCheck(ObjectIndex);
-		TestObjectProtocol = Server->GetReplicationSystem()->GetReplicationProtocol(TestObject->NetHandle);
+		uint32 ObjectIndex = Server->GetReplicationSystem()->GetReplicationSystemInternal()->GetNetRefHandleManager().GetInternalIndex(TestObject->NetRefHandle);
+		TestObjectStateBuffer = Server->GetReplicationSystem()->GetReplicationSystemInternal()->GetNetRefHandleManager().GetReplicatedObjectStateBufferNoCheck(ObjectIndex);
+		TestObjectProtocol = Server->GetReplicationSystem()->GetReplicationProtocol(TestObject->NetRefHandle);
 
 		// InitCollectors
 		CollectOnlyValidCollector.bCollectAll = false;
@@ -160,6 +160,7 @@ struct FTestObjectReferencesFixture : public FReplicationSystemServerClientTestF
 
 		// To trigger copy of data so we have valid data in buffers
 		Server->PreSendUpdate();
+		Server->PostSendUpdate();
 	}
 };
 
@@ -167,6 +168,7 @@ UE_NET_TEST_FIXTURE(FTestObjectReferencesFixture, TestReferenceVisitor_NoRef)
 {
 	// To trigger copy of data
 	Server->PreSendUpdate();
+	Server->PostSendUpdate();
 
 	FNetSerializationContext Context;
 
@@ -190,6 +192,7 @@ UE_NET_TEST_FIXTURE(FTestObjectReferencesFixture, TestReferenceVisitor_SingleRef
 
 	// To trigger copy of data
 	Server->PreSendUpdate();
+	Server->PostSendUpdate();
 	
 	FReferenceCollector& Collector = CollectOnlyValidCollector;
 
@@ -197,7 +200,7 @@ UE_NET_TEST_FIXTURE(FTestObjectReferencesFixture, TestReferenceVisitor_SingleRef
 	VisitReferences(TestObjectStateBuffer, TestObjectProtocol, Collector);
 
 	UE_NET_ASSERT_EQ(Collector.References.Num(), 1);
-	UE_NET_ASSERT_TRUE(Collector.References[0].GetRefHandle() == TestReferences[0]->NetHandle);
+	UE_NET_ASSERT_TRUE(Collector.References[0].GetRefHandle() == TestReferences[0]->NetRefHandle);
 }
 
 UE_NET_TEST_FIXTURE(FTestObjectReferencesFixture, TestReferenceVisitor_ObjectRef)
@@ -207,14 +210,15 @@ UE_NET_TEST_FIXTURE(FTestObjectReferencesFixture, TestReferenceVisitor_ObjectRef
 
 	// To trigger copy of data
 	Server->PreSendUpdate();
-	
+	Server->PostSendUpdate();
+
 	FReferenceCollector& Collector = CollectOnlyValidCollector;
 
 	// Collect 
 	VisitReferences(TestObjectStateBuffer, TestObjectProtocol, Collector);
 
 	UE_NET_ASSERT_EQ(Collector.References.Num(), 1);
-	UE_NET_ASSERT_TRUE(Collector.References[0].GetRefHandle() == TestReferences[0]->NetHandle);
+	UE_NET_ASSERT_TRUE(Collector.References[0].GetRefHandle() == TestReferences[0]->NetRefHandle);
 }
 
 UE_NET_TEST_FIXTURE(FTestObjectReferencesFixture, TestReferenceVisitor_StructWithRef)
@@ -224,6 +228,7 @@ UE_NET_TEST_FIXTURE(FTestObjectReferencesFixture, TestReferenceVisitor_StructWit
 
 	// To trigger copy of data
 	Server->PreSendUpdate();
+	Server->PostSendUpdate();
 	
 	FReferenceCollector& Collector = CollectOnlyValidCollector;
 
@@ -231,7 +236,7 @@ UE_NET_TEST_FIXTURE(FTestObjectReferencesFixture, TestReferenceVisitor_StructWit
 	VisitReferences(TestObjectStateBuffer, TestObjectProtocol, Collector);
 
 	UE_NET_ASSERT_EQ(Collector.References.Num(), 1);
-	UE_NET_ASSERT_TRUE(Collector.References[0].GetRefHandle() == TestReferences[0]->NetHandle);
+	UE_NET_ASSERT_TRUE(Collector.References[0].GetRefHandle() == TestReferences[0]->NetRefHandle);
 }
 
 UE_NET_TEST_FIXTURE(FTestObjectReferencesFixture, TestReferenceVisitor_StructWithRefCArray)
@@ -241,6 +246,7 @@ UE_NET_TEST_FIXTURE(FTestObjectReferencesFixture, TestReferenceVisitor_StructWit
 
 	// To trigger copy of data
 	Server->PreSendUpdate();
+	Server->PostSendUpdate();
 	
 	FReferenceCollector& Collector = CollectOnlyValidCollector;
 
@@ -248,7 +254,7 @@ UE_NET_TEST_FIXTURE(FTestObjectReferencesFixture, TestReferenceVisitor_StructWit
 	VisitReferences(TestObjectStateBuffer, TestObjectProtocol, Collector);
 
 	UE_NET_ASSERT_EQ(Collector.References.Num(), 1);
-	UE_NET_ASSERT_TRUE(Collector.References[0].GetRefHandle() == TestReferences[0]->NetHandle);
+	UE_NET_ASSERT_TRUE(Collector.References[0].GetRefHandle() == TestReferences[0]->NetRefHandle);
 }
 
 UE_NET_TEST_FIXTURE(FTestObjectReferencesFixture, TestReferenceVisitor_StructWithRefCArrayMulti)
@@ -260,6 +266,7 @@ UE_NET_TEST_FIXTURE(FTestObjectReferencesFixture, TestReferenceVisitor_StructWit
 
 	// To trigger copy of data
 	Server->PreSendUpdate();
+	Server->PostSendUpdate();
 	
 	FReferenceCollector& Collector = CollectOnlyValidCollector;
 
@@ -267,9 +274,9 @@ UE_NET_TEST_FIXTURE(FTestObjectReferencesFixture, TestReferenceVisitor_StructWit
 	VisitReferences(TestObjectStateBuffer, TestObjectProtocol, Collector);
 
 	UE_NET_ASSERT_EQ(Collector.References.Num(), 3);
-	UE_NET_ASSERT_TRUE(Collector.References[0].GetRefHandle() == TestReferences[0]->NetHandle);
-	UE_NET_ASSERT_TRUE(Collector.References[1].GetRefHandle() == TestReferences[1]->NetHandle);
-	UE_NET_ASSERT_TRUE(Collector.References[2].GetRefHandle() == TestReferences[2]->NetHandle);
+	UE_NET_ASSERT_TRUE(Collector.References[0].GetRefHandle() == TestReferences[0]->NetRefHandle);
+	UE_NET_ASSERT_TRUE(Collector.References[1].GetRefHandle() == TestReferences[1]->NetRefHandle);
+	UE_NET_ASSERT_TRUE(Collector.References[2].GetRefHandle() == TestReferences[2]->NetRefHandle);
 }
 
 UE_NET_TEST_FIXTURE(FTestObjectReferencesFixture, TestReferenceVisitor_StructWithRefTArray)
@@ -279,6 +286,7 @@ UE_NET_TEST_FIXTURE(FTestObjectReferencesFixture, TestReferenceVisitor_StructWit
 
 	// To trigger copy of data
 	Server->PreSendUpdate();
+	Server->PostSendUpdate();
 	
 	FReferenceCollector& Collector = CollectOnlyValidCollector;
 
@@ -286,7 +294,7 @@ UE_NET_TEST_FIXTURE(FTestObjectReferencesFixture, TestReferenceVisitor_StructWit
 	VisitReferences(TestObjectStateBuffer, TestObjectProtocol, Collector);
 
 	UE_NET_ASSERT_EQ(Collector.References.Num(), 1);
-	UE_NET_ASSERT_TRUE(Collector.References[0].GetRefHandle() == TestReferences[0]->NetHandle);
+	UE_NET_ASSERT_TRUE(Collector.References[0].GetRefHandle() == TestReferences[0]->NetRefHandle);
 }
 
 UE_NET_TEST_FIXTURE(FTestObjectReferencesFixture, TestReferenceVisitor_StructWithRefTArrayMulti)
@@ -298,6 +306,7 @@ UE_NET_TEST_FIXTURE(FTestObjectReferencesFixture, TestReferenceVisitor_StructWit
 
 	// To trigger copy of data
 	Server->PreSendUpdate();
+	Server->PostSendUpdate();
 	
 	FReferenceCollector& Collector = CollectOnlyValidCollector;
 
@@ -305,9 +314,9 @@ UE_NET_TEST_FIXTURE(FTestObjectReferencesFixture, TestReferenceVisitor_StructWit
 	VisitReferences(TestObjectStateBuffer, TestObjectProtocol, Collector);
 
 	UE_NET_ASSERT_EQ(Collector.References.Num(), 3);
-	UE_NET_ASSERT_TRUE(Collector.References[0].GetRefHandle() == TestReferences[0]->NetHandle);
-	UE_NET_ASSERT_TRUE(Collector.References[1].GetRefHandle() == TestReferences[1]->NetHandle);
-	UE_NET_ASSERT_TRUE(Collector.References[2].GetRefHandle() == TestReferences[2]->NetHandle);
+	UE_NET_ASSERT_TRUE(Collector.References[0].GetRefHandle() == TestReferences[0]->NetRefHandle);
+	UE_NET_ASSERT_TRUE(Collector.References[1].GetRefHandle() == TestReferences[1]->NetRefHandle);
+	UE_NET_ASSERT_TRUE(Collector.References[2].GetRefHandle() == TestReferences[2]->NetRefHandle);
 }
 
 UE_NET_TEST_FIXTURE(FTestObjectReferencesFixture, TestReferenceVisitor_RefCArray)
@@ -317,6 +326,7 @@ UE_NET_TEST_FIXTURE(FTestObjectReferencesFixture, TestReferenceVisitor_RefCArray
 
 	// To trigger copy of data
 	Server->PreSendUpdate();
+	Server->PostSendUpdate();
 	
 	FReferenceCollector& Collector = CollectOnlyValidCollector;
 
@@ -324,7 +334,7 @@ UE_NET_TEST_FIXTURE(FTestObjectReferencesFixture, TestReferenceVisitor_RefCArray
 	VisitReferences(TestObjectStateBuffer, TestObjectProtocol, Collector);
 
 	UE_NET_ASSERT_EQ(Collector.References.Num(), 1);
-	UE_NET_ASSERT_TRUE(Collector.References[0].GetRefHandle() == TestReferences[0]->NetHandle);
+	UE_NET_ASSERT_TRUE(Collector.References[0].GetRefHandle() == TestReferences[0]->NetRefHandle);
 }
 
 UE_NET_TEST_FIXTURE(FTestObjectReferencesFixture, TestReferenceVisitor_RefCArrayMulti)
@@ -336,6 +346,7 @@ UE_NET_TEST_FIXTURE(FTestObjectReferencesFixture, TestReferenceVisitor_RefCArray
 
 	// To trigger copy of data
 	Server->PreSendUpdate();
+	Server->PostSendUpdate();
 	
 	FReferenceCollector& Collector = CollectOnlyValidCollector;
 
@@ -343,9 +354,9 @@ UE_NET_TEST_FIXTURE(FTestObjectReferencesFixture, TestReferenceVisitor_RefCArray
 	VisitReferences(TestObjectStateBuffer, TestObjectProtocol, Collector);
 
 	UE_NET_ASSERT_EQ(Collector.References.Num(), 3);
-	UE_NET_ASSERT_TRUE(Collector.References[0].GetRefHandle() == TestReferences[0]->NetHandle);
-	UE_NET_ASSERT_TRUE(Collector.References[1].GetRefHandle() == TestReferences[1]->NetHandle);
-	UE_NET_ASSERT_TRUE(Collector.References[2].GetRefHandle() == TestReferences[2]->NetHandle);
+	UE_NET_ASSERT_TRUE(Collector.References[0].GetRefHandle() == TestReferences[0]->NetRefHandle);
+	UE_NET_ASSERT_TRUE(Collector.References[1].GetRefHandle() == TestReferences[1]->NetRefHandle);
+	UE_NET_ASSERT_TRUE(Collector.References[2].GetRefHandle() == TestReferences[2]->NetRefHandle);
 }
 
 UE_NET_TEST_FIXTURE(FTestObjectReferencesFixture, TestReferenceVisitor_RefTArray)
@@ -355,6 +366,7 @@ UE_NET_TEST_FIXTURE(FTestObjectReferencesFixture, TestReferenceVisitor_RefTArray
 
 	// To trigger copy of data
 	Server->PreSendUpdate();
+	Server->PostSendUpdate();
 	
 	FReferenceCollector& Collector = CollectOnlyValidCollector;
 
@@ -362,7 +374,7 @@ UE_NET_TEST_FIXTURE(FTestObjectReferencesFixture, TestReferenceVisitor_RefTArray
 	VisitReferences(TestObjectStateBuffer, TestObjectProtocol, Collector);
 
 	UE_NET_ASSERT_EQ(Collector.References.Num(), 1);
-	UE_NET_ASSERT_TRUE(Collector.References[0].GetRefHandle() == TestReferences[0]->NetHandle);
+	UE_NET_ASSERT_TRUE(Collector.References[0].GetRefHandle() == TestReferences[0]->NetRefHandle);
 }
 
 UE_NET_TEST_FIXTURE(FTestObjectReferencesFixture, TestReferenceVisitor_RefTArrayMulti)
@@ -374,6 +386,7 @@ UE_NET_TEST_FIXTURE(FTestObjectReferencesFixture, TestReferenceVisitor_RefTArray
 
 	// To trigger copy of data
 	Server->PreSendUpdate();
+	Server->PostSendUpdate();
 	
 	FReferenceCollector& Collector = CollectOnlyValidCollector;
 
@@ -381,9 +394,9 @@ UE_NET_TEST_FIXTURE(FTestObjectReferencesFixture, TestReferenceVisitor_RefTArray
 	VisitReferences(TestObjectStateBuffer, TestObjectProtocol, Collector);
 
 	UE_NET_ASSERT_EQ(Collector.References.Num(), 3);
-	UE_NET_ASSERT_TRUE(Collector.References[0].GetRefHandle() == TestReferences[0]->NetHandle);
-	UE_NET_ASSERT_TRUE(Collector.References[1].GetRefHandle() == TestReferences[1]->NetHandle);
-	UE_NET_ASSERT_TRUE(Collector.References[2].GetRefHandle() == TestReferences[2]->NetHandle);
+	UE_NET_ASSERT_TRUE(Collector.References[0].GetRefHandle() == TestReferences[0]->NetRefHandle);
+	UE_NET_ASSERT_TRUE(Collector.References[1].GetRefHandle() == TestReferences[1]->NetRefHandle);
+	UE_NET_ASSERT_TRUE(Collector.References[2].GetRefHandle() == TestReferences[2]->NetRefHandle);
 }
 
 UE_NET_TEST_FIXTURE(FTestObjectReferencesFixture, TestReferenceVisitor_StructWithRef_CArray)
@@ -393,6 +406,7 @@ UE_NET_TEST_FIXTURE(FTestObjectReferencesFixture, TestReferenceVisitor_StructWit
 
 	// To trigger copy of data
 	Server->PreSendUpdate();
+	Server->PostSendUpdate();
 	
 	FReferenceCollector& Collector = CollectOnlyValidCollector;
 
@@ -400,7 +414,7 @@ UE_NET_TEST_FIXTURE(FTestObjectReferencesFixture, TestReferenceVisitor_StructWit
 	VisitReferences(TestObjectStateBuffer, TestObjectProtocol, Collector);
 
 	UE_NET_ASSERT_EQ(Collector.References.Num(), 1);
-	UE_NET_ASSERT_TRUE(Collector.References[0].GetRefHandle() == TestReferences[0]->NetHandle);
+	UE_NET_ASSERT_TRUE(Collector.References[0].GetRefHandle() == TestReferences[0]->NetRefHandle);
 }
 
 UE_NET_TEST_FIXTURE(FTestObjectReferencesFixture, TestReferenceVisitor_StructWithRef_CArrayMulti)
@@ -412,16 +426,17 @@ UE_NET_TEST_FIXTURE(FTestObjectReferencesFixture, TestReferenceVisitor_StructWit
 
 	// To trigger copy of data
 	Server->PreSendUpdate();
-	
+	Server->PostSendUpdate();
+
 	FReferenceCollector& Collector = CollectOnlyValidCollector;
 
 	// Collect 
 	VisitReferences(TestObjectStateBuffer, TestObjectProtocol, Collector);
 
 	UE_NET_ASSERT_EQ(Collector.References.Num(), 3);
-	UE_NET_ASSERT_TRUE(Collector.References[0].GetRefHandle() == TestReferences[0]->NetHandle);
-	UE_NET_ASSERT_TRUE(Collector.References[1].GetRefHandle() == TestReferences[1]->NetHandle);
-	UE_NET_ASSERT_TRUE(Collector.References[2].GetRefHandle() == TestReferences[2]->NetHandle);
+	UE_NET_ASSERT_TRUE(Collector.References[0].GetRefHandle() == TestReferences[0]->NetRefHandle);
+	UE_NET_ASSERT_TRUE(Collector.References[1].GetRefHandle() == TestReferences[1]->NetRefHandle);
+	UE_NET_ASSERT_TRUE(Collector.References[2].GetRefHandle() == TestReferences[2]->NetRefHandle);
 }
 
 UE_NET_TEST_FIXTURE(FTestObjectReferencesFixture, TestReferenceVisitor_StructWithRef_TArray)
@@ -432,6 +447,7 @@ UE_NET_TEST_FIXTURE(FTestObjectReferencesFixture, TestReferenceVisitor_StructWit
 
 	// To trigger copy of data
 	Server->PreSendUpdate();
+	Server->PostSendUpdate();
 	
 	FReferenceCollector& Collector = CollectOnlyValidCollector;
 
@@ -439,7 +455,7 @@ UE_NET_TEST_FIXTURE(FTestObjectReferencesFixture, TestReferenceVisitor_StructWit
 	VisitReferences(TestObjectStateBuffer, TestObjectProtocol, Collector);
 
 	UE_NET_ASSERT_EQ(Collector.References.Num(), 1);
-	UE_NET_ASSERT_TRUE(Collector.References[0].GetRefHandle() == TestReferences[0]->NetHandle);
+	UE_NET_ASSERT_TRUE(Collector.References[0].GetRefHandle() == TestReferences[0]->NetRefHandle);
 }
 
 UE_NET_TEST_FIXTURE(FTestObjectReferencesFixture, TestReferenceVisitor_StructWithRef_TArrayMulti)
@@ -452,16 +468,17 @@ UE_NET_TEST_FIXTURE(FTestObjectReferencesFixture, TestReferenceVisitor_StructWit
 
 	// To trigger copy of data
 	Server->PreSendUpdate();
-	
+	Server->PostSendUpdate();
+
 	FReferenceCollector& Collector = CollectOnlyValidCollector;
 
 	// Collect 
 	VisitReferences(TestObjectStateBuffer, TestObjectProtocol, Collector);
 
 	UE_NET_ASSERT_EQ(Collector.References.Num(), 3);
-	UE_NET_ASSERT_TRUE(Collector.References[0].GetRefHandle() == TestReferences[0]->NetHandle);
-	UE_NET_ASSERT_TRUE(Collector.References[1].GetRefHandle() == TestReferences[1]->NetHandle);
-	UE_NET_ASSERT_TRUE(Collector.References[2].GetRefHandle() == TestReferences[2]->NetHandle);
+	UE_NET_ASSERT_TRUE(Collector.References[0].GetRefHandle() == TestReferences[0]->NetRefHandle);
+	UE_NET_ASSERT_TRUE(Collector.References[1].GetRefHandle() == TestReferences[1]->NetRefHandle);
+	UE_NET_ASSERT_TRUE(Collector.References[2].GetRefHandle() == TestReferences[2]->NetRefHandle);
 }
 
 UE_NET_TEST_FIXTURE(FTestObjectReferencesFixture, TestStructWithNestedRefCArray_CArray)
@@ -471,6 +488,7 @@ UE_NET_TEST_FIXTURE(FTestObjectReferencesFixture, TestStructWithNestedRefCArray_
 
 	// To trigger copy of data
 	Server->PreSendUpdate();
+	Server->PostSendUpdate();
 	
 	FReferenceCollector& Collector = CollectOnlyValidCollector;
 
@@ -478,7 +496,7 @@ UE_NET_TEST_FIXTURE(FTestObjectReferencesFixture, TestStructWithNestedRefCArray_
 	VisitReferences(TestObjectStateBuffer, TestObjectProtocol, Collector);
 
 	UE_NET_ASSERT_EQ(Collector.References.Num(), 1);
-	UE_NET_ASSERT_TRUE(Collector.References[0].GetRefHandle() == TestReferences[0]->NetHandle);
+	UE_NET_ASSERT_TRUE(Collector.References[0].GetRefHandle() == TestReferences[0]->NetRefHandle);
 }
 
 UE_NET_TEST_FIXTURE(FTestObjectReferencesFixture, TestStructWithNestedRefCArray_CArrayMulti)
@@ -490,6 +508,7 @@ UE_NET_TEST_FIXTURE(FTestObjectReferencesFixture, TestStructWithNestedRefCArray_
 
 	// To trigger copy of data
 	Server->PreSendUpdate();
+	Server->PostSendUpdate();
 	
 	FReferenceCollector& Collector = CollectOnlyValidCollector;
 
@@ -497,9 +516,9 @@ UE_NET_TEST_FIXTURE(FTestObjectReferencesFixture, TestStructWithNestedRefCArray_
 	VisitReferences(TestObjectStateBuffer, TestObjectProtocol, Collector);
 
 	UE_NET_ASSERT_EQ(Collector.References.Num(), 3);
-	UE_NET_ASSERT_TRUE(Collector.References[0].GetRefHandle() == TestReferences[0]->NetHandle);
-	UE_NET_ASSERT_TRUE(Collector.References[1].GetRefHandle() == TestReferences[1]->NetHandle);
-	UE_NET_ASSERT_TRUE(Collector.References[2].GetRefHandle() == TestReferences[2]->NetHandle);
+	UE_NET_ASSERT_TRUE(Collector.References[0].GetRefHandle() == TestReferences[0]->NetRefHandle);
+	UE_NET_ASSERT_TRUE(Collector.References[1].GetRefHandle() == TestReferences[1]->NetRefHandle);
+	UE_NET_ASSERT_TRUE(Collector.References[2].GetRefHandle() == TestReferences[2]->NetRefHandle);
 }
 
 UE_NET_TEST_FIXTURE(FTestObjectReferencesFixture, TestStructWithNestedRefTArray_TArray)
@@ -512,6 +531,7 @@ UE_NET_TEST_FIXTURE(FTestObjectReferencesFixture, TestStructWithNestedRefTArray_
 
 	// To trigger copy of data
 	Server->PreSendUpdate();
+	Server->PostSendUpdate();
 	
 	FReferenceCollector& Collector = CollectOnlyValidCollector;
 
@@ -519,7 +539,7 @@ UE_NET_TEST_FIXTURE(FTestObjectReferencesFixture, TestStructWithNestedRefTArray_
 	VisitReferences(TestObjectStateBuffer, TestObjectProtocol, Collector);
 
 	UE_NET_ASSERT_EQ(Collector.References.Num(), 1);
-	UE_NET_ASSERT_TRUE(Collector.References[0].GetRefHandle() == TestReferences[0]->NetHandle);
+	UE_NET_ASSERT_TRUE(Collector.References[0].GetRefHandle() == TestReferences[0]->NetRefHandle);
 }
 
 UE_NET_TEST_FIXTURE(FTestObjectReferencesFixture, TestStructWithNestedRefTArray_TArrayMulti)
@@ -538,6 +558,7 @@ UE_NET_TEST_FIXTURE(FTestObjectReferencesFixture, TestStructWithNestedRefTArray_
 
 	// To trigger copy of data
 	Server->PreSendUpdate();
+	Server->PostSendUpdate();
 	
 	FReferenceCollector& Collector = CollectOnlyValidCollector;
 
@@ -545,9 +566,9 @@ UE_NET_TEST_FIXTURE(FTestObjectReferencesFixture, TestStructWithNestedRefTArray_
 	VisitReferences(TestObjectStateBuffer, TestObjectProtocol, Collector);
 
 	UE_NET_ASSERT_EQ(Collector.References.Num(), 3);
-	UE_NET_ASSERT_TRUE(Collector.References[0].GetRefHandle() == TestReferences[0]->NetHandle);
-	UE_NET_ASSERT_TRUE(Collector.References[1].GetRefHandle() == TestReferences[1]->NetHandle);
-	UE_NET_ASSERT_TRUE(Collector.References[2].GetRefHandle() == TestReferences[2]->NetHandle);
+	UE_NET_ASSERT_TRUE(Collector.References[0].GetRefHandle() == TestReferences[0]->NetRefHandle);
+	UE_NET_ASSERT_TRUE(Collector.References[1].GetRefHandle() == TestReferences[1]->NetRefHandle);
+	UE_NET_ASSERT_TRUE(Collector.References[2].GetRefHandle() == TestReferences[2]->NetRefHandle);
 }
 
 UE_NET_TEST_FIXTURE(FTestObjectReferencesFixture, TestStructWithNestedRefTArray_Combo)
@@ -601,6 +622,7 @@ UE_NET_TEST_FIXTURE(FTestObjectReferencesFixture, TestStructWithNestedRefTArray_
 
 	// To trigger copy of data
 	Server->PreSendUpdate();
+	Server->PostSendUpdate();
 	
 	FReferenceCollector& Collector = CollectOnlyValidCollector;
 
@@ -610,7 +632,7 @@ UE_NET_TEST_FIXTURE(FTestObjectReferencesFixture, TestStructWithNestedRefTArray_
 	UE_NET_ASSERT_EQ((uint32)Collector.References.Num(), CurrentReferenceIndex);
 	for (int32 It = 0; It < Collector.References.Num(); ++It)
 	{
-		UE_NET_ASSERT_TRUE(Collector.References[It].GetRefHandle() == TestReferences[It]->NetHandle);
+		UE_NET_ASSERT_TRUE(Collector.References[It].GetRefHandle() == TestReferences[It]->NetRefHandle);
 	}	
 }
 
@@ -669,10 +691,9 @@ UE_NET_TEST_FIXTURE(FTestObjectReferencesFixture, TestVisitAllDirty)
 
 	// To trigger copy of data
 	Server->PreSendUpdate();
-	
-	
+	Server->PostSendUpdate();
+		
 	// Collect with empty changemask
-	
 	{
 		FCollectValidReferenceCollector Collector(false, ChangeMask);
 
@@ -686,7 +707,7 @@ UE_NET_TEST_FIXTURE(FTestObjectReferencesFixture, TestVisitAllDirty)
 
 		VisitReferences(TestObjectStateBuffer, TestObjectProtocol, Collector);
 		UE_NET_ASSERT_EQ((uint32)Collector.References.Num(), 1U);
-		UE_NET_ASSERT_TRUE(Collector.References[0].GetRefHandle() == TestReferences[0]->NetHandle);
+		UE_NET_ASSERT_TRUE(Collector.References[0].GetRefHandle() == TestReferences[0]->NetRefHandle);
 	}
 
 	// Partial
@@ -699,11 +720,10 @@ UE_NET_TEST_FIXTURE(FTestObjectReferencesFixture, TestVisitAllDirty)
 		UE_NET_ASSERT_EQ(Collector.References.Num(), 3);
 		for (int32 It = 0; It < Collector.References.Num(); ++It)
 		{
-			UE_NET_ASSERT_TRUE(Collector.References[It].GetRefHandle() == TestReferences[It + 20]->NetHandle);
+			UE_NET_ASSERT_TRUE(Collector.References[It].GetRefHandle() == TestReferences[It + 20]->NetRefHandle);
 		}
 	}
 
-
 	// Mark dirty
 	{
 		FCollectValidReferenceCollector Collector(true, ChangeMask);
@@ -713,7 +733,7 @@ UE_NET_TEST_FIXTURE(FTestObjectReferencesFixture, TestVisitAllDirty)
 		UE_NET_ASSERT_EQ((uint32)Collector.References.Num(), CurrentReferenceIndex);
 		for (int32 It = 0; It < Collector.References.Num(); ++It)
 		{
-			UE_NET_ASSERT_TRUE(Collector.References[It].GetRefHandle() == TestReferences[It]->NetHandle);
+			UE_NET_ASSERT_TRUE(Collector.References[It].GetRefHandle() == TestReferences[It]->NetRefHandle);
 		}
 	}
 
@@ -726,7 +746,7 @@ UE_NET_TEST_FIXTURE(FTestObjectReferencesFixture, TestVisitAllDirty)
 		UE_NET_ASSERT_EQ((uint32)Collector.References.Num(), CurrentReferenceIndex - 1);
 		for (int32 It = 0; It < Collector.References.Num(); ++It)
 		{
-			UE_NET_ASSERT_TRUE(Collector.References[It].GetRefHandle() == TestReferences[It + 1]->NetHandle);
+			UE_NET_ASSERT_TRUE(Collector.References[It].GetRefHandle() == TestReferences[It + 1]->NetRefHandle);
 		}
 	}
 }
@@ -750,7 +770,7 @@ UE_NET_TEST_FIXTURE(FReplicationSystemServerClientTestFixture, TestInlinedSubObj
 	Server->PostSendUpdate();
 
 	// Verify that object has been spawned on client
-	auto ClientObject = Client->GetObjectAs<UTestObjectReferences_TestClassWithDefaultSubObject>(Object->NetHandle);
+	auto ClientObject = Client->GetObjectAs<UTestObjectReferences_TestClassWithDefaultSubObject>(Object->NetRefHandle);
 
 	// Verify that default subobject is created
 	UE_NET_ASSERT_TRUE(ClientObject->CreatedSubObjectRef != nullptr);
@@ -783,7 +803,7 @@ UE_NET_TEST_FIXTURE(FReplicationSystemServerClientTestFixture, TestExternalSubOb
 	UE_NET_ASSERT_TRUE(Object->CreatedSubObjectRef->HasAnyFlags(RF_DefaultSubObject) || Object->IsDefaultSubobject());
 
 	// Create external subobject
-	UTestObjectReferences_TestClassWithDefaultSubObject* ExternalSubObject = Server->CreateSubObject<UTestObjectReferences_TestClassWithDefaultSubObject>(Object->NetHandle);
+	UTestObjectReferences_TestClassWithDefaultSubObject* ExternalSubObject = Server->CreateSubObject<UTestObjectReferences_TestClassWithDefaultSubObject>(Object->NetRefHandle);
 
 	// Replicate object
 	Server->PreSendUpdate();
@@ -791,8 +811,8 @@ UE_NET_TEST_FIXTURE(FReplicationSystemServerClientTestFixture, TestExternalSubOb
 	Server->PostSendUpdate();
 
 	// Verify that object has been spawned on client
-	auto ClientObject = Client->GetObjectAs<UTestObjectReferences_TestClassWithDefaultSubObject>(Object->NetHandle);
-	auto ClientExternalSubObject = Client->GetObjectAs<UTestObjectReferences_TestClassWithDefaultSubObject>(ExternalSubObject->NetHandle);
+	auto ClientObject = Client->GetObjectAs<UTestObjectReferences_TestClassWithDefaultSubObject>(Object->NetRefHandle);
+	auto ClientExternalSubObject = Client->GetObjectAs<UTestObjectReferences_TestClassWithDefaultSubObject>(ExternalSubObject->NetRefHandle);
 
 	// Verify that default subobject is created
 	UE_NET_ASSERT_TRUE(ClientObject->CreatedSubObjectRef != nullptr);
@@ -832,9 +852,9 @@ UE_NET_TEST_FIXTURE(FReplicationSystemServerClientTestFixture, TestPartialResolv
 	ServerObject->Ref_TArray.Add(ServerReferencedObjectC);
 
 	// Make sure that references are not replicated to client
-	Server->ReplicationSystem->AddToGroup(NotReplicatedNetObjectGroupHandle, ServerReferencedObjectA->NetHandle);
-	Server->ReplicationSystem->AddToGroup(NotReplicatedNetObjectGroupHandle, ServerReferencedObjectB->NetHandle);
-	Server->ReplicationSystem->AddToGroup(NotReplicatedNetObjectGroupHandle, ServerReferencedObjectC->NetHandle);
+	Server->ReplicationSystem->AddToGroup(NotReplicatedNetObjectGroupHandle, ServerReferencedObjectA->NetRefHandle);
+	Server->ReplicationSystem->AddToGroup(NotReplicatedNetObjectGroupHandle, ServerReferencedObjectB->NetRefHandle);
+	Server->ReplicationSystem->AddToGroup(NotReplicatedNetObjectGroupHandle, ServerReferencedObjectC->NetRefHandle);
 	
 	// Replicate
 	Server->PreSendUpdate();
@@ -842,7 +862,7 @@ UE_NET_TEST_FIXTURE(FReplicationSystemServerClientTestFixture, TestPartialResolv
 	Server->PostSendUpdate();
 
 	// Verify that object has been spawned on client
-	auto ClientObject = Client->GetObjectAs<UTestObjectReferences_TestClassWithReferences>(ServerObject->NetHandle);
+	auto ClientObject = Client->GetObjectAs<UTestObjectReferences_TestClassWithReferences>(ServerObject->NetRefHandle);
 	UE_NET_ASSERT_TRUE(ClientObject != nullptr);
 
 	// Verify that array has been replicated
@@ -854,7 +874,7 @@ UE_NET_TEST_FIXTURE(FReplicationSystemServerClientTestFixture, TestPartialResolv
 	UE_NET_ASSERT_TRUE(ClientObject->Ref_TArray[2] == nullptr);
 
 	// Enable replication for ServerReferenceObjectA
-	Server->ReplicationSystem->RemoveFromGroup(NotReplicatedNetObjectGroupHandle, ServerReferencedObjectA->NetHandle);
+	Server->ReplicationSystem->RemoveFromGroup(NotReplicatedNetObjectGroupHandle, ServerReferencedObjectA->NetRefHandle);
 
 	// Replicate
 	Server->PreSendUpdate();
@@ -867,7 +887,7 @@ UE_NET_TEST_FIXTURE(FReplicationSystemServerClientTestFixture, TestPartialResolv
 	UE_NET_ASSERT_TRUE(ClientObject->Ref_TArray[2] == nullptr);
 
 	// Enable replication for ServerReferenceObjectB 
-	Server->ReplicationSystem->RemoveFromGroup(NotReplicatedNetObjectGroupHandle, ServerReferencedObjectB->NetHandle);
+	Server->ReplicationSystem->RemoveFromGroup(NotReplicatedNetObjectGroupHandle, ServerReferencedObjectB->NetRefHandle);
 
 	// Replicate
 	Server->PreSendUpdate();
@@ -880,7 +900,7 @@ UE_NET_TEST_FIXTURE(FReplicationSystemServerClientTestFixture, TestPartialResolv
 	UE_NET_ASSERT_TRUE(ClientObject->Ref_TArray[2] == nullptr);
 
 	// Enable replication for ServerReferenceObjectB 
-	Server->ReplicationSystem->RemoveFromGroup(NotReplicatedNetObjectGroupHandle, ServerReferencedObjectC->NetHandle);
+	Server->ReplicationSystem->RemoveFromGroup(NotReplicatedNetObjectGroupHandle, ServerReferencedObjectC->NetRefHandle);
 
 	// Replicate
 	Server->PreSendUpdate();
@@ -922,12 +942,12 @@ UE_NET_TEST_FIXTURE(FReplicationSystemServerClientTestFixture, TestOutOfOrderRes
 	Item.ObjectRef = ServerReferencedObjectC;
 	ServerFastArray.Items.Add(Item);
 	ServerFastArray.MarkItemDirty(ServerFastArray.Items[2]);
-	UE::Net::IrisDebugHelper::DebugOutputNetObjectState(ServerObject->NetHandle.GetId(), ServerObject->NetHandle.GetReplicationSystemId());
+	UE::Net::IrisDebugHelper::DebugOutputNetObjectState(ServerObject->NetRefHandle.GetId(), ServerObject->NetRefHandle.GetReplicationSystemId());
 
 	// Make sure that references are not replicated to client
-	Server->ReplicationSystem->AddToGroup(NotReplicatedNetObjectGroupHandle, ServerReferencedObjectA->NetHandle);
-	Server->ReplicationSystem->AddToGroup(NotReplicatedNetObjectGroupHandle, ServerReferencedObjectB->NetHandle);
-	Server->ReplicationSystem->AddToGroup(NotReplicatedNetObjectGroupHandle, ServerReferencedObjectC->NetHandle);
+	Server->ReplicationSystem->AddToGroup(NotReplicatedNetObjectGroupHandle, ServerReferencedObjectA->NetRefHandle);
+	Server->ReplicationSystem->AddToGroup(NotReplicatedNetObjectGroupHandle, ServerReferencedObjectB->NetRefHandle);
+	Server->ReplicationSystem->AddToGroup(NotReplicatedNetObjectGroupHandle, ServerReferencedObjectC->NetRefHandle);
 	
 	// Replicate
 	Server->PreSendUpdate();
@@ -935,7 +955,7 @@ UE_NET_TEST_FIXTURE(FReplicationSystemServerClientTestFixture, TestOutOfOrderRes
 	Server->PostSendUpdate();
 
 	// Verify that object has been spawned on client
-	auto ClientObject = Client->GetObjectAs<UTestObjectReferences_TestClassWithReferences>(ServerObject->NetHandle);
+	auto ClientObject = Client->GetObjectAs<UTestObjectReferences_TestClassWithReferences>(ServerObject->NetRefHandle);
 	UE_NET_ASSERT_TRUE(ClientObject != nullptr);
 
 	auto& ClientFastArray = ClientObject->Ref_FastArray;
@@ -949,14 +969,14 @@ UE_NET_TEST_FIXTURE(FReplicationSystemServerClientTestFixture, TestOutOfOrderRes
 	UE_NET_ASSERT_TRUE(ClientFastArray.Items[2].ObjectRef == nullptr);
 
 	// Enable replication for ServerReferenceObjectA
-	Server->ReplicationSystem->RemoveFromGroup(NotReplicatedNetObjectGroupHandle, ServerReferencedObjectA->NetHandle);
+	Server->ReplicationSystem->RemoveFromGroup(NotReplicatedNetObjectGroupHandle, ServerReferencedObjectA->NetRefHandle);
 
 	// Replicate
 	Server->PreSendUpdate();
 	Server->SendAndDeliverTo(Client, true);
 	Server->PostSendUpdate();
 
-	UE::Net::IrisDebugHelper::DebugOutputNetObjectState(ClientObject->NetHandle.GetId(), ClientObject->NetHandle.GetReplicationSystemId());
+	UE::Net::IrisDebugHelper::DebugOutputNetObjectState(ClientObject->NetRefHandle.GetId(), ClientObject->NetRefHandle.GetReplicationSystemId());
 
 	// Verify that we managed to resolve objectA even though we have other unresolved objects in the array
 	UE_NET_ASSERT_TRUE(ClientFastArray.Items[0].ObjectRef != nullptr);
@@ -964,7 +984,7 @@ UE_NET_TEST_FIXTURE(FReplicationSystemServerClientTestFixture, TestOutOfOrderRes
 	UE_NET_ASSERT_TRUE(ClientFastArray.Items[2].ObjectRef == nullptr);
 
 	// Enable replication for ServerReferenceObjectB 
-	Server->ReplicationSystem->RemoveFromGroup(NotReplicatedNetObjectGroupHandle, ServerReferencedObjectB->NetHandle);
+	Server->ReplicationSystem->RemoveFromGroup(NotReplicatedNetObjectGroupHandle, ServerReferencedObjectB->NetRefHandle);
 
 	// Replicate
 	Server->PreSendUpdate();
@@ -977,7 +997,7 @@ UE_NET_TEST_FIXTURE(FReplicationSystemServerClientTestFixture, TestOutOfOrderRes
 	UE_NET_ASSERT_TRUE(ClientFastArray.Items[2].ObjectRef == nullptr);
 
 	// Enable replication for ServerReferenceObjectC
-	Server->ReplicationSystem->RemoveFromGroup(NotReplicatedNetObjectGroupHandle, ServerReferencedObjectC->NetHandle);
+	Server->ReplicationSystem->RemoveFromGroup(NotReplicatedNetObjectGroupHandle, ServerReferencedObjectC->NetRefHandle);
 
 	// Replicate
 	Server->PreSendUpdate();
@@ -1018,9 +1038,9 @@ UE_NET_TEST_FIXTURE(FReplicationSystemServerClientTestFixture, TestOutOfOrderRes
 	ServerFastArrayEditor.Add(Item);
 
 	// Make sure that references are not replicated to client
-	Server->ReplicationSystem->AddToGroup(NotReplicatedNetObjectGroupHandle, ServerReferencedObjectA->NetHandle);
-	Server->ReplicationSystem->AddToGroup(NotReplicatedNetObjectGroupHandle, ServerReferencedObjectB->NetHandle);
-	Server->ReplicationSystem->AddToGroup(NotReplicatedNetObjectGroupHandle, ServerReferencedObjectC->NetHandle);
+	Server->ReplicationSystem->AddToGroup(NotReplicatedNetObjectGroupHandle, ServerReferencedObjectA->NetRefHandle);
+	Server->ReplicationSystem->AddToGroup(NotReplicatedNetObjectGroupHandle, ServerReferencedObjectB->NetRefHandle);
+	Server->ReplicationSystem->AddToGroup(NotReplicatedNetObjectGroupHandle, ServerReferencedObjectC->NetRefHandle);
 	
 	// Replicate
 	Server->PreSendUpdate();
@@ -1028,7 +1048,7 @@ UE_NET_TEST_FIXTURE(FReplicationSystemServerClientTestFixture, TestOutOfOrderRes
 	Server->PostSendUpdate();
 
 	// Verify that object has been spawned on client
-	auto ClientObject = Client->GetObjectAs<UTestObjectReferences_TestClassWithReferences>(ServerObject->NetHandle);
+	auto ClientObject = Client->GetObjectAs<UTestObjectReferences_TestClassWithReferences>(ServerObject->NetRefHandle);
 	UE_NET_ASSERT_TRUE(ClientObject != nullptr);
 
 	auto& ClientFastArray = ClientObject->Ref_NativeFastArray;
@@ -1042,7 +1062,7 @@ UE_NET_TEST_FIXTURE(FReplicationSystemServerClientTestFixture, TestOutOfOrderRes
 	UE_NET_ASSERT_TRUE(ClientFastArray.Items[2].ObjectRef == nullptr);
 
 	// Enable replication for ServerReferenceObjectA
-	Server->ReplicationSystem->RemoveFromGroup(NotReplicatedNetObjectGroupHandle, ServerReferencedObjectA->NetHandle);
+	Server->ReplicationSystem->RemoveFromGroup(NotReplicatedNetObjectGroupHandle, ServerReferencedObjectA->NetRefHandle);
 
 	// Replicate
 	Server->PreSendUpdate();
@@ -1055,7 +1075,7 @@ UE_NET_TEST_FIXTURE(FReplicationSystemServerClientTestFixture, TestOutOfOrderRes
 	UE_NET_ASSERT_TRUE(ClientFastArray.Items[2].ObjectRef == nullptr);
 
 	// Enable replication for ServerReferenceObjectB 
-	Server->ReplicationSystem->RemoveFromGroup(NotReplicatedNetObjectGroupHandle, ServerReferencedObjectB->NetHandle);
+	Server->ReplicationSystem->RemoveFromGroup(NotReplicatedNetObjectGroupHandle, ServerReferencedObjectB->NetRefHandle);
 
 	// Replicate
 	Server->PreSendUpdate();
@@ -1068,7 +1088,7 @@ UE_NET_TEST_FIXTURE(FReplicationSystemServerClientTestFixture, TestOutOfOrderRes
 	UE_NET_ASSERT_TRUE(ClientFastArray.Items[2].ObjectRef == nullptr);
 
 	// Enable replication for ServerReferenceObjectC
-	Server->ReplicationSystem->RemoveFromGroup(NotReplicatedNetObjectGroupHandle, ServerReferencedObjectC->NetHandle);
+	Server->ReplicationSystem->RemoveFromGroup(NotReplicatedNetObjectGroupHandle, ServerReferencedObjectC->NetRefHandle);
 
 	// Replicate
 	Server->PreSendUpdate();
@@ -1113,9 +1133,9 @@ UE_NET_TEST_FIXTURE(FReplicationSystemServerClientTestFixture, TestOutOfOrderRes
 	ServerFastArrayEditor.Add(Item);
 
 	// Make sure that references are not replicated to client
-	Server->ReplicationSystem->AddToGroup(NotReplicatedNetObjectGroupHandle, ServerReferencedObjectA->NetHandle);
-	Server->ReplicationSystem->AddToGroup(NotReplicatedNetObjectGroupHandle, ServerReferencedObjectB->NetHandle);
-	Server->ReplicationSystem->AddToGroup(NotReplicatedNetObjectGroupHandle, ServerReferencedObjectC->NetHandle);
+	Server->ReplicationSystem->AddToGroup(NotReplicatedNetObjectGroupHandle, ServerReferencedObjectA->NetRefHandle);
+	Server->ReplicationSystem->AddToGroup(NotReplicatedNetObjectGroupHandle, ServerReferencedObjectB->NetRefHandle);
+	Server->ReplicationSystem->AddToGroup(NotReplicatedNetObjectGroupHandle, ServerReferencedObjectC->NetRefHandle);
 	
 	// Replicate
 	Server->PreSendUpdate();
@@ -1127,7 +1147,7 @@ UE_NET_TEST_FIXTURE(FReplicationSystemServerClientTestFixture, TestOutOfOrderRes
 	//ServerFastArrayEditor.Add(Item);
 
 	// Verify that object has been spawned on client
-	auto ClientObject = Client->GetObjectAs<UTestObjectReferences_TestClassWithReferences>(ServerObject->NetHandle);
+	auto ClientObject = Client->GetObjectAs<UTestObjectReferences_TestClassWithReferences>(ServerObject->NetRefHandle);
 	UE_NET_ASSERT_TRUE(ClientObject != nullptr);
 
 	auto& ClientFastArray = ClientObject->Ref_NativeFastArray;
@@ -1144,7 +1164,7 @@ UE_NET_TEST_FIXTURE(FReplicationSystemServerClientTestFixture, TestOutOfOrderRes
 	UE_NET_ASSERT_TRUE(ClientFastArray.Items[2].Ref_TArray[0] == nullptr);
 
 	// Enable replication for ServerReferenceObjectA
-	Server->ReplicationSystem->RemoveFromGroup(NotReplicatedNetObjectGroupHandle, ServerReferencedObjectA->NetHandle);
+	Server->ReplicationSystem->RemoveFromGroup(NotReplicatedNetObjectGroupHandle, ServerReferencedObjectA->NetRefHandle);
 
 	// Replicate
 	Server->PreSendUpdate();
@@ -1157,7 +1177,7 @@ UE_NET_TEST_FIXTURE(FReplicationSystemServerClientTestFixture, TestOutOfOrderRes
 	UE_NET_ASSERT_TRUE(ClientFastArray.Items[2].Ref_TArray[0] == nullptr);
 
 	// Enable replication for ServerReferenceObjectB 
-	Server->ReplicationSystem->RemoveFromGroup(NotReplicatedNetObjectGroupHandle, ServerReferencedObjectB->NetHandle);
+	Server->ReplicationSystem->RemoveFromGroup(NotReplicatedNetObjectGroupHandle, ServerReferencedObjectB->NetRefHandle);
 
 	// Replicate
 	Server->PreSendUpdate();
@@ -1170,7 +1190,7 @@ UE_NET_TEST_FIXTURE(FReplicationSystemServerClientTestFixture, TestOutOfOrderRes
 	UE_NET_ASSERT_TRUE(ClientFastArray.Items[2].Ref_TArray[0] == nullptr);
 
 	// Enable replication for ServerReferenceObjectC
-	Server->ReplicationSystem->RemoveFromGroup(NotReplicatedNetObjectGroupHandle, ServerReferencedObjectC->NetHandle);
+	Server->ReplicationSystem->RemoveFromGroup(NotReplicatedNetObjectGroupHandle, ServerReferencedObjectC->NetRefHandle);
 
 	// Replicate
 	Server->PreSendUpdate();
@@ -1204,14 +1224,14 @@ UE_NET_TEST_FIXTURE(FReplicationSystemServerClientTestFixture, TestRemovedRefere
 	Server->PostSendUpdate();
 
 	// Verify that object has been spawned on client
-	auto ClientObject = Client->GetObjectAs<UTestObjectReferences_TestClassWithReferences>(ServerObject->NetHandle);
+	auto ClientObject = Client->GetObjectAs<UTestObjectReferences_TestClassWithReferences>(ServerObject->NetRefHandle);
 	UE_NET_ASSERT_TRUE(ClientObject != nullptr);
 
 	// Verify that object can be resolved
 	UE_NET_ASSERT_NE(ClientObject->StructWithRef.ObjectRef, TObjectPtr<UObject>(nullptr));
 
 	// Disable replication for ServerReferenceObjectA
-	Server->ReplicationSystem->AddToGroup(NotReplicatedNetObjectGroupHandle, ServerReferencedObjectA->NetHandle);
+	Server->ReplicationSystem->AddToGroup(NotReplicatedNetObjectGroupHandle, ServerReferencedObjectA->NetRefHandle);
 
 	// Replicate to make sure object is destroyed on client
 	Server->PreSendUpdate();
@@ -1230,7 +1250,7 @@ UE_NET_TEST_FIXTURE(FReplicationSystemServerClientTestFixture, TestRemovedRefere
 	UE_NET_ASSERT_EQ(ClientObject->StructWithRef.ObjectRef, TObjectPtr<UObject>(nullptr));
 
 	// Enable replication for ServerReferenceObjectA
-	Server->ReplicationSystem->RemoveFromGroup(NotReplicatedNetObjectGroupHandle, ServerReferencedObjectA->NetHandle);
+	Server->ReplicationSystem->RemoveFromGroup(NotReplicatedNetObjectGroupHandle, ServerReferencedObjectA->NetRefHandle);
 
 	// Replicate
 	Server->PreSendUpdate();
@@ -1271,7 +1291,7 @@ UE_NET_TEST_FIXTURE(FReplicationSystemServerClientTestFixture, TestRemovedRefere
 	Client->PostSendUpdate();
 
 	// Verify that object has been spawned on client
-	auto ClientObject = Client->GetObjectAs<UTestObjectReferences_TestClassWithReferences>(ServerObject->NetHandle);
+	auto ClientObject = Client->GetObjectAs<UTestObjectReferences_TestClassWithReferences>(ServerObject->NetRefHandle);
 	UE_NET_ASSERT_TRUE(ClientObject != nullptr);
 
 	// Verify that object can be resolved
@@ -1279,7 +1299,7 @@ UE_NET_TEST_FIXTURE(FReplicationSystemServerClientTestFixture, TestRemovedRefere
 	UE_NET_ASSERT_NE(ClientObject->StructWithRef_CArray[0].ObjectRef, TObjectPtr<UObject>(nullptr));
 
 	// Disable replication for ServerReferenceObjectA
-	Server->ReplicationSystem->AddToGroup(NotReplicatedNetObjectGroupHandle, ServerReferencedObjectA->NetHandle);
+	Server->ReplicationSystem->AddToGroup(NotReplicatedNetObjectGroupHandle, ServerReferencedObjectA->NetRefHandle);
 
 	// Replicate to make sure object is destroyed on client
 	Server->PreSendUpdate();
@@ -1300,7 +1320,7 @@ UE_NET_TEST_FIXTURE(FReplicationSystemServerClientTestFixture, TestRemovedRefere
 	UE_NET_ASSERT_EQ(ClientObject->StructWithRef_CArray[0].ObjectRef, TObjectPtr<UObject>(nullptr));
 
 	// Enable replication for ServerReferenceObjectA
-	Server->ReplicationSystem->RemoveFromGroup(NotReplicatedNetObjectGroupHandle, ServerReferencedObjectA->NetHandle);
+	Server->ReplicationSystem->RemoveFromGroup(NotReplicatedNetObjectGroupHandle, ServerReferencedObjectA->NetRefHandle);
 
 	// Replicate
 	Server->PreSendUpdate();
@@ -1354,7 +1374,7 @@ UE_NET_TEST_FIXTURE(FReplicationSystemServerClientTestFixture, TestRemovingRefer
 	Server->PostSendUpdate();
 
 	// Verify that object has been spawned on client
-	auto ClientObject = Client->GetObjectAs<UTestObjectReferences_TestClassWithReferences>(ServerObject->NetHandle);
+	auto ClientObject = Client->GetObjectAs<UTestObjectReferences_TestClassWithReferences>(ServerObject->NetRefHandle);
 	UE_NET_ASSERT_NE(ClientObject, nullptr);
 	auto& ClientTArray = ClientObject->StructWithRef_TArray;
 
@@ -1365,9 +1385,9 @@ UE_NET_TEST_FIXTURE(FReplicationSystemServerClientTestFixture, TestRemovingRefer
 	UE_NET_ASSERT_NE(ClientTArray[3].ObjectRef, TObjectPtr<UObject>(nullptr));
 
 	// Disable replication for ServerReferenceObjectA, B and D.
-	Server->ReplicationSystem->AddToGroup(NotReplicatedNetObjectGroupHandle, ServerReferencedObjectA->NetHandle);
-	Server->ReplicationSystem->AddToGroup(NotReplicatedNetObjectGroupHandle, ServerReferencedObjectB->NetHandle);
-	Server->ReplicationSystem->AddToGroup(NotReplicatedNetObjectGroupHandle, ServerReferencedObjectD->NetHandle);
+	Server->ReplicationSystem->AddToGroup(NotReplicatedNetObjectGroupHandle, ServerReferencedObjectA->NetRefHandle);
+	Server->ReplicationSystem->AddToGroup(NotReplicatedNetObjectGroupHandle, ServerReferencedObjectB->NetRefHandle);
+	Server->ReplicationSystem->AddToGroup(NotReplicatedNetObjectGroupHandle, ServerReferencedObjectD->NetRefHandle);
 
 	// Replicate to make sure objects are destroyed on client
 	Server->PreSendUpdate();
@@ -1400,9 +1420,9 @@ UE_NET_TEST_FIXTURE(FReplicationSystemServerClientTestFixture, TestRemovingRefer
 	UE_NET_ASSERT_EQ(ClientTArray[3].ObjectRef, TObjectPtr<UObject>(nullptr));
 
 	// Enable replication for ServerReferenceObjectA, B and D.
-	Server->ReplicationSystem->RemoveFromGroup(NotReplicatedNetObjectGroupHandle, ServerReferencedObjectA->NetHandle);
-	Server->ReplicationSystem->RemoveFromGroup(NotReplicatedNetObjectGroupHandle, ServerReferencedObjectB->NetHandle);
-	Server->ReplicationSystem->RemoveFromGroup(NotReplicatedNetObjectGroupHandle, ServerReferencedObjectD->NetHandle);
+	Server->ReplicationSystem->RemoveFromGroup(NotReplicatedNetObjectGroupHandle, ServerReferencedObjectA->NetRefHandle);
+	Server->ReplicationSystem->RemoveFromGroup(NotReplicatedNetObjectGroupHandle, ServerReferencedObjectB->NetRefHandle);
+	Server->ReplicationSystem->RemoveFromGroup(NotReplicatedNetObjectGroupHandle, ServerReferencedObjectD->NetRefHandle);
 
 	// Replicate
 	Server->PreSendUpdate();
@@ -1445,7 +1465,7 @@ UE_NET_TEST_FIXTURE(FReplicationSystemServerClientTestFixture, TestTearOffDoesIn
 	Client->PostSendUpdate();
 
 	// Verify that object has been spawned on client
-	auto ClientObject = Client->GetObjectAs<UTestObjectReferences_TestClassWithReferences>(ServerObject->NetHandle);
+	auto ClientObject = Client->GetObjectAs<UTestObjectReferences_TestClassWithReferences>(ServerObject->NetRefHandle);
 	UE_NET_ASSERT_NE(ClientObject, nullptr);
 
 	// Verify that object can be resolved
