Iris - Implemented globally unique NetHandles.

diff --git a/Engine/Plugins/Runtime/ReplicationSystemTestPlugin/Source/Private/Tests/ReplicationSystem/TestObjectSplitting.cpp b/Engine/Plugins/Runtime/ReplicationSystemTestPlugin/Source/Private/Tests/ReplicationSystem/TestObjectSplitting.cpp
--- a/Engine/Plugins/Runtime/ReplicationSystemTestPlugin/Source/Private/Tests/ReplicationSystem/TestObjectSplitting.cpp
+++ b/Engine/Plugins/Runtime/ReplicationSystemTestPlugin/Source/Private/Tests/ReplicationSystem/TestObjectSplitting.cpp
@@ -75,7 +75,7 @@ protected:
 		return Object;
 	}
 
-	UTestReplicatedIrisObject* CreateSubObject(FReplicationSystemTestNode* Node, FNetHandle Parent)
+	UTestReplicatedIrisObject* CreateSubObject(FReplicationSystemTestNode* Node, FNetRefHandle Parent)
 	{
 		UTestReplicatedIrisObject::FComponents Components;
 		Components.DynamicStateComponentCount = 1;
@@ -124,7 +124,7 @@ UE_NET_TEST_FIXTURE(FSplitObjectTestFixture, SplitHugeObjectOnCreation)
 	Server->PostSendUpdate();
 
 	// As the payload is huge we don't expect the whole payload to arrive the first frame
-	const UTestReplicatedIrisObject* ClientObject = Cast<UTestReplicatedIrisObject>(Client->GetReplicationBridge()->GetReplicatedObject(ServerObject->NetHandle));
+	const UTestReplicatedIrisObject* ClientObject = Cast<UTestReplicatedIrisObject>(Client->GetReplicationBridge()->GetReplicatedObject(ServerObject->NetRefHandle));
 	UE_NET_ASSERT_EQ(ClientObject, nullptr);
 
 	for (uint32 RetryIt = 0; RetryIt != HugeObjectMaxNetTickCountToArrive && ClientObject == nullptr; ++RetryIt)
@@ -132,7 +132,7 @@ UE_NET_TEST_FIXTURE(FSplitObjectTestFixture, SplitHugeObjectOnCreation)
 		Server->PreSendUpdate();
 		Server->SendAndDeliverTo(Client, DeliverPacket);
 		Server->PostSendUpdate();
-		ClientObject = Cast<UTestReplicatedIrisObject>(Client->GetReplicationBridge()->GetReplicatedObject(ServerObject->NetHandle));
+		ClientObject = Cast<UTestReplicatedIrisObject>(Client->GetReplicationBridge()->GetReplicatedObject(ServerObject->NetRefHandle));
 	}
 	UE_NET_ASSERT_NE(ClientObject, nullptr);
 }
@@ -149,7 +149,7 @@ UE_NET_TEST_FIXTURE(FSplitObjectTestFixture, SplitHugeObjectAfterCreation)
 	Server->SendAndDeliverTo(Client, DeliverPacket);
 	Server->PostSendUpdate();
 
-	const UTestReplicatedIrisObject* ClientObject = Cast<UTestReplicatedIrisObject>(Client->GetReplicationBridge()->GetReplicatedObject(ServerObject->NetHandle));
+	const UTestReplicatedIrisObject* ClientObject = Cast<UTestReplicatedIrisObject>(Client->GetReplicationBridge()->GetReplicatedObject(ServerObject->NetRefHandle));
 	UE_NET_ASSERT_NE(ClientObject, nullptr);
 
 	SetObjectPayloadByteCount(ServerObject, HugeObjectPayloadByteCount);
@@ -177,7 +177,7 @@ UE_NET_TEST_FIXTURE(FSplitObjectTestFixture, SplitObjectWithHugeSubObjectsOnCrea
 	UTestReplicatedIrisObject* ServerSubObjects[SubObjectCount];
 	for (uint32 SubObjectIt=0; SubObjectIt != SubObjectCount; ++SubObjectIt)
 	{
-		ServerSubObjects[SubObjectIt] = CreateSubObject(Server, ServerObject->NetHandle);
+		ServerSubObjects[SubObjectIt] = CreateSubObject(Server, ServerObject->NetRefHandle);
 		SetObjectPayloadByteCount(ServerSubObjects[SubObjectIt], 4096U);
 	}
 
@@ -187,7 +187,7 @@ UE_NET_TEST_FIXTURE(FSplitObjectTestFixture, SplitObjectWithHugeSubObjectsOnCrea
 	Server->PostSendUpdate();
 
 	// As the payload is huge we don't expect the whole payload to arrive the first frame
-	const UTestReplicatedIrisObject* ClientObject = Cast<UTestReplicatedIrisObject>(Client->GetReplicationBridge()->GetReplicatedObject(ServerObject->NetHandle));
+	const UTestReplicatedIrisObject* ClientObject = Cast<UTestReplicatedIrisObject>(Client->GetReplicationBridge()->GetReplicatedObject(ServerObject->NetRefHandle));
 	UE_NET_ASSERT_EQ(ClientObject, nullptr);
 
 	for (uint32 RetryIt = 0; RetryIt != HugeObjectMaxNetTickCountToArrive && ClientObject == nullptr; ++RetryIt)
@@ -195,14 +195,14 @@ UE_NET_TEST_FIXTURE(FSplitObjectTestFixture, SplitObjectWithHugeSubObjectsOnCrea
 		Server->PreSendUpdate();
 		Server->SendAndDeliverTo(Client, DeliverPacket);
 		Server->PostSendUpdate();
-		ClientObject = Cast<UTestReplicatedIrisObject>(Client->GetReplicationBridge()->GetReplicatedObject(ServerObject->NetHandle));
+		ClientObject = Cast<UTestReplicatedIrisObject>(Client->GetReplicationBridge()->GetReplicatedObject(ServerObject->NetRefHandle));
 	}
 	UE_NET_ASSERT_NE(ClientObject, nullptr);
 
 	// Verify the subobjects made it through as well.
 	for (uint32 SubObjectIt=0; SubObjectIt != SubObjectCount; ++SubObjectIt)
 	{
-		const UTestReplicatedIrisObject* ClientSubObject = Cast<UTestReplicatedIrisObject>(Client->GetReplicationBridge()->GetReplicatedObject(ServerSubObjects[SubObjectIt]->NetHandle));
+		const UTestReplicatedIrisObject* ClientSubObject = Cast<UTestReplicatedIrisObject>(Client->GetReplicationBridge()->GetReplicatedObject(ServerSubObjects[SubObjectIt]->NetRefHandle));
 		UE_NET_ASSERT_NE(ClientSubObject, nullptr);
 	}
 }
@@ -222,7 +222,7 @@ UE_NET_TEST_FIXTURE(FSplitObjectTestFixture, SplitObjectWithSubObjectsWithHugeAt
 
 	for (uint32 SubObjectIt=0; SubObjectIt != SubObjectCount; ++SubObjectIt)
 	{
-		UTestReplicatedIrisObject* ServerSubObject = CreateSubObject(Server, ServerObject->NetHandle);
+		UTestReplicatedIrisObject* ServerSubObject = CreateSubObject(Server, ServerObject->NetRefHandle);
 		ServerSubObjects[SubObjectIt] = ServerSubObject;
 		SetObjectPayloadByteCount(ServerSubObject, SubObjectPayloadByteCount);
 
@@ -237,7 +237,7 @@ UE_NET_TEST_FIXTURE(FSplitObjectTestFixture, SplitObjectWithSubObjectsWithHugeAt
 			Attachment = ServerMockNetObjectAttachmentHandler->CreateUnreliableNetObjectAttachment(SubObjectPayloadByteCount*8U);
 		}
 
-		FNetObjectReference AttachmentTarget = FObjectReferenceCache::MakeNetObjectReference(ServerSubObject->NetHandle);
+		FNetObjectReference AttachmentTarget = FObjectReferenceCache::MakeNetObjectReference(ServerSubObject->NetRefHandle);
 		Server->GetReplicationSystem()->QueueNetObjectAttachment(Client->ConnectionIdOnServer, AttachmentTarget, Attachment);
 	}
 
@@ -247,7 +247,7 @@ UE_NET_TEST_FIXTURE(FSplitObjectTestFixture, SplitObjectWithSubObjectsWithHugeAt
 	Server->PostSendUpdate();
 
 	// As the payload is huge we don't expect the whole payload to arrive the first frame
-	const UTestReplicatedIrisObject* ClientObject = Cast<UTestReplicatedIrisObject>(Client->GetReplicationBridge()->GetReplicatedObject(ServerObject->NetHandle));
+	const UTestReplicatedIrisObject* ClientObject = Cast<UTestReplicatedIrisObject>(Client->GetReplicationBridge()->GetReplicatedObject(ServerObject->NetRefHandle));
 	UE_NET_ASSERT_EQ(ClientObject, nullptr);
 
 	for (uint32 RetryIt = 0; RetryIt != HugeObjectMaxNetTickCountToArrive && ClientObject == nullptr; ++RetryIt)
@@ -255,14 +255,14 @@ UE_NET_TEST_FIXTURE(FSplitObjectTestFixture, SplitObjectWithSubObjectsWithHugeAt
 		Server->PreSendUpdate();
 		Server->SendAndDeliverTo(Client, DeliverPacket);
 		Server->PostSendUpdate();
-		ClientObject = Cast<UTestReplicatedIrisObject>(Client->GetReplicationBridge()->GetReplicatedObject(ServerObject->NetHandle));
+		ClientObject = Cast<UTestReplicatedIrisObject>(Client->GetReplicationBridge()->GetReplicatedObject(ServerObject->NetRefHandle));
 	}
 	UE_NET_ASSERT_NE(ClientObject, nullptr);
 
 	// Verify the subobjects made it through.
 	for (uint32 SubObjectIt=0; SubObjectIt != SubObjectCount; ++SubObjectIt)
 	{
-		const UTestReplicatedIrisObject* ClientSubObject = Cast<UTestReplicatedIrisObject>(Client->GetReplicationBridge()->GetReplicatedObject(ServerSubObjects[SubObjectIt]->NetHandle));
+		const UTestReplicatedIrisObject* ClientSubObject = Cast<UTestReplicatedIrisObject>(Client->GetReplicationBridge()->GetReplicatedObject(ServerSubObjects[SubObjectIt]->NetRefHandle));
 		UE_NET_ASSERT_NE(ClientSubObject, nullptr);
 	}
 
@@ -294,7 +294,7 @@ UE_NET_TEST_FIXTURE(FSplitObjectTestFixture, SplitObjectWithSubObjectsWithHugeAt
 
 	for (uint32 SubObjectIt=0; SubObjectIt != SubObjectCount; ++SubObjectIt)
 	{
-		UTestReplicatedIrisObject* ServerSubObject = CreateSubObject(Server, ServerObject->NetHandle);
+		UTestReplicatedIrisObject* ServerSubObject = CreateSubObject(Server, ServerObject->NetRefHandle);
 		ServerSubObjects[SubObjectIt] = ServerSubObject;
 	}
 
@@ -304,13 +304,13 @@ UE_NET_TEST_FIXTURE(FSplitObjectTestFixture, SplitObjectWithSubObjectsWithHugeAt
 	Server->PostSendUpdate();
 
 	// As the payload is huge we don't expect the whole payload to arrive the first frame
-	const UTestReplicatedIrisObject* ClientObject = Cast<UTestReplicatedIrisObject>(Client->GetReplicationBridge()->GetReplicatedObject(ServerObject->NetHandle));
+	const UTestReplicatedIrisObject* ClientObject = Cast<UTestReplicatedIrisObject>(Client->GetReplicationBridge()->GetReplicatedObject(ServerObject->NetRefHandle));
 	UE_NET_ASSERT_NE(ClientObject, nullptr);
 
 	// Verify the subobjects made it through.
 	for (uint32 SubObjectIt=0; SubObjectIt != SubObjectCount; ++SubObjectIt)
 	{
-		const UTestReplicatedIrisObject* ClientSubObject = Cast<UTestReplicatedIrisObject>(Client->GetReplicationBridge()->GetReplicatedObject(ServerSubObjects[SubObjectIt]->NetHandle));
+		const UTestReplicatedIrisObject* ClientSubObject = Cast<UTestReplicatedIrisObject>(Client->GetReplicationBridge()->GetReplicatedObject(ServerSubObjects[SubObjectIt]->NetRefHandle));
 		UE_NET_ASSERT_NE(ClientSubObject, nullptr);
 	}
 
@@ -331,7 +331,7 @@ UE_NET_TEST_FIXTURE(FSplitObjectTestFixture, SplitObjectWithSubObjectsWithHugeAt
 			Attachment = ServerMockNetObjectAttachmentHandler->CreateUnreliableNetObjectAttachment(SubObjectPayloadByteCount*8U);
 		}
 
-		FNetObjectReference AttachmentTarget = FObjectReferenceCache::MakeNetObjectReference(ServerSubObject->NetHandle);
+		FNetObjectReference AttachmentTarget = FObjectReferenceCache::MakeNetObjectReference(ServerSubObject->NetRefHandle);
 		Server->GetReplicationSystem()->QueueNetObjectAttachment(Client->ConnectionIdOnServer, AttachmentTarget, Attachment);
 	}
 
@@ -343,7 +343,7 @@ UE_NET_TEST_FIXTURE(FSplitObjectTestFixture, SplitObjectWithSubObjectsWithHugeAt
 		Server->PostSendUpdate();
 
 		// Assume that if one subobject has received its huge state then all of them have
-		const UTestReplicatedIrisObject* ClientSubObject = Cast<UTestReplicatedIrisObject>(Client->GetReplicationBridge()->GetReplicatedObject(ServerSubObjects[SubObjectCount - 1U]->NetHandle));
+		const UTestReplicatedIrisObject* ClientSubObject = Cast<UTestReplicatedIrisObject>(Client->GetReplicationBridge()->GetReplicatedObject(ServerSubObjects[SubObjectCount - 1U]->NetRefHandle));
 		if (ClientSubObject->DynamicStateComponents[0]->IntArray.Num() > 0)
 		{
 			bHasReceivedHugeState = true;
@@ -371,7 +371,7 @@ UE_NET_TEST_FIXTURE(FSplitObjectTestFixture, HugeObjectStateCanBeSentBackToBack)
 
 	for (uint32 SubObjectIt=0; SubObjectIt != SubObjectCount; ++SubObjectIt)
 	{
-		UTestReplicatedIrisObject* ServerSubObject = CreateSubObject(Server, ServerObject->NetHandle);
+		UTestReplicatedIrisObject* ServerSubObject = CreateSubObject(Server, ServerObject->NetRefHandle);
 		ServerSubObjects[SubObjectIt] = ServerSubObject;
 	}
 
@@ -381,13 +381,13 @@ UE_NET_TEST_FIXTURE(FSplitObjectTestFixture, HugeObjectStateCanBeSentBackToBack)
 	Server->PostSendUpdate();
 
 	// As the payload is huge we don't expect the whole payload to arrive the first frame
-	const UTestReplicatedIrisObject* ClientObject = Cast<UTestReplicatedIrisObject>(Client->GetReplicationBridge()->GetReplicatedObject(ServerObject->NetHandle));
+	const UTestReplicatedIrisObject* ClientObject = Cast<UTestReplicatedIrisObject>(Client->GetReplicationBridge()->GetReplicatedObject(ServerObject->NetRefHandle));
 	UE_NET_ASSERT_NE(ClientObject, nullptr);
 
 	// Verify the subobjects made it through.
 	for (uint32 SubObjectIt=0; SubObjectIt != SubObjectCount; ++SubObjectIt)
 	{
-		const UTestReplicatedIrisObject* ClientSubObject = Cast<UTestReplicatedIrisObject>(Client->GetReplicationBridge()->GetReplicatedObject(ServerSubObjects[SubObjectIt]->NetHandle));
+		const UTestReplicatedIrisObject* ClientSubObject = Cast<UTestReplicatedIrisObject>(Client->GetReplicationBridge()->GetReplicatedObject(ServerSubObjects[SubObjectIt]->NetRefHandle));
 		UE_NET_ASSERT_NE(ClientSubObject, nullptr);
 	}
 
@@ -408,7 +408,7 @@ UE_NET_TEST_FIXTURE(FSplitObjectTestFixture, HugeObjectStateCanBeSentBackToBack)
 			Attachment = ServerMockNetObjectAttachmentHandler->CreateUnreliableNetObjectAttachment(SubObjectPayloadByteCount*8U);
 		}
 
-		FNetObjectReference AttachmentTarget = FObjectReferenceCache::MakeNetObjectReference(ServerSubObject->NetHandle);
+		FNetObjectReference AttachmentTarget = FObjectReferenceCache::MakeNetObjectReference(ServerSubObject->NetRefHandle);
 		Server->GetReplicationSystem()->QueueNetObjectAttachment(Client->ConnectionIdOnServer, AttachmentTarget, Attachment);
 	}
 
@@ -420,7 +420,7 @@ UE_NET_TEST_FIXTURE(FSplitObjectTestFixture, HugeObjectStateCanBeSentBackToBack)
 		Server->PostSendUpdate();
 
 		// Assume that if one subobject has received its huge state then all of them have
-		const UTestReplicatedIrisObject* ClientSubObject = Cast<UTestReplicatedIrisObject>(Client->GetReplicationBridge()->GetReplicatedObject(ServerSubObjects[SubObjectCount - 1U]->NetHandle));
+		const UTestReplicatedIrisObject* ClientSubObject = Cast<UTestReplicatedIrisObject>(Client->GetReplicationBridge()->GetReplicatedObject(ServerSubObjects[SubObjectCount - 1U]->NetRefHandle));
 		if (ClientSubObject->DynamicStateComponents[0]->IntArray.Num() > 0)
 		{
 			bHasReceivedHugeState = true;
@@ -446,7 +446,7 @@ UE_NET_TEST_FIXTURE(FSplitObjectTestFixture, SplitObjectCanBeSentBackToBack)
 	Server->PostSendUpdate();
 
 	// As the payload is huge we don't expect the whole payload to arrive the first frame
-	const UTestReplicatedIrisObject* ClientObject = Cast<UTestReplicatedIrisObject>(Client->GetReplicationBridge()->GetReplicatedObject(ServerObject->NetHandle));
+	const UTestReplicatedIrisObject* ClientObject = Cast<UTestReplicatedIrisObject>(Client->GetReplicationBridge()->GetReplicatedObject(ServerObject->NetRefHandle));
 	UE_NET_ASSERT_EQ(ClientObject, nullptr);
 
 	const int OriginalArrayCount = ServerObject->DynamicStateComponents[0]->IntArray.Num();
@@ -459,7 +459,7 @@ UE_NET_TEST_FIXTURE(FSplitObjectTestFixture, SplitObjectCanBeSentBackToBack)
 		Server->PreSendUpdate();
 		Server->SendAndDeliverTo(Client, DeliverPacket);
 		Server->PostSendUpdate();
-		ClientObject = Cast<UTestReplicatedIrisObject>(Client->GetReplicationBridge()->GetReplicatedObject(ServerObject->NetHandle));
+		ClientObject = Cast<UTestReplicatedIrisObject>(Client->GetReplicationBridge()->GetReplicatedObject(ServerObject->NetRefHandle));
 	}
 	UE_NET_ASSERT_NE(ClientObject, nullptr);
 	UE_NET_ASSERT_EQ(ClientObject->DynamicStateComponents[0]->IntArray.Num(), OriginalArrayCount);
@@ -485,7 +485,7 @@ UE_NET_TEST_FIXTURE(FSplitObjectTestFixture, SplitObjectIsDeletedAfterBeingCreat
 	FReplicationSystemTestClient* Client = CreateClient();
 
 	UTestReplicatedIrisObject* ServerObject = CreateHugeObject(Server);
-	const FNetHandle ServerNetHandle = ServerObject->NetHandle;
+	const FNetRefHandle ServerNetRefHandle = ServerObject->NetRefHandle;
 
 	// Send and deliver packet
 	Server->PreSendUpdate();
@@ -493,7 +493,7 @@ UE_NET_TEST_FIXTURE(FSplitObjectTestFixture, SplitObjectIsDeletedAfterBeingCreat
 	Server->PostSendUpdate();
 
 	// As the payload is huge we don't expect the whole payload to arrive the first frame
-	const UTestReplicatedIrisObject* ClientObject = Cast<UTestReplicatedIrisObject>(Client->GetReplicationBridge()->GetReplicatedObject(ServerNetHandle));
+	const UTestReplicatedIrisObject* ClientObject = Cast<UTestReplicatedIrisObject>(Client->GetReplicationBridge()->GetReplicatedObject(ServerNetRefHandle));
 	UE_NET_ASSERT_EQ(ClientObject, nullptr);
 
 	Server->DestroyObject(ServerObject);
@@ -503,7 +503,7 @@ UE_NET_TEST_FIXTURE(FSplitObjectTestFixture, SplitObjectIsDeletedAfterBeingCreat
 		Server->PreSendUpdate();
 		Server->SendAndDeliverTo(Client, DeliverPacket);
 		Server->PostSendUpdate();
-		ClientObject = Cast<UTestReplicatedIrisObject>(Client->GetReplicationBridge()->GetReplicatedObject(ServerNetHandle));
+		ClientObject = Cast<UTestReplicatedIrisObject>(Client->GetReplicationBridge()->GetReplicatedObject(ServerNetRefHandle));
 	}
 	UE_NET_ASSERT_NE(ClientObject, nullptr);
 
@@ -512,7 +512,7 @@ UE_NET_TEST_FIXTURE(FSplitObjectTestFixture, SplitObjectIsDeletedAfterBeingCreat
 	Server->SendAndDeliverTo(Client, DeliverPacket);
 	Server->PostSendUpdate();
 
-	ClientObject = Cast<UTestReplicatedIrisObject>(Client->GetReplicationBridge()->GetReplicatedObject(ServerNetHandle));
+	ClientObject = Cast<UTestReplicatedIrisObject>(Client->GetReplicationBridge()->GetReplicatedObject(ServerNetRefHandle));
 	UE_NET_ASSERT_EQ(ClientObject, nullptr);
 }
 
@@ -526,7 +526,7 @@ UE_NET_TEST_FIXTURE(FSplitObjectTestFixture, SubObjectToHugeObjectCanBeDeleted)
 	UTestReplicatedIrisObject* ServerSubObjects[SubObjectCount];
 	for (uint32 SubObjectIt=0; SubObjectIt != SubObjectCount; ++SubObjectIt)
 	{
-		ServerSubObjects[SubObjectIt] = CreateSubObject(Server, ServerObject->NetHandle);
+		ServerSubObjects[SubObjectIt] = CreateSubObject(Server, ServerObject->NetRefHandle);
 	}
 
 	// Send and deliver packet
@@ -534,11 +534,11 @@ UE_NET_TEST_FIXTURE(FSplitObjectTestFixture, SubObjectToHugeObjectCanBeDeleted)
 	Server->SendAndDeliverTo(Client, DeliverPacket);
 	Server->PostSendUpdate();
 
-	const UTestReplicatedIrisObject* ClientObject = Cast<UTestReplicatedIrisObject>(Client->GetReplicationBridge()->GetReplicatedObject(ServerObject->NetHandle));
+	const UTestReplicatedIrisObject* ClientObject = Cast<UTestReplicatedIrisObject>(Client->GetReplicationBridge()->GetReplicatedObject(ServerObject->NetRefHandle));
 	UE_NET_ASSERT_NE(ClientObject, nullptr);
 
-	const FNetHandle SubObjectNetHandle = ServerSubObjects[0]->NetHandle;
-	const UTestReplicatedIrisObject* ClientSubObject = Cast<UTestReplicatedIrisObject>(Client->GetReplicationBridge()->GetReplicatedObject(ServerObject->NetHandle));
+	const FNetRefHandle SubObjectNetRefHandle = ServerSubObjects[0]->NetRefHandle;
+	const UTestReplicatedIrisObject* ClientSubObject = Cast<UTestReplicatedIrisObject>(Client->GetReplicationBridge()->GetReplicatedObject(ServerObject->NetRefHandle));
 	UE_NET_ASSERT_NE(ClientSubObject, nullptr);
 
 	// Make subobject payloads huge
@@ -563,7 +563,7 @@ UE_NET_TEST_FIXTURE(FSplitObjectTestFixture, SubObjectToHugeObjectCanBeDeleted)
 		Server->PostSendUpdate();
 
 		// Assume that if one subobject has received its huge state then all of them have
-		ClientSubObject = Cast<UTestReplicatedIrisObject>(Client->GetReplicationBridge()->GetReplicatedObject(SubObjectNetHandle));
+		ClientSubObject = Cast<UTestReplicatedIrisObject>(Client->GetReplicationBridge()->GetReplicatedObject(SubObjectNetRefHandle));
 		UE_NET_ASSERT_NE(ClientSubObject, nullptr);
 		if (ClientSubObject->DynamicStateComponents[0]->IntArray.Num() > 0)
 		{
@@ -577,7 +577,7 @@ UE_NET_TEST_FIXTURE(FSplitObjectTestFixture, SubObjectToHugeObjectCanBeDeleted)
 	Server->SendAndDeliverTo(Client, DeliverPacket);
 	Server->PostSendUpdate();
 
-	ClientSubObject = Cast<UTestReplicatedIrisObject>(Client->GetReplicationBridge()->GetReplicatedObject(SubObjectNetHandle));
+	ClientSubObject = Cast<UTestReplicatedIrisObject>(Client->GetReplicationBridge()->GetReplicatedObject(SubObjectNetRefHandle));
 	UE_NET_ASSERT_EQ(ClientSubObject, nullptr);
 }
 
@@ -610,7 +610,7 @@ UE_NET_TEST_FIXTURE(FSplitObjectTestFixture, TearOffOnCreation)
 	// Verify that ClientObject is torn-off and that the final state was applied
 	UTestReplicatedIrisObject* ClientObjectThatWasTornOff = Cast<UTestReplicatedIrisObject>(Client->CreatedObjects[NumObjectsCreatedOnClientBeforeReplication].Get());
 	UE_NET_ASSERT_EQ(ClientObjectThatWasTornOff->DynamicStateComponents[0]->IntArray.Num(), ServerObject->DynamicStateComponents[0]->IntArray.Num());
-	UE_NET_ASSERT_EQ(Cast<UTestReplicatedIrisObject>(Client->GetReplicationBridge()->GetReplicatedObject(ServerObject->NetHandle)), nullptr);
+	UE_NET_ASSERT_EQ(Cast<UTestReplicatedIrisObject>(Client->GetReplicationBridge()->GetReplicatedObject(ServerObject->NetRefHandle)), nullptr);
 }
 
 // Test TearOff for existing confirmed object during huge object state send
@@ -626,7 +626,7 @@ UE_NET_TEST_FIXTURE(FSplitObjectTestFixture, TearOffCreatedObjectWithHugePayload
 	Server->PostSendUpdate();
 
 	// Store client object while it can still be found using the server net handle.
-	UTestReplicatedIrisObject* ClientObjectThatWillBeTornOff = Cast<UTestReplicatedIrisObject>(Client->GetReplicationBridge()->GetReplicatedObject(ServerObject->NetHandle));
+	UTestReplicatedIrisObject* ClientObjectThatWillBeTornOff = Cast<UTestReplicatedIrisObject>(Client->GetReplicationBridge()->GetReplicatedObject(ServerObject->NetRefHandle));
 	UE_NET_ASSERT_NE(ClientObjectThatWillBeTornOff, nullptr);
 
 	// Set huge payload and TearOff the object
@@ -647,7 +647,7 @@ UE_NET_TEST_FIXTURE(FSplitObjectTestFixture, TearOffCreatedObjectWithHugePayload
 	}
 
 	// Verify that ClientObject is torn-off and that the final state was applied
-	UE_NET_ASSERT_EQ(Cast<UTestReplicatedIrisObject>(Client->GetReplicationBridge()->GetReplicatedObject(ServerObject->NetHandle)), nullptr);
+	UE_NET_ASSERT_EQ(Cast<UTestReplicatedIrisObject>(Client->GetReplicationBridge()->GetReplicatedObject(ServerObject->NetRefHandle)), nullptr);
 	UE_NET_ASSERT_EQ(ClientObjectThatWillBeTornOff->DynamicStateComponents[0]->IntArray.Num(), ServerObject->DynamicStateComponents[0]->IntArray.Num());
 }
 
@@ -684,7 +684,7 @@ UE_NET_TEST_FIXTURE(FSplitObjectTestFixture, TearOffWhileHugeObjectStateIsSendin
 	UE_NET_ASSERT_EQ(Client->CreatedObjects.Num(), NumObjectsCreatedOnClientBeforeReplication + 1);
 
 	// Verify we have the previous state
-	const UTestReplicatedIrisObject* ClientObjectThatWillBeTornOff = Cast<UTestReplicatedIrisObject>(Client->GetReplicationBridge()->GetReplicatedObject(ServerObject->NetHandle));
+	const UTestReplicatedIrisObject* ClientObjectThatWillBeTornOff = Cast<UTestReplicatedIrisObject>(Client->GetReplicationBridge()->GetReplicatedObject(ServerObject->NetRefHandle));
 	UE_NET_ASSERT_EQ(ClientObjectThatWillBeTornOff->IntA ^ 1, ServerObject->IntA);
 
 	Server->PreSendUpdate();
@@ -692,7 +692,7 @@ UE_NET_TEST_FIXTURE(FSplitObjectTestFixture, TearOffWhileHugeObjectStateIsSendin
 	Server->PostSendUpdate();
 
 	// Verify that ClientObject is torn-off and that the final state was applied
-	UE_NET_ASSERT_EQ(Cast<UTestReplicatedIrisObject>(Client->GetReplicationBridge()->GetReplicatedObject(ServerObject->NetHandle)), nullptr);
+	UE_NET_ASSERT_EQ(Cast<UTestReplicatedIrisObject>(Client->GetReplicationBridge()->GetReplicatedObject(ServerObject->NetRefHandle)), nullptr);
 	UE_NET_ASSERT_EQ(ClientObjectThatWillBeTornOff->IntA, ServerObject->IntA);
 }
 
@@ -713,12 +713,12 @@ UE_NET_TEST_FIXTURE(FSplitObjectTestFixture, TestCancelPendingDestroyOfHugeObjec
 	}
 
 	// Filter out object to cause a PendingDestroy
-	Server->GetReplicationSystem()->AddToGroup(NotReplicatedNetObjectGroupHandle, ServerObject->NetHandle);
+	Server->GetReplicationSystem()->AddToGroup(NotReplicatedNetObjectGroupHandle, ServerObject->NetRefHandle);
 	Server->PreSendUpdate();
 	Server->PostSendUpdate();
 
 	// Remove object from filter to cause object to end up in CancelPendingDestroy
-	Server->GetReplicationSystem()->RemoveFromGroup(NotReplicatedNetObjectGroupHandle, ServerObject->NetHandle);
+	Server->GetReplicationSystem()->RemoveFromGroup(NotReplicatedNetObjectGroupHandle, ServerObject->NetRefHandle);
 	Server->PreSendUpdate();
 	Server->PostSendUpdate();
 
@@ -734,7 +734,7 @@ UE_NET_TEST_FIXTURE(FSplitObjectTestFixture, TestCancelPendingDestroyOfHugeObjec
 	}
 
 	// Verify that the object now exists on client
-	UE_NET_ASSERT_NE(Client->GetReplicationBridge()->GetReplicatedObject(ServerObject->NetHandle), nullptr);
+	UE_NET_ASSERT_NE(Client->GetReplicationBridge()->GetReplicatedObject(ServerObject->NetRefHandle), nullptr);
 
 	// Modify a property on the object and make sure it's replicated as the object should now be confirmed created
 	ServerObject->IntA ^= 1;
@@ -743,7 +743,7 @@ UE_NET_TEST_FIXTURE(FSplitObjectTestFixture, TestCancelPendingDestroyOfHugeObjec
 	Server->SendAndDeliverTo(Client, DeliverPacket);
 	Server->PostSendUpdate();
 
-	UTestReplicatedIrisObject* ClientObject = Cast<UTestReplicatedIrisObject>(Client->GetReplicationBridge()->GetReplicatedObject(ServerObject->NetHandle));
+	UTestReplicatedIrisObject* ClientObject = Cast<UTestReplicatedIrisObject>(Client->GetReplicationBridge()->GetReplicatedObject(ServerObject->NetRefHandle));
 	UE_NET_ASSERT_NE(ClientObject, nullptr);
 	UE_NET_ASSERT_EQ(ClientObject->IntA, ServerObject->IntA);
 }
@@ -766,12 +766,12 @@ UE_NET_TEST_FIXTURE(FSplitObjectTestFixture, TestCancelPendingDestroyOfHugeObjec
 	}
 
 	// Filter out object to cause a PendingDestroy
-	Server->GetReplicationSystem()->AddToGroup(NotReplicatedNetObjectGroupHandle, ServerObject->NetHandle);
+	Server->GetReplicationSystem()->AddToGroup(NotReplicatedNetObjectGroupHandle, ServerObject->NetRefHandle);
 	Server->PreSendUpdate();
 	Server->PostSendUpdate();
 
 	// Remove object from filter to cause object to end up in CancelPendingDestroy
-	Server->GetReplicationSystem()->RemoveFromGroup(NotReplicatedNetObjectGroupHandle, ServerObject->NetHandle);
+	Server->GetReplicationSystem()->RemoveFromGroup(NotReplicatedNetObjectGroupHandle, ServerObject->NetRefHandle);
 	Server->PreSendUpdate();
 	Server->PostSendUpdate();
 
@@ -803,7 +803,7 @@ UE_NET_TEST_FIXTURE(FSplitObjectTestFixture, TestCancelPendingDestroyOfHugeObjec
 		}
 	}
 
-	UE_NET_ASSERT_NE(Client->GetReplicationBridge()->GetReplicatedObject(ServerObject->NetHandle), nullptr);
+	UE_NET_ASSERT_NE(Client->GetReplicationBridge()->GetReplicatedObject(ServerObject->NetRefHandle), nullptr);
 }
 
 UE_NET_TEST_FIXTURE(FSplitObjectTestFixture, TestCancelPendingDestroyDuringHugeObjectStateUpdate)
@@ -831,12 +831,12 @@ UE_NET_TEST_FIXTURE(FSplitObjectTestFixture, TestCancelPendingDestroyDuringHugeO
 	}
 
 	// Filter out object to cause a PendingDestroy
-	Server->GetReplicationSystem()->AddToGroup(NotReplicatedNetObjectGroupHandle, ServerObject->NetHandle);
+	Server->GetReplicationSystem()->AddToGroup(NotReplicatedNetObjectGroupHandle, ServerObject->NetRefHandle);
 	Server->PreSendUpdate();
 	Server->PostSendUpdate();
 
 	// Remove object from filter to cause object to end up in CancelPendingDestroy
-	Server->GetReplicationSystem()->RemoveFromGroup(NotReplicatedNetObjectGroupHandle, ServerObject->NetHandle);
+	Server->GetReplicationSystem()->RemoveFromGroup(NotReplicatedNetObjectGroupHandle, ServerObject->NetRefHandle);
 	Server->PreSendUpdate();
 	Server->PostSendUpdate();
 
@@ -859,7 +859,7 @@ UE_NET_TEST_FIXTURE(FSplitObjectTestFixture, TestCancelPendingDestroyDuringHugeO
 	Server->SendAndDeliverTo(Client, DeliverPacket);
 	Server->PostSendUpdate();
 
-	UTestReplicatedIrisObject* ClientObject = Cast<UTestReplicatedIrisObject>(Client->GetReplicationBridge()->GetReplicatedObject(ServerObject->NetHandle));
+	UTestReplicatedIrisObject* ClientObject = Cast<UTestReplicatedIrisObject>(Client->GetReplicationBridge()->GetReplicatedObject(ServerObject->NetRefHandle));
 	UE_NET_ASSERT_NE(ClientObject, nullptr);
 	UE_NET_ASSERT_EQ(ClientObject->IntA, ServerObject->IntA);
 }
diff --git a/Engine/Plugins/Runtime/ReplicationSystemTestPlugin/Source/Private/Tests/ReplicationSystem/TestRepTag.cpp b/Engine/Plugins/Runtime/ReplicationSystemTestPlugin/Source/Private/Tests/ReplicationSystem/TestRepTag.cpp
--- a/Engine/Plugins/Runtime/ReplicationSystemTestPlugin/Source/Private/Tests/ReplicationSystem/TestRepTag.cpp
+++ b/Engine/Plugins/Runtime/ReplicationSystemTestPlugin/Source/Private/Tests/ReplicationSystem/TestRepTag.cpp
@@ -24,8 +24,8 @@ private:
 			constexpr uint32 PropertyComponentCount = 1;
 			constexpr uint32 IrisComponentCount = 2;
 			UTestReplicatedIrisObject* TestObject = CreateObject(PropertyComponentCount, IrisComponentCount);
-			const FNetHandle NetHandle = ReplicationBridge->BeginReplication(TestObject);
-			ReplicationProtocol = ReplicationSystem->GetReplicationProtocol(NetHandle);
+			const FNetRefHandle NetRefHandle = ReplicationBridge->BeginReplication(TestObject);
+			ReplicationProtocol = ReplicationSystem->GetReplicationProtocol(NetRefHandle);
 
 			StateDescriptorWithTag = PropertyComponentCount + 2U;
 			RepTagToFind = ReplicationProtocol->ReplicationStateDescriptors[StateDescriptorWithTag]->MemberTagDescriptors[0].Tag;
diff --git a/Engine/Plugins/Runtime/ReplicationSystemTestPlugin/Source/Private/Tests/ReplicationSystem/TestReplicationBridge.cpp b/Engine/Plugins/Runtime/ReplicationSystemTestPlugin/Source/Private/Tests/ReplicationSystem/TestReplicationBridge.cpp
--- a/Engine/Plugins/Runtime/ReplicationSystemTestPlugin/Source/Private/Tests/ReplicationSystem/TestReplicationBridge.cpp
+++ b/Engine/Plugins/Runtime/ReplicationSystemTestPlugin/Source/Private/Tests/ReplicationSystem/TestReplicationBridge.cpp
@@ -20,16 +20,16 @@ UE_NET_TEST_FIXTURE(FReplicationBridgeTestFixture, CreateObjectWithSingleFragmen
 {
 	UTestReplicatedIrisObject* TestObject = CreateObject(0, 0);
 
-	// Create NetHandle for the CreatedHandle,
-	FNetHandle CreatedHandle = ReplicationBridge->BeginReplication(TestObject);
+	// Start replicating the object.
+	FNetRefHandle CreatedHandle = ReplicationBridge->BeginReplication(TestObject);
 
-	// Verify that we managed to create a valid NetHandle	
+	// Verify that we managed to create a valid handle	
 	UE_NET_ASSERT_TRUE(CreatedHandle.IsValid());
-	UE_NET_ASSERT_TRUE(CreatedHandle == TestObject->NetHandle);
-	UE_NET_ASSERT_TRUE(CreatedHandle == ReplicationBridge->GetReplicatedHandle(TestObject));
+	UE_NET_ASSERT_TRUE(CreatedHandle == TestObject->NetRefHandle);
+	UE_NET_ASSERT_TRUE(CreatedHandle == ReplicationBridge->GetReplicatedRefHandle(TestObject));
 
 	// Verify ReplicationProtocol
-	const FReplicationProtocol* Protocol = ReplicationSystem->GetReplicationProtocol(TestObject->NetHandle);
+	const FReplicationProtocol* Protocol = ReplicationSystem->GetReplicationProtocol(TestObject->NetRefHandle);
 	UE_NET_ASSERT_TRUE(Protocol != nullptr);
 	UE_NET_ASSERT_TRUE(Protocol->ProtocolIdentifier != 0);
 	UE_NET_ASSERT_EQ(1u, Protocol->ReplicationStateCount);
@@ -49,16 +49,16 @@ UE_NET_TEST_FIXTURE(FReplicationBridgeTestFixture, CreateObjectWithMultipleFragm
 	constexpr uint32 IrisComponentCount = 0;
 	UTestReplicatedIrisObject* TestObject = CreateObject(PropertyComponentCount, IrisComponentCount);
 
-	// Create NetHandle for the CreatedHandle,
-	FNetHandle CreatedHandle = ReplicationBridge->BeginReplication(TestObject);
+	// Start replicating the object.
+	FNetRefHandle CreatedHandle = ReplicationBridge->BeginReplication(TestObject);
 
-	// Verify that we managed to create a valid NetHandle	
+	// Verify that we managed to create a valid handle	
 	UE_NET_ASSERT_TRUE(CreatedHandle.IsValid());
-	UE_NET_ASSERT_TRUE(CreatedHandle == TestObject->NetHandle);
-	UE_NET_ASSERT_TRUE(CreatedHandle == ReplicationBridge->GetReplicatedHandle(TestObject));
+	UE_NET_ASSERT_TRUE(CreatedHandle == TestObject->NetRefHandle);
+	UE_NET_ASSERT_TRUE(CreatedHandle == ReplicationBridge->GetReplicatedRefHandle(TestObject));
 
 	// Verify ReplicationProtocol
-	const FReplicationProtocol* Protocol = ReplicationSystem->GetReplicationProtocol(TestObject->NetHandle);
+	const FReplicationProtocol* Protocol = ReplicationSystem->GetReplicationProtocol(TestObject->NetRefHandle);
 	UE_NET_ASSERT_TRUE(Protocol != nullptr);
 	UE_NET_ASSERT_EQ(25u, Protocol->ReplicationStateCount);
 	UE_NET_ASSERT_EQ(8u, Protocol->InternalTotalAlignment);
@@ -77,16 +77,16 @@ UE_NET_TEST_FIXTURE(FReplicationBridgeTestFixture, CreateObjectWithMixedFragment
 	constexpr uint32 IrisComponentCount = 3;
 	UTestReplicatedIrisObject* TestObject = CreateObject(PropertyComponentCount, IrisComponentCount);
 
-	// Create NetHandle for the CreatedHandle,
-	FNetHandle CreatedHandle = ReplicationBridge->BeginReplication(TestObject);
+	// Start replicating the object.
+	FNetRefHandle CreatedHandle = ReplicationBridge->BeginReplication(TestObject);
 
-	// Verify that we managed to create a valid NetHandle	
+	// Verify that we managed to create a valid handle	
 	UE_NET_ASSERT_TRUE(CreatedHandle.IsValid());
-	UE_NET_ASSERT_TRUE(CreatedHandle == TestObject->NetHandle);
-	UE_NET_ASSERT_TRUE(CreatedHandle == ReplicationBridge->GetReplicatedHandle(TestObject));
+	UE_NET_ASSERT_TRUE(CreatedHandle == TestObject->NetRefHandle);
+	UE_NET_ASSERT_TRUE(CreatedHandle == ReplicationBridge->GetReplicatedRefHandle(TestObject));
 
 	// Verify ReplicationProtocol
-	const FReplicationProtocol* Protocol = ReplicationSystem->GetReplicationProtocol(TestObject->NetHandle);;
+	const FReplicationProtocol* Protocol = ReplicationSystem->GetReplicationProtocol(TestObject->NetRefHandle);
 	UE_NET_ASSERT_TRUE(Protocol != nullptr);
 	UE_NET_ASSERT_EQ(28u, Protocol->ReplicationStateCount);
 	UE_NET_ASSERT_EQ(8u, Protocol->InternalTotalAlignment);
@@ -104,16 +104,16 @@ UE_NET_TEST_FIXTURE(FReplicationBridgeTestFixture, CreateObjectWithDynamicStateC
 	constexpr uint32 DynamicStateComponentCount = 1U;
 	UTestReplicatedIrisObject* TestObject = CreateObjectWithDynamicState(0, 0, DynamicStateComponentCount);
 
-	// Create NetHandle for the CreatedHandle,
-	FNetHandle CreatedHandle = ReplicationBridge->BeginReplication(TestObject);
+	// Start replicating the object.
+	FNetRefHandle CreatedHandle = ReplicationBridge->BeginReplication(TestObject);
 
-	// Verify that we managed to create a valid NetHandle	
+	// Verify that we managed to create a valid handle	
 	UE_NET_ASSERT_TRUE(CreatedHandle.IsValid());
-	UE_NET_ASSERT_TRUE(CreatedHandle == TestObject->NetHandle);
-	UE_NET_ASSERT_TRUE(CreatedHandle == ReplicationBridge->GetReplicatedHandle(TestObject));
+	UE_NET_ASSERT_TRUE(CreatedHandle == TestObject->NetRefHandle);
+	UE_NET_ASSERT_TRUE(CreatedHandle == ReplicationBridge->GetReplicatedRefHandle(TestObject));
 
 	// Verify ReplicationProtocol
-	const FReplicationProtocol* Protocol = ReplicationSystem->GetReplicationProtocol(TestObject->NetHandle);
+	const FReplicationProtocol* Protocol = ReplicationSystem->GetReplicationProtocol(TestObject->NetRefHandle);
 	UE_NET_ASSERT_NE(Protocol, nullptr);
 	UE_NET_ASSERT_NE(Protocol->ProtocolIdentifier, FReplicationProtocolIdentifier(0));
 	UE_NET_ASSERT_EQ(Protocol->ReplicationStateCount, 1U + DynamicStateComponentCount);
diff --git a/Engine/Plugins/Runtime/ReplicationSystemTestPlugin/Source/Private/Tests/ReplicationSystem/TestReplicationOperations.cpp b/Engine/Plugins/Runtime/ReplicationSystemTestPlugin/Source/Private/Tests/ReplicationSystem/TestReplicationOperations.cpp
--- a/Engine/Plugins/Runtime/ReplicationSystemTestPlugin/Source/Private/Tests/ReplicationSystem/TestReplicationOperations.cpp
+++ b/Engine/Plugins/Runtime/ReplicationSystemTestPlugin/Source/Private/Tests/ReplicationSystem/TestReplicationOperations.cpp
@@ -68,8 +68,8 @@ UE_NET_TEST_FIXTURE(FTestReplicationOperationsFixture, CanPollAndRefreshProperty
 
 	UTestReplicatedIrisObject* TestObject = CreateObjectWithDynamicState(PropertyComponentCount, IrisComponentCount, DynamicStateComponentCount);
 
-	// Create NetHandle for the CreatedHandle,
-	FNetHandle CreatedHandle = ReplicationBridge->BeginReplication(TestObject);
+	// Create NetRefHandle for the CreatedHandle,
+	FNetRefHandle CreatedHandle = ReplicationBridge->BeginReplication(TestObject);
 
 	auto&& GetPropertyReplicationStateChecked = [](const FReplicationFragment* Fragment)
 	{ 
@@ -108,7 +108,7 @@ UE_NET_TEST_FIXTURE(FTestReplicationOperationsFixture, CanPollAndRefreshProperty
 	TestObject->DynamicStateComponents[DynamicStateComponentIndex]->IntStaticArray[5] = 43;
 
 	// Poll data to update property replication state
-	FReplicationInstanceOperations::PollAndRefreshCachedPropertyData(ReplicationBridge->GetReplicationInstanceProtocol(TestObject->NetHandle));
+	FReplicationInstanceOperations::PollAndRefreshCachedPropertyData(ReplicationBridge->GetReplicationInstanceProtocol(TestObject->NetRefHandle));
 
 	// Verify values after poll
 	VerifyAs<int32>(1, State, 0);
@@ -136,8 +136,8 @@ UE_NET_TEST_FIXTURE(FTestReplicationOperationsFixture, CanQuantizeAndDequantize)
 
 	UTestReplicatedIrisObject* TestObject = CreateObjectWithDynamicState(PropertyComponentCount, IrisComponentCount, DynamicStateComponentCount);
 	
-	// Create NetHandle for the CreatedHandle,
-	FNetHandle CreatedHandle = ReplicationBridge->BeginReplication(TestObject);
+	// Create NetRefHandle for the CreatedHandle,
+	FNetRefHandle CreatedHandle = ReplicationBridge->BeginReplication(TestObject);
 	const FReplicationProtocol* ProtocolA = ReplicationSystem->GetReplicationProtocol(CreatedHandle);
 
 	// Set some values
@@ -156,7 +156,7 @@ UE_NET_TEST_FIXTURE(FTestReplicationOperationsFixture, CanQuantizeAndDequantize)
 	TestObject->DynamicStateComponents[DynamicStateComponentCount - 1]->IntStaticArray[4] = 127;
 
 	// Poll data to update property replication state
-	FReplicationInstanceOperations::PollAndRefreshCachedPropertyData(ReplicationBridge->GetReplicationInstanceProtocol(TestObject->NetHandle));
+	FReplicationInstanceOperations::PollAndRefreshCachedPropertyData(ReplicationBridge->GetReplicationInstanceProtocol(TestObject->NetRefHandle));
 
 	// Quantize to buffer
 	const uint32 BufferSize = 1024;
@@ -177,7 +177,7 @@ UE_NET_TEST_FIXTURE(FTestReplicationOperationsFixture, CanQuantizeAndDequantize)
 	SerializationContext.SetIsInitState(true);
 
 	// Copy data from test object to buffer
-	FReplicationInstanceOperations::CopyAndQuantize(SerializationContext, AlignedBuffer, GetChangeMaskWriter(ProtocolA->ChangeMaskBitCount), ReplicationBridge->GetReplicationInstanceProtocol(TestObject->NetHandle), ProtocolA);
+	FReplicationInstanceOperations::CopyAndQuantize(SerializationContext, AlignedBuffer, GetChangeMaskWriter(ProtocolA->ChangeMaskBitCount), ReplicationBridge->GetReplicationInstanceProtocol(TestObject->NetRefHandle), ProtocolA);
 
 	// Verify that we did not overshoot
 	UE_NET_ASSERT_EQ(ValueToSet, AlignedBuffer[ProtocolA->InternalTotalSize]);
@@ -185,11 +185,11 @@ UE_NET_TEST_FIXTURE(FTestReplicationOperationsFixture, CanQuantizeAndDequantize)
 	// Create second object 
 	UTestReplicatedIrisObject* TestObjectB = CreateObjectWithDynamicState(PropertyComponentCount, IrisComponentCount, DynamicStateComponentCount);
 	ReplicationBridge->BeginReplication(TestObjectB);
-	const FReplicationProtocol* ProtocolB = ReplicationSystem->GetReplicationProtocol(TestObjectB->NetHandle);
+	const FReplicationProtocol* ProtocolB = ReplicationSystem->GetReplicationProtocol(TestObjectB->NetRefHandle);
 	UE_NET_ASSERT_EQ(ProtocolA, ProtocolB);
 
 	// Push state data to ObjectB
-	FReplicationInstanceOperations::DequantizeAndApply(SerializationContext, GetTempAllocator(), GetChangeMaskStorage(), ReplicationBridge->GetReplicationInstanceProtocol(TestObjectB->NetHandle), AlignedBuffer, ProtocolB);
+	FReplicationInstanceOperations::DequantizeAndApply(SerializationContext, GetTempAllocator(), GetChangeMaskStorage(), ReplicationBridge->GetReplicationInstanceProtocol(TestObjectB->NetRefHandle), AlignedBuffer, ProtocolB);
 
 	// check if we did get the expected values
 	UE_NET_ASSERT_EQ(1, TestObjectB->IntA);
@@ -233,8 +233,8 @@ UE_NET_TEST_FIXTURE(FTestReplicationOperationsFixture, CanSerializeAndDeserializ
 
 	UTestReplicatedIrisObject* TestObject = CreateObjectWithDynamicState(PropertyComponentCount, IrisComponentCount, DynamicStateComponentCount);
 	
-	// Create NetHandle for the CreatedHandle,
-	FNetHandle CreatedHandle = ReplicationBridge->BeginReplication(TestObject);
+	// Create NetRefHandle for the CreatedHandle,
+	FNetRefHandle CreatedHandle = ReplicationBridge->BeginReplication(TestObject);
 	const FReplicationProtocol* ProtocolA = ReplicationSystem->GetReplicationProtocol(CreatedHandle);
 
 	// Set some values
@@ -251,7 +251,7 @@ UE_NET_TEST_FIXTURE(FTestReplicationOperationsFixture, CanSerializeAndDeserializ
 	TestObject->DynamicStateComponents[DynamicStateComponentCount - 1]->IntStaticArray[4] = 127;
 
 	// Poll data to update property replication state
-	FReplicationInstanceOperations::PollAndRefreshCachedPropertyData(ReplicationBridge->GetReplicationInstanceProtocol(TestObject->NetHandle));
+	FReplicationInstanceOperations::PollAndRefreshCachedPropertyData(ReplicationBridge->GetReplicationInstanceProtocol(TestObject->NetRefHandle));
 
 	// Quantize to buffer
 	const uint32 BufferSize = 1024;
@@ -271,7 +271,7 @@ UE_NET_TEST_FIXTURE(FTestReplicationOperationsFixture, CanSerializeAndDeserializ
 	FInternalNetSerializationContext InternalContext;
 	SerializationContext.SetInternalContext(&InternalContext);
 	SerializationContext.SetIsInitState(true);
-	FReplicationInstanceOperations::CopyAndQuantize(SerializationContext, AlignedBuffer, GetChangeMaskWriter(ProtocolA->ChangeMaskBitCount), ReplicationBridge->GetReplicationInstanceProtocol(TestObject->NetHandle), ProtocolA);
+	FReplicationInstanceOperations::CopyAndQuantize(SerializationContext, AlignedBuffer, GetChangeMaskWriter(ProtocolA->ChangeMaskBitCount), ReplicationBridge->GetReplicationInstanceProtocol(TestObject->NetRefHandle), ProtocolA);
 
 	// Verify that we did not overshoot
 	UE_NET_ASSERT_EQ(ValueToSet, AlignedBuffer[ProtocolA->InternalTotalSize]);
@@ -279,7 +279,7 @@ UE_NET_TEST_FIXTURE(FTestReplicationOperationsFixture, CanSerializeAndDeserializ
 	// Create second object 
 	UTestReplicatedIrisObject* TestObjectB = CreateObjectWithDynamicState(PropertyComponentCount, IrisComponentCount, DynamicStateComponentCount);
 	ReplicationBridge->BeginReplication(TestObjectB);
-	const FReplicationProtocol* ProtocolB = ReplicationSystem->GetReplicationProtocol(TestObjectB->NetHandle);
+	const FReplicationProtocol* ProtocolB = ReplicationSystem->GetReplicationProtocol(TestObjectB->NetRefHandle);
 	UE_NET_ASSERT_EQ(ProtocolA, ProtocolB);
 
 
@@ -318,7 +318,7 @@ UE_NET_TEST_FIXTURE(FTestReplicationOperationsFixture, CanSerializeAndDeserializ
 		FReplicationProtocolOperations::DeserializeWithMask(Context, GetChangeMaskStorage(), AlignedBufferB, ProtocolB);
 
 		// Push deserialized state data to ObjectB 
-		FReplicationInstanceOperations::DequantizeAndApply(Context, GetTempAllocator(), GetChangeMaskStorage(), ReplicationBridge->GetReplicationInstanceProtocol(TestObjectB->NetHandle), AlignedBufferB, ProtocolB);
+		FReplicationInstanceOperations::DequantizeAndApply(Context, GetTempAllocator(), GetChangeMaskStorage(), ReplicationBridge->GetReplicationInstanceProtocol(TestObjectB->NetRefHandle), AlignedBufferB, ProtocolB);
 	}
 
 	// check if we did get the expected values
@@ -351,8 +351,8 @@ UE_NET_TEST_FIXTURE(FTestReplicationOperationsFixture, PreAndPostNetReceivedOnly
 {
 	UTestReplicatedIrisObject* TestObject = CreateObject(1, 0);
 	
-	// Create NetHandle for the CreatedHandle,
-	FNetHandle CreatedHandle = ReplicationBridge->BeginReplication(TestObject);
+	// Create NetRefHandle for the CreatedHandle,
+	FNetRefHandle CreatedHandle = ReplicationBridge->BeginReplication(TestObject);
 	const FReplicationProtocol* ProtocolA = ReplicationSystem->GetReplicationProtocol(CreatedHandle);
 
 	UE_NET_ASSERT_EQ(3U, ProtocolA->ReplicationStateCount);
@@ -362,7 +362,7 @@ UE_NET_TEST_FIXTURE(FTestReplicationOperationsFixture, PreAndPostNetReceivedOnly
 	TestObject->Components[0]->IntB = 3;
 
 	// Poll data to update property replication state
-	FReplicationInstanceOperations::PollAndRefreshCachedPropertyData(ReplicationBridge->GetReplicationInstanceProtocol(TestObject->NetHandle));
+	FReplicationInstanceOperations::PollAndRefreshCachedPropertyData(ReplicationBridge->GetReplicationInstanceProtocol(TestObject->NetRefHandle));
 
 	// Quantize to buffer
 	uint8* AlignedBuffer = (uint8*)(GetTempAllocator()).Alloc(ProtocolA->InternalTotalSize, ProtocolA->InternalTotalAlignment);
@@ -370,12 +370,12 @@ UE_NET_TEST_FIXTURE(FTestReplicationOperationsFixture, PreAndPostNetReceivedOnly
 	// Copy data from test object to buffer
 	FNetSerializationContext NetContext;
 	NetContext.SetIsInitState(true);
-	FReplicationInstanceOperations::CopyAndQuantize(NetContext, AlignedBuffer, GetChangeMaskWriter(ProtocolA->ChangeMaskBitCount), ReplicationBridge->GetReplicationInstanceProtocol(TestObject->NetHandle), ProtocolA);
+	FReplicationInstanceOperations::CopyAndQuantize(NetContext, AlignedBuffer, GetChangeMaskWriter(ProtocolA->ChangeMaskBitCount), ReplicationBridge->GetReplicationInstanceProtocol(TestObject->NetRefHandle), ProtocolA);
 
 	// Create second object 
 	UTestReplicatedIrisObject* TestObjectB = CreateObject(1, 0);
 	ReplicationBridge->BeginReplication(TestObjectB);
-	const FReplicationProtocol* ProtocolB = ReplicationSystem->GetReplicationProtocol(TestObjectB->NetHandle);
+	const FReplicationProtocol* ProtocolB = ReplicationSystem->GetReplicationProtocol(TestObjectB->NetRefHandle);
 	UE_NET_ASSERT_EQ(ProtocolA, ProtocolB);
 
 	// Check values before we try to apply data
@@ -383,7 +383,7 @@ UE_NET_TEST_FIXTURE(FTestReplicationOperationsFixture, PreAndPostNetReceivedOnly
 	UE_NET_ASSERT_EQ(0U, TestObjectB->Components[0]->CallCounts.PostNetReceiveCounter);
 
 	// Push state data to ObjectB
-	FReplicationInstanceOperations::DequantizeAndApply(NetContext, GetTempAllocator(), GetChangeMaskStorage(), ReplicationBridge->GetReplicationInstanceProtocol(TestObjectB->NetHandle), AlignedBuffer, ProtocolB);
+	FReplicationInstanceOperations::DequantizeAndApply(NetContext, GetTempAllocator(), GetChangeMaskStorage(), ReplicationBridge->GetReplicationInstanceProtocol(TestObjectB->NetRefHandle), AlignedBuffer, ProtocolB);
 
 	// Check value after apply
 	UE_NET_ASSERT_EQ(1U, TestObjectB->Components[0]->CallCounts.PreNetReceiveCounter);
@@ -398,15 +398,15 @@ UE_NET_TEST_FIXTURE(FTestReplicationOperationsFixture, RepNotifyIsCalledForSimpl
 	constexpr uint32 PropertyComponentCount = 1;
 	UTestReplicatedIrisObject* TestObject = CreateObject(PropertyComponentCount, 0);
 	
-	// Create NetHandle for the CreatedHandle,
-	FNetHandle CreatedHandle = ReplicationBridge->BeginReplication(TestObject);
+	// Create NetRefHandle for the CreatedHandle,
+	FNetRefHandle CreatedHandle = ReplicationBridge->BeginReplication(TestObject);
 	const FReplicationProtocol* Protocol = ReplicationSystem->GetReplicationProtocol(CreatedHandle);
 
 	// Set some values
 	TestObject->Components[0]->IntB ^= 1;
 
 	// Poll data to update property replication state
-	FReplicationInstanceOperations::PollAndRefreshCachedPropertyData(ReplicationBridge->GetReplicationInstanceProtocol(TestObject->NetHandle));
+	FReplicationInstanceOperations::PollAndRefreshCachedPropertyData(ReplicationBridge->GetReplicationInstanceProtocol(TestObject->NetRefHandle));
 
 	// Quantize to buffer
 	uint8* AlignedBuffer = (uint8*)(GetTempAllocator()).Alloc(Protocol->InternalTotalSize, Protocol->InternalTotalAlignment);
@@ -414,7 +414,7 @@ UE_NET_TEST_FIXTURE(FTestReplicationOperationsFixture, RepNotifyIsCalledForSimpl
 	// Copy data from test object to buffer
 	FNetSerializationContext NetContext;
 	NetContext.SetIsInitState(true);
-	FReplicationInstanceOperations::CopyAndQuantize(NetContext, AlignedBuffer, GetChangeMaskWriter(Protocol->ChangeMaskBitCount), ReplicationBridge->GetReplicationInstanceProtocol(TestObject->NetHandle), Protocol);
+	FReplicationInstanceOperations::CopyAndQuantize(NetContext, AlignedBuffer, GetChangeMaskWriter(Protocol->ChangeMaskBitCount), ReplicationBridge->GetReplicationInstanceProtocol(TestObject->NetRefHandle), Protocol);
 
 	// Create second object 
 	UTestReplicatedIrisObject* TestObjectB = CreateObject(PropertyComponentCount, 0);
@@ -424,7 +424,7 @@ UE_NET_TEST_FIXTURE(FTestReplicationOperationsFixture, RepNotifyIsCalledForSimpl
 	UE_NET_ASSERT_EQ(TestObjectB->Components[0]->CallCounts.IntBRepNotifyCounter, 0U);
 
 	// Push state data to ObjectB
-	FReplicationInstanceOperations::DequantizeAndApply(NetContext, GetTempAllocator(), GetChangeMaskStorage(), ReplicationBridge->GetReplicationInstanceProtocol(TestObjectB->NetHandle), AlignedBuffer, Protocol);
+	FReplicationInstanceOperations::DequantizeAndApply(NetContext, GetTempAllocator(), GetChangeMaskStorage(), ReplicationBridge->GetReplicationInstanceProtocol(TestObjectB->NetRefHandle), AlignedBuffer, Protocol);
 
 	// Check that we got exactly one RepNotify call
 	UE_NET_ASSERT_EQ(TestObjectB->Components[0]->CallCounts.IntBRepNotifyCounter, 1U);
@@ -438,15 +438,15 @@ UE_NET_TEST_FIXTURE(FTestReplicationOperationsFixture, InitPropertyIsNotUpdatedW
 	constexpr uint32 PropertyComponentCount = 1;
 	UTestReplicatedIrisObject* TestObject = CreateObject(PropertyComponentCount, 0);
 	
-	// Create NetHandle for the CreatedHandle,
-	FNetHandle CreatedHandle = ReplicationBridge->BeginReplication(TestObject);
+	// Create NetRefHandle for the CreatedHandle,
+	FNetRefHandle CreatedHandle = ReplicationBridge->BeginReplication(TestObject);
 	const FReplicationProtocol* Protocol = ReplicationSystem->GetReplicationProtocol(CreatedHandle);
 
 	// Set some values
 	TestObject->Components[0]->IntB ^= 1;
 
 	// Poll data to update property replication state
-	FReplicationInstanceOperations::PollAndRefreshCachedPropertyData(ReplicationBridge->GetReplicationInstanceProtocol(TestObject->NetHandle));
+	FReplicationInstanceOperations::PollAndRefreshCachedPropertyData(ReplicationBridge->GetReplicationInstanceProtocol(TestObject->NetRefHandle));
 
 	// Quantize to buffer
 	uint8* AlignedBuffer = (uint8*)(GetTempAllocator()).Alloc(Protocol->InternalTotalSize, Protocol->InternalTotalAlignment);
@@ -454,7 +454,7 @@ UE_NET_TEST_FIXTURE(FTestReplicationOperationsFixture, InitPropertyIsNotUpdatedW
 	// Copy data from test object to buffer
 	FNetSerializationContext NetContext;
 	NetContext.SetIsInitState(false);
-	FReplicationInstanceOperations::CopyAndQuantize(NetContext, AlignedBuffer, GetChangeMaskWriter(Protocol->ChangeMaskBitCount), ReplicationBridge->GetReplicationInstanceProtocol(TestObject->NetHandle), Protocol);
+	FReplicationInstanceOperations::CopyAndQuantize(NetContext, AlignedBuffer, GetChangeMaskWriter(Protocol->ChangeMaskBitCount), ReplicationBridge->GetReplicationInstanceProtocol(TestObject->NetRefHandle), Protocol);
 
 	// Create second object 
 	UTestReplicatedIrisObject* TestObjectB = CreateObject(PropertyComponentCount, 0);
@@ -464,7 +464,7 @@ UE_NET_TEST_FIXTURE(FTestReplicationOperationsFixture, InitPropertyIsNotUpdatedW
 	UE_NET_ASSERT_EQ(TestObjectB->Components[0]->CallCounts.IntBRepNotifyCounter, 0U);
 
 	// Push state data to ObjectB
-	FReplicationInstanceOperations::DequantizeAndApply(NetContext, GetTempAllocator(), GetChangeMaskStorage(), ReplicationBridge->GetReplicationInstanceProtocol(TestObjectB->NetHandle), AlignedBuffer, Protocol);
+	FReplicationInstanceOperations::DequantizeAndApply(NetContext, GetTempAllocator(), GetChangeMaskStorage(), ReplicationBridge->GetReplicationInstanceProtocol(TestObjectB->NetRefHandle), AlignedBuffer, Protocol);
 
 	// Check that the IntB property has not been modified
 	UE_NET_ASSERT_EQ(TestObjectB->Components[0]->IntB, 0);
@@ -481,8 +481,8 @@ UE_NET_TEST_FIXTURE(FTestReplicationOperationsFixture, RepNotifyIsCalledExactlyO
 	constexpr uint32 DynamicStateComponentCount = 1;
 	UTestReplicatedIrisObject* TestObject = CreateObjectWithDynamicState(0, 0, DynamicStateComponentCount);
 	
-	// Create NetHandle for the TestObject
-	FNetHandle CreatedHandle = ReplicationBridge->BeginReplication(TestObject);
+	// Create NetRefHandle for the TestObject
+	FNetRefHandle CreatedHandle = ReplicationBridge->BeginReplication(TestObject);
 	const FReplicationProtocol* Protocol = ReplicationSystem->GetReplicationProtocol(CreatedHandle);
 
 	// Set some values
@@ -496,7 +496,7 @@ UE_NET_TEST_FIXTURE(FTestReplicationOperationsFixture, RepNotifyIsCalledExactlyO
 	TestObject->DynamicStateComponents[0]->IntStaticArray[5] = 1;
 
 	// Poll data to update property replication state
-	FReplicationInstanceOperations::PollAndRefreshCachedPropertyData(ReplicationBridge->GetReplicationInstanceProtocol(TestObject->NetHandle));
+	FReplicationInstanceOperations::PollAndRefreshCachedPropertyData(ReplicationBridge->GetReplicationInstanceProtocol(TestObject->NetRefHandle));
 
 	// Quantize to buffer
 	uint8* AlignedBuffer = (uint8*)(GetTempAllocator()).Alloc(Protocol->InternalTotalSize, Protocol->InternalTotalAlignment);
@@ -507,7 +507,7 @@ UE_NET_TEST_FIXTURE(FTestReplicationOperationsFixture, RepNotifyIsCalledExactlyO
 	FInternalNetSerializationContext InternalContext;
 	SerializationContext.SetInternalContext(&InternalContext);
 	SerializationContext.SetIsInitState(true);
-	FReplicationInstanceOperations::CopyAndQuantize(SerializationContext, AlignedBuffer, GetChangeMaskWriter(Protocol->ChangeMaskBitCount), ReplicationBridge->GetReplicationInstanceProtocol(TestObject->NetHandle), Protocol);
+	FReplicationInstanceOperations::CopyAndQuantize(SerializationContext, AlignedBuffer, GetChangeMaskWriter(Protocol->ChangeMaskBitCount), ReplicationBridge->GetReplicationInstanceProtocol(TestObject->NetRefHandle), Protocol);
 
 	// Create second object 
 	UTestReplicatedIrisObject* TestObjectB = CreateObjectWithDynamicState(0, 0, DynamicStateComponentCount);
@@ -518,7 +518,7 @@ UE_NET_TEST_FIXTURE(FTestReplicationOperationsFixture, RepNotifyIsCalledExactlyO
 	UE_NET_ASSERT_EQ(TestObjectB->DynamicStateComponents[0]->CallCounts.IntStaticArrayRepNotifyCounter, 0U);
 
 	// Push state data to ObjectB
-	FReplicationInstanceOperations::DequantizeAndApply(SerializationContext, GetTempAllocator(), GetChangeMaskStorage(), ReplicationBridge->GetReplicationInstanceProtocol(TestObjectB->NetHandle), AlignedBuffer, Protocol);
+	FReplicationInstanceOperations::DequantizeAndApply(SerializationContext, GetTempAllocator(), GetChangeMaskStorage(), ReplicationBridge->GetReplicationInstanceProtocol(TestObjectB->NetRefHandle), AlignedBuffer, Protocol);
 
 	// Check that we got exactly one RepNotify call per array
 	UE_NET_ASSERT_EQ(TestObjectB->DynamicStateComponents[0]->CallCounts.IntArrayRepNotifyCounter, 1U);
@@ -559,13 +559,13 @@ UE_NET_TEST_FIXTURE(FTestReplicationOperationsForObjectsFixture, StaleObjectPoin
 {
 	// Create two objects, one which is referencing the other. Verify reference gets updated when first object is destroyed.
 	UTestReplicatedIrisObject* Object0 = CreateObject();
-	FNetHandle HandleToObject0 = ReplicationBridge->BeginReplication(Object0);
+	FNetRefHandle HandleToObject0 = ReplicationBridge->BeginReplication(Object0);
 
 	UTestReplicatedIrisObject* Object1 = CreateObject();
 	constexpr uint32 ObjectReferenceComponentCount = 1;
 	Object1->AddComponents(UTestReplicatedIrisObject::FComponents{0,0,0,0,ObjectReferenceComponentCount});
 	Object1->ObjectReferenceComponents[0]->RawObjectPtrRef = Object0;
-	FNetHandle HandleToObject1 = ReplicationBridge->BeginReplication(Object1);
+	FNetRefHandle HandleToObject1 = ReplicationBridge->BeginReplication(Object1);
 
 	const FReplicationInstanceProtocol* InstanceProtocol1 = ReplicationBridge->GetReplicationInstanceProtocol(HandleToObject1);
 	const FReplicationProtocol* Protocol1 = ReplicationSystem->GetReplicationProtocol(HandleToObject1);
@@ -626,7 +626,7 @@ UE_NET_TEST_FIXTURE(FTestReplicationOperationsOnInitStateFixture, StateWithInitT
 	constexpr uint32 PropertyComponentCount = 2;
 	UTestReplicatedIrisObject* TestObject = CreateObject(PropertyComponentCount, 0);
 
-	const FNetHandle Handle = ReplicationBridge->BeginReplication(TestObject);
+	const FNetRefHandle Handle = ReplicationBridge->BeginReplication(TestObject);
 	const FReplicationProtocol* Protocol = ReplicationSystem->GetReplicationProtocol(Handle);
 	uint32 InitStateCount = 0;
 	for (const FReplicationStateDescriptor* StateDescriptor : MakeArrayView(Protocol->ReplicationStateDescriptors, Protocol->ReplicationStateCount))
@@ -642,7 +642,7 @@ UE_NET_TEST_FIXTURE(FTestReplicationOperationsOnInitStateFixture, InitStateHasNo
 	constexpr uint32 PropertyComponentCount = 1;
 	UTestReplicatedIrisObject* TestObject = CreateObject(PropertyComponentCount, 0);
 
-	const FNetHandle Handle = ReplicationBridge->BeginReplication(TestObject);
+	const FNetRefHandle Handle = ReplicationBridge->BeginReplication(TestObject);
 	const FReplicationProtocol* Protocol = ReplicationSystem->GetReplicationProtocol(Handle);
 	const FReplicationStateDescriptor* InitStateDescriptor = GetFirstInitStateDescriptor(Protocol);
 
@@ -663,7 +663,7 @@ UE_NET_TEST_FIXTURE(FTestReplicationOperationsOnConnectionFilterStateFixture, Pr
 	UTestReplicatedIrisObject* TestObject = CreateObject();
 	TestObject->AddComponents(Components);
 
-	const FNetHandle Handle = ReplicationBridge->BeginReplication(TestObject);
+	const FNetRefHandle Handle = ReplicationBridge->BeginReplication(TestObject);
 	const FReplicationProtocol* Protocol = ReplicationSystem->GetReplicationProtocol(Handle);
 
 	UE_NET_ASSERT_TRUE(EnumHasAnyFlags(Protocol->ProtocolTraits, EReplicationProtocolTraits::HasLifetimeConditionals));
@@ -678,7 +678,7 @@ UE_NET_TEST_FIXTURE(FTestReplicationOperationsOnConnectionFilterStateFixture, St
 	UTestReplicatedIrisObject* TestObject = CreateObject();
 	TestObject->AddComponents(Components);
 
-	const FNetHandle Handle = ReplicationBridge->BeginReplication(TestObject);
+	const FNetRefHandle Handle = ReplicationBridge->BeginReplication(TestObject);
 	const FReplicationProtocol* Protocol = ReplicationSystem->GetReplicationProtocol(Handle);
 	uint32 LifetimeConditionalsStateCount = 0;
 	for (const FReplicationStateDescriptor* StateDescriptor : MakeArrayView(Protocol->ReplicationStateDescriptors, Protocol->ReplicationStateCount))
