Iris - Implemented globally unique NetHandles.

diff --git a/Engine/Plugins/Runtime/ReplicationSystemTestPlugin/Source/Private/Tests/ReplicationSystem/TestReplicationSystem.cpp b/Engine/Plugins/Runtime/ReplicationSystemTestPlugin/Source/Private/Tests/ReplicationSystem/TestReplicationSystem.cpp
--- a/Engine/Plugins/Runtime/ReplicationSystemTestPlugin/Source/Private/Tests/ReplicationSystem/TestReplicationSystem.cpp
+++ b/Engine/Plugins/Runtime/ReplicationSystemTestPlugin/Source/Private/Tests/ReplicationSystem/TestReplicationSystem.cpp
@@ -25,7 +25,7 @@ UE_NET_TEST_FIXTURE(FReplicationSystemServerClientTestFixture, ReplicateAndDestr
 	Server->PostSendUpdate();
 
 	// Verify that created server handle now also exists on client
-	UE_NET_ASSERT_TRUE(Client->GetReplicationBridge()->GetReplicatedObject(ServerObject->NetHandle) != nullptr);
+	UE_NET_ASSERT_TRUE(Client->GetReplicationBridge()->GetReplicatedObject(ServerObject->NetRefHandle) != nullptr);
 
 	// Destroy the spawned object on server
 	Server->DestroyObject(ServerObject);
@@ -36,7 +36,7 @@ UE_NET_TEST_FIXTURE(FReplicationSystemServerClientTestFixture, ReplicateAndDestr
 	Server->PostSendUpdate();
 
 	// Verify that object now is destroyed on client as well
-	UE_NET_ASSERT_FALSE(Client->GetReplicationBridge()->GetReplicatedObject(ServerObject->NetHandle) != nullptr);
+	UE_NET_ASSERT_FALSE(Client->GetReplicationBridge()->GetReplicatedObject(ServerObject->NetRefHandle) != nullptr);
 }
 
 UE_NET_TEST_FIXTURE(FReplicationSystemServerClientTestFixture, TestDroppedDestroyed)
@@ -53,7 +53,7 @@ UE_NET_TEST_FIXTURE(FReplicationSystemServerClientTestFixture, TestDroppedDestro
 	Server->PostSendUpdate();
 
 	// Verify that created server handle now also exists on client
-	UE_NET_ASSERT_TRUE(Client->GetReplicationBridge()->GetReplicatedObject(ServerObject->NetHandle) != nullptr);
+	UE_NET_ASSERT_TRUE(Client->GetReplicationBridge()->GetReplicatedObject(ServerObject->NetRefHandle) != nullptr);
 
 	// Destroy the spawned object on server
 	Server->DestroyObject(ServerObject);
@@ -69,7 +69,7 @@ UE_NET_TEST_FIXTURE(FReplicationSystemServerClientTestFixture, TestDroppedDestro
 	Server->PostSendUpdate();
 
 	// Verify that object now is destroyed on client
-	UE_NET_ASSERT_FALSE(Client->GetReplicationBridge()->GetReplicatedObject(ServerObject->NetHandle) != nullptr);
+	UE_NET_ASSERT_FALSE(Client->GetReplicationBridge()->GetReplicatedObject(ServerObject->NetRefHandle) != nullptr);
 }
 
 UE_NET_TEST_FIXTURE(FReplicationSystemServerClientTestFixture, TestDestroyWhilePendingCreate)
@@ -90,7 +90,7 @@ UE_NET_TEST_FIXTURE(FReplicationSystemServerClientTestFixture, TestDestroyWhileP
 	Server->DestroyObject(ServerObject);
 
 	// Verify that created server handle does not exists on client
-	UE_NET_ASSERT_FALSE(Client->GetReplicationBridge()->GetReplicatedObject(ServerObject->NetHandle) != nullptr);
+	UE_NET_ASSERT_FALSE(Client->GetReplicationBridge()->GetReplicatedObject(ServerObject->NetRefHandle) != nullptr);
 
 	// Send and drop packet
 	Server->PreSendUpdate();
@@ -103,7 +103,7 @@ UE_NET_TEST_FIXTURE(FReplicationSystemServerClientTestFixture, TestDestroyWhileP
 	Server->PostSendUpdate();
 
 	// Verify that object now is destroyed on client
-	UE_NET_ASSERT_FALSE(Client->GetReplicationBridge()->GetReplicatedObject(ServerObject->NetHandle) != nullptr);
+	UE_NET_ASSERT_FALSE(Client->GetReplicationBridge()->GetReplicatedObject(ServerObject->NetRefHandle) != nullptr);
 }
 
 UE_NET_TEST_FIXTURE(FReplicationSystemServerClientTestFixture, TestDestroyWhileWaitingOnCreate)
@@ -141,7 +141,7 @@ UE_NET_TEST_FIXTURE(FReplicationSystemServerClientTestFixture, TestDestroyWithDa
 
 	// Spawn object on server
 	UTestReplicatedIrisObject* ServerObject = Server->CreateObject(0,0);
-	Server->GetReplicationSystem()->SetStaticPriority(ServerObject->NetHandle, 1.f);
+	Server->GetReplicationSystem()->SetStaticPriority(ServerObject->NetRefHandle, 1.f);
 
 	// Send packet with create
 	Server->PreSendUpdate();
@@ -180,7 +180,7 @@ UE_NET_TEST_FIXTURE(FReplicationSystemServerClientTestFixture, ReplicateAndDestr
 	UReplicatedTestObject* ServerObject = Server->CreateObject(0,0);
 
 	// Spawn second object on server as a subobject
-	UReplicatedTestObject* ServerSubObject = Server->CreateSubObject(ServerObject->NetHandle, 0, 0);
+	UReplicatedTestObject* ServerSubObject = Server->CreateSubObject(ServerObject->NetRefHandle, 0, 0);
 
 	// Send and deliver packet
 	Server->PreSendUpdate();
@@ -188,8 +188,8 @@ UE_NET_TEST_FIXTURE(FReplicationSystemServerClientTestFixture, ReplicateAndDestr
 	Server->PostSendUpdate();
 
 	// Verify that created server handles now also exists on client
-	UE_NET_ASSERT_TRUE(Client->GetReplicationBridge()->GetReplicatedObject(ServerObject->NetHandle) != nullptr);
-	UE_NET_ASSERT_TRUE(Client->GetReplicationBridge()->GetReplicatedObject(ServerSubObject->NetHandle) != nullptr);
+	UE_NET_ASSERT_TRUE(Client->GetReplicationBridge()->GetReplicatedObject(ServerObject->NetRefHandle) != nullptr);
+	UE_NET_ASSERT_TRUE(Client->GetReplicationBridge()->GetReplicatedObject(ServerSubObject->NetRefHandle) != nullptr);
 
 	// Destroy the spawned subobject on server
 	Server->DestroyObject(ServerSubObject);
@@ -200,7 +200,7 @@ UE_NET_TEST_FIXTURE(FReplicationSystemServerClientTestFixture, ReplicateAndDestr
 	Server->PostSendUpdate();
 
 	// Verify that object now is destroyed on client as well
-	UE_NET_ASSERT_FALSE(Client->GetReplicationBridge()->GetReplicatedObject(ServerSubObject->NetHandle) != nullptr);
+	UE_NET_ASSERT_FALSE(Client->GetReplicationBridge()->GetReplicatedObject(ServerSubObject->NetRefHandle) != nullptr);
 }
 
 UE_NET_TEST_FIXTURE(FReplicationSystemServerClientTestFixture, ReplicateAndDestroyMultipleSubObjects)
@@ -212,7 +212,7 @@ UE_NET_TEST_FIXTURE(FReplicationSystemServerClientTestFixture, ReplicateAndDestr
 	UReplicatedTestObject* ServerObject = Server->CreateObject(0, 0);
 
 	// Spawn second object on server as a subobject
-	UReplicatedTestObject* ServerSubObject = Server->CreateSubObject(ServerObject->NetHandle, 0, 0);
+	UReplicatedTestObject* ServerSubObject = Server->CreateSubObject(ServerObject->NetRefHandle, 0, 0);
 
 	// Send and deliver packet
 	Server->PreSendUpdate();
@@ -220,8 +220,8 @@ UE_NET_TEST_FIXTURE(FReplicationSystemServerClientTestFixture, ReplicateAndDestr
 	Server->PostSendUpdate();
 
 	// Verify that created server handles now also exists on client
-	UE_NET_ASSERT_NE(Client->GetReplicationBridge()->GetReplicatedObject(ServerObject->NetHandle), nullptr);
-	UE_NET_ASSERT_NE(Client->GetReplicationBridge()->GetReplicatedObject(ServerSubObject->NetHandle), nullptr);
+	UE_NET_ASSERT_NE(Client->GetReplicationBridge()->GetReplicatedObject(ServerObject->NetRefHandle), nullptr);
+	UE_NET_ASSERT_NE(Client->GetReplicationBridge()->GetReplicatedObject(ServerSubObject->NetRefHandle), nullptr);
 
 	// Destroy the spawned subobject on server
 	Server->DestroyObject(ServerSubObject);
@@ -232,10 +232,10 @@ UE_NET_TEST_FIXTURE(FReplicationSystemServerClientTestFixture, ReplicateAndDestr
 	Server->PostSendUpdate();
 
 	// Verify that subobject now is destroyed on client as well
-	UE_NET_ASSERT_EQ(Client->GetReplicationBridge()->GetReplicatedObject(ServerSubObject->NetHandle), nullptr);
+	UE_NET_ASSERT_EQ(Client->GetReplicationBridge()->GetReplicatedObject(ServerSubObject->NetRefHandle), nullptr);
 
 	// Spawn second object on server as a subobject
-	ServerSubObject = Server->CreateSubObject(ServerObject->NetHandle, 0, 0);
+	ServerSubObject = Server->CreateSubObject(ServerObject->NetRefHandle, 0, 0);
 
 	// Send and deliver packet
 	Server->PreSendUpdate();
@@ -243,7 +243,7 @@ UE_NET_TEST_FIXTURE(FReplicationSystemServerClientTestFixture, ReplicateAndDestr
 	Server->PostSendUpdate();
 
 	// Verify that second subobject replicated properly to server
-	UE_NET_ASSERT_NE(Client->GetReplicationBridge()->GetReplicatedObject(ServerSubObject->NetHandle), nullptr);
+	UE_NET_ASSERT_NE(Client->GetReplicationBridge()->GetReplicatedObject(ServerSubObject->NetRefHandle), nullptr);
 
 	// Destroy the spawned object on server
 	Server->DestroyObject(ServerSubObject);
@@ -254,7 +254,7 @@ UE_NET_TEST_FIXTURE(FReplicationSystemServerClientTestFixture, ReplicateAndDestr
 	Server->PostSendUpdate();
 
 	// Verify that object now is destroyed on client as well
-	UE_NET_ASSERT_EQ(Client->GetReplicationBridge()->GetReplicatedObject(ServerSubObject->NetHandle), nullptr);
+	UE_NET_ASSERT_EQ(Client->GetReplicationBridge()->GetReplicatedObject(ServerSubObject->NetRefHandle), nullptr);
 }
 
 UE_NET_TEST_FIXTURE(FReplicationSystemServerClientTestFixture, ReplicateAndDestroySubObjectAndDestroyOwner)
@@ -266,7 +266,7 @@ UE_NET_TEST_FIXTURE(FReplicationSystemServerClientTestFixture, ReplicateAndDestr
 	UReplicatedTestObject* ServerObject = Server->CreateObject(0, 0);
 
 	// Spawn second object on server as a subobject
-	UReplicatedTestObject* ServerSubObject = Server->CreateSubObject(ServerObject->NetHandle, 0, 0);
+	UReplicatedTestObject* ServerSubObject = Server->CreateSubObject(ServerObject->NetRefHandle, 0, 0);
 
 	// Send and deliver packet
 	Server->PreSendUpdate();
@@ -274,8 +274,8 @@ UE_NET_TEST_FIXTURE(FReplicationSystemServerClientTestFixture, ReplicateAndDestr
 	Server->PostSendUpdate();
 
 	// Verify that created server handles now also exists on client
-	UE_NET_ASSERT_NE(Client->GetReplicationBridge()->GetReplicatedObject(ServerObject->NetHandle), nullptr);
-	UE_NET_ASSERT_NE(Client->GetReplicationBridge()->GetReplicatedObject(ServerSubObject->NetHandle), nullptr);
+	UE_NET_ASSERT_NE(Client->GetReplicationBridge()->GetReplicatedObject(ServerObject->NetRefHandle), nullptr);
+	UE_NET_ASSERT_NE(Client->GetReplicationBridge()->GetReplicatedObject(ServerSubObject->NetRefHandle), nullptr);
 
 	// Destroy the spawned subobject on server
 	Server->DestroyObject(ServerSubObject);
@@ -289,7 +289,7 @@ UE_NET_TEST_FIXTURE(FReplicationSystemServerClientTestFixture, ReplicateAndDestr
 	Server->PostSendUpdate();
 
 	// Verify that subobject now is destroyed on client as well
-	UE_NET_ASSERT_EQ(Client->GetReplicationBridge()->GetReplicatedObject(ServerSubObject->NetHandle), nullptr);
+	UE_NET_ASSERT_EQ(Client->GetReplicationBridge()->GetReplicatedObject(ServerSubObject->NetRefHandle), nullptr);
 }
 
 UE_NET_TEST_FIXTURE(FReplicationSystemServerClientTestFixture, ReplicateAndDestroySubObjectAndDestroyOwnerWithDataInFlight)
@@ -301,7 +301,7 @@ UE_NET_TEST_FIXTURE(FReplicationSystemServerClientTestFixture, ReplicateAndDestr
 	UReplicatedTestObject* ServerObject = Server->CreateObject(0, 0);
 
 	// Spawn second object on server as a subobject
-	UReplicatedTestObject* ServerSubObject = Server->CreateSubObject(ServerObject->NetHandle, 0, 0);
+	UReplicatedTestObject* ServerSubObject = Server->CreateSubObject(ServerObject->NetRefHandle, 0, 0);
 
 	// Send and deliver packet
 	Server->PreSendUpdate();
@@ -309,8 +309,8 @@ UE_NET_TEST_FIXTURE(FReplicationSystemServerClientTestFixture, ReplicateAndDestr
 	Server->PostSendUpdate();
 
 	// Verify that created server handles now also exists on client
-	UE_NET_ASSERT_NE(Client->GetReplicationBridge()->GetReplicatedObject(ServerObject->NetHandle), nullptr);
-	UE_NET_ASSERT_NE(Client->GetReplicationBridge()->GetReplicatedObject(ServerSubObject->NetHandle), nullptr);
+	UE_NET_ASSERT_NE(Client->GetReplicationBridge()->GetReplicatedObject(ServerObject->NetRefHandle), nullptr);
+	UE_NET_ASSERT_NE(Client->GetReplicationBridge()->GetReplicatedObject(ServerSubObject->NetRefHandle), nullptr);
 
 	// Destroy the spawned subobject on server
 	Server->DestroyObject(ServerSubObject);
@@ -329,7 +329,7 @@ UE_NET_TEST_FIXTURE(FReplicationSystemServerClientTestFixture, ReplicateAndDestr
 	Server->PostSendUpdate();
 
 	// Verify that subobject now is destroyed on client as well
-	UE_NET_ASSERT_EQ(Client->GetReplicationBridge()->GetReplicatedObject(ServerSubObject->NetHandle), nullptr);
+	UE_NET_ASSERT_EQ(Client->GetReplicationBridge()->GetReplicatedObject(ServerSubObject->NetRefHandle), nullptr);
 }
 
 UE_NET_TEST_FIXTURE(FReplicationSystemServerClientTestFixture, ReplicateAndDestroySubObjectWithLostData)
@@ -341,7 +341,7 @@ UE_NET_TEST_FIXTURE(FReplicationSystemServerClientTestFixture, ReplicateAndDestr
 	UReplicatedTestObject* ServerObject = Server->CreateObject(0, 0);
 
 	// Spawn second object on server as a subobject
-	UReplicatedTestObject* ServerSubObject = Server->CreateSubObject(ServerObject->NetHandle, 0, 0);
+	UReplicatedTestObject* ServerSubObject = Server->CreateSubObject(ServerObject->NetRefHandle, 0, 0);
 
 	// Send and deliver packet
 	Server->PreSendUpdate();
@@ -349,8 +349,8 @@ UE_NET_TEST_FIXTURE(FReplicationSystemServerClientTestFixture, ReplicateAndDestr
 	Server->PostSendUpdate();
 
 	// Verify that created server handles now also exists on client
-	UE_NET_ASSERT_NE(Client->GetReplicationBridge()->GetReplicatedObject(ServerObject->NetHandle), nullptr);
-	UE_NET_ASSERT_NE(Client->GetReplicationBridge()->GetReplicatedObject(ServerSubObject->NetHandle), nullptr);
+	UE_NET_ASSERT_NE(Client->GetReplicationBridge()->GetReplicatedObject(ServerObject->NetRefHandle), nullptr);
+	UE_NET_ASSERT_NE(Client->GetReplicationBridge()->GetReplicatedObject(ServerSubObject->NetRefHandle), nullptr);
 
 	// Destroy the spawned subobject on server
 	Server->DestroyObject(ServerSubObject);
@@ -366,7 +366,7 @@ UE_NET_TEST_FIXTURE(FReplicationSystemServerClientTestFixture, ReplicateAndDestr
 	Server->PostSendUpdate();
 
 	// Verify that subobject now is destroyed on client as well
-	UE_NET_ASSERT_EQ(Client->GetReplicationBridge()->GetReplicatedObject(ServerSubObject->NetHandle), nullptr);
+	UE_NET_ASSERT_EQ(Client->GetReplicationBridge()->GetReplicatedObject(ServerSubObject->NetRefHandle), nullptr);
 }
 
 UE_NET_TEST_FIXTURE(FReplicationSystemServerClientTestFixture, ReplicateAndDestroySubObjectPendingCreateConfirmation)
@@ -378,7 +378,7 @@ UE_NET_TEST_FIXTURE(FReplicationSystemServerClientTestFixture, ReplicateAndDestr
 	UReplicatedTestObject* ServerObject = Server->CreateObject(0, 0);
 
 	// Spawn second object on server as a subobject
-	UReplicatedTestObject* ServerSubObject = Server->CreateSubObject(ServerObject->NetHandle, 0, 0);
+	UReplicatedTestObject* ServerSubObject = Server->CreateSubObject(ServerObject->NetRefHandle, 0, 0);
 
 	uint32 NumUnAcknowledgedPackets = 0;
 	// Send and deliver packet
@@ -408,7 +408,7 @@ UE_NET_TEST_FIXTURE(FReplicationSystemServerClientTestFixture, ReplicateAndDestr
 	Server->PostSendUpdate();
 
 	// Verify that subobject now is destroyed on client as well
-	UE_NET_ASSERT_EQ(Client->GetReplicationBridge()->GetReplicatedObject(ServerSubObject->NetHandle), nullptr);
+	UE_NET_ASSERT_EQ(Client->GetReplicationBridge()->GetReplicatedObject(ServerSubObject->NetRefHandle), nullptr);
 }
 
 UE_NET_TEST_FIXTURE(FReplicationSystemServerClientTestFixture, TestSubObjectDefaultReplicationOrder)
@@ -420,9 +420,9 @@ UE_NET_TEST_FIXTURE(FReplicationSystemServerClientTestFixture, TestSubObjectDefa
 	UReplicatedTestObject* ServerObject = Server->CreateObject(0, 0);
 
 	// Spawn some subobjects
-	UReplicatedSubObjectOrderObject* ServerSubObject0 = Server->CreateSubObject<UReplicatedSubObjectOrderObject>(ServerObject->NetHandle);
-	UReplicatedSubObjectOrderObject* ServerSubObject1 = Server->CreateSubObject<UReplicatedSubObjectOrderObject>(ServerObject->NetHandle);
-	UReplicatedSubObjectOrderObject* ServerSubObject2 = Server->CreateSubObject<UReplicatedSubObjectOrderObject>(ServerObject->NetHandle);
+	UReplicatedSubObjectOrderObject* ServerSubObject0 = Server->CreateSubObject<UReplicatedSubObjectOrderObject>(ServerObject->NetRefHandle);
+	UReplicatedSubObjectOrderObject* ServerSubObject1 = Server->CreateSubObject<UReplicatedSubObjectOrderObject>(ServerObject->NetRefHandle);
+	UReplicatedSubObjectOrderObject* ServerSubObject2 = Server->CreateSubObject<UReplicatedSubObjectOrderObject>(ServerObject->NetRefHandle);
 
 	UReplicatedSubObjectOrderObject::RepOrderCounter = 0U;
 
@@ -432,9 +432,9 @@ UE_NET_TEST_FIXTURE(FReplicationSystemServerClientTestFixture, TestSubObjectDefa
 	Server->PostSendUpdate();
 
 	// Verify that objects have replicated
-	UReplicatedSubObjectOrderObject* ClientSubObject0 = Cast<UReplicatedSubObjectOrderObject>(Client->GetReplicationBridge()->GetReplicatedObject(ServerSubObject0->NetHandle));
-	UReplicatedSubObjectOrderObject* ClientSubObject1 = Cast<UReplicatedSubObjectOrderObject>(Client->GetReplicationBridge()->GetReplicatedObject(ServerSubObject1->NetHandle));
-	UReplicatedSubObjectOrderObject* ClientSubObject2 = Cast<UReplicatedSubObjectOrderObject>(Client->GetReplicationBridge()->GetReplicatedObject(ServerSubObject2->NetHandle));
+	UReplicatedSubObjectOrderObject* ClientSubObject0 = Cast<UReplicatedSubObjectOrderObject>(Client->GetReplicationBridge()->GetReplicatedObject(ServerSubObject0->NetRefHandle));
+	UReplicatedSubObjectOrderObject* ClientSubObject1 = Cast<UReplicatedSubObjectOrderObject>(Client->GetReplicationBridge()->GetReplicatedObject(ServerSubObject1->NetRefHandle));
+	UReplicatedSubObjectOrderObject* ClientSubObject2 = Cast<UReplicatedSubObjectOrderObject>(Client->GetReplicationBridge()->GetReplicatedObject(ServerSubObject2->NetRefHandle));
 
 	UE_NET_ASSERT_NE(ClientSubObject0, nullptr);
 	UE_NET_ASSERT_NE(ClientSubObject1, nullptr);
@@ -455,11 +455,11 @@ UE_NET_TEST_FIXTURE(FReplicationSystemServerClientTestFixture, TestSubObjectSpec
 	UReplicatedTestObject* ServerObject = Server->CreateObject(0, 0);
 
 	// Spawn some subobjects
-	UReplicatedSubObjectOrderObject* ServerSubObject0 = Server->CreateSubObject<UReplicatedSubObjectOrderObject>(ServerObject->NetHandle);
+	UReplicatedSubObjectOrderObject* ServerSubObject0 = Server->CreateSubObject<UReplicatedSubObjectOrderObject>(ServerObject->NetRefHandle);
 	// Specify Subobject1 to replicate with SubObject0, which means that it will replicate before Subobjet0 is replicated
-	UReplicatedSubObjectOrderObject* ServerSubObject1 = Server->CreateSubObject<UReplicatedSubObjectOrderObject>(ServerObject->NetHandle, ServerSubObject0->NetHandle, UReplicationBridge::ESubObjectInsertionOrder::ReplicateWith);
+	UReplicatedSubObjectOrderObject* ServerSubObject1 = Server->CreateSubObject<UReplicatedSubObjectOrderObject>(ServerObject->NetRefHandle, ServerSubObject0->NetRefHandle, UReplicationBridge::ESubObjectInsertionOrder::ReplicateWith);
 	// Specify SubObect 2 to replicate with no specific order (it will be added to the owner and thus replicate last)
-	UReplicatedSubObjectOrderObject* ServerSubObject2 = Server->CreateSubObject<UReplicatedSubObjectOrderObject>(ServerObject->NetHandle, ServerSubObject1->NetHandle, UReplicationBridge::ESubObjectInsertionOrder::None);
+	UReplicatedSubObjectOrderObject* ServerSubObject2 = Server->CreateSubObject<UReplicatedSubObjectOrderObject>(ServerObject->NetRefHandle, ServerSubObject1->NetRefHandle, UReplicationBridge::ESubObjectInsertionOrder::None);
 
 	UReplicatedSubObjectOrderObject::RepOrderCounter = 0U;
 
@@ -469,9 +469,9 @@ UE_NET_TEST_FIXTURE(FReplicationSystemServerClientTestFixture, TestSubObjectSpec
 	Server->PostSendUpdate();
 
 	// Verify that objects have replicated
-	UReplicatedSubObjectOrderObject* ClientSubObject0 = Cast<UReplicatedSubObjectOrderObject>(Client->GetReplicationBridge()->GetReplicatedObject(ServerSubObject0->NetHandle));
-	UReplicatedSubObjectOrderObject* ClientSubObject1 = Cast<UReplicatedSubObjectOrderObject>(Client->GetReplicationBridge()->GetReplicatedObject(ServerSubObject1->NetHandle));
-	UReplicatedSubObjectOrderObject* ClientSubObject2 = Cast<UReplicatedSubObjectOrderObject>(Client->GetReplicationBridge()->GetReplicatedObject(ServerSubObject2->NetHandle));
+	UReplicatedSubObjectOrderObject* ClientSubObject0 = Cast<UReplicatedSubObjectOrderObject>(Client->GetReplicationBridge()->GetReplicatedObject(ServerSubObject0->NetRefHandle));
+	UReplicatedSubObjectOrderObject* ClientSubObject1 = Cast<UReplicatedSubObjectOrderObject>(Client->GetReplicationBridge()->GetReplicatedObject(ServerSubObject1->NetRefHandle));
+	UReplicatedSubObjectOrderObject* ClientSubObject2 = Cast<UReplicatedSubObjectOrderObject>(Client->GetReplicationBridge()->GetReplicatedObject(ServerSubObject2->NetRefHandle));
 
 	UE_NET_ASSERT_NE(ClientSubObject0, nullptr);
 	UE_NET_ASSERT_NE(ClientSubObject1, nullptr);
@@ -713,7 +713,7 @@ UE_NET_TEST_FIXTURE(FReplicationSystemServerClientTestFixture, AddRemoveFromConn
 
 	// Add to group
 	FNetObjectGroupHandle Group = ReplicationSystem->CreateGroup();
-	ReplicationSystem->AddToGroup(Group, ServerObject->NetHandle);
+	ReplicationSystem->AddToGroup(Group, ServerObject->NetRefHandle);
 
 	ReplicationSystem->AddGroupFilter(Group);
 	ReplicationSystem->SetGroupFilterStatus(Group, ENetFilterStatus::Allow);
@@ -743,11 +743,11 @@ UE_NET_TEST_FIXTURE(FReplicationSystemServerClientTestFixture, AddRemoveFromConn
 
 	// Expect client to create object
 	Server->DeliverTo(Client, true);
-	UE_NET_ASSERT_TRUE(Client->GetReplicationBridge()->GetReplicatedObject(ServerObject->NetHandle) != nullptr);
+	UE_NET_ASSERT_TRUE(Client->GetReplicationBridge()->GetReplicatedObject(ServerObject->NetRefHandle) != nullptr);
 
 	// Expect client to destroy object and server to move to Destroyed state
 	Server->DeliverTo(Client, true);
-	UE_NET_ASSERT_TRUE(Client->GetReplicationBridge()->GetReplicatedObject(ServerObject->NetHandle) == nullptr);
+	UE_NET_ASSERT_TRUE(Client->GetReplicationBridge()->GetReplicatedObject(ServerObject->NetRefHandle) == nullptr);
 
 	// Trigger replication
 	++ServerObject->IntA;
@@ -760,7 +760,7 @@ UE_NET_TEST_FIXTURE(FReplicationSystemServerClientTestFixture, AddRemoveFromConn
 	Server->PostSendUpdate();
 
 	// Verify that the object got created again
-	UE_NET_ASSERT_TRUE(Client->GetReplicationBridge()->GetReplicatedObject(ServerObject->NetHandle) != nullptr);
+	UE_NET_ASSERT_TRUE(Client->GetReplicationBridge()->GetReplicatedObject(ServerObject->NetRefHandle) != nullptr);
 }
 
 UE_NET_TEST_FIXTURE(FReplicationSystemServerClientTestFixture, TestNetTemporary)
@@ -785,14 +785,14 @@ UE_NET_TEST_FIXTURE(FReplicationSystemServerClientTestFixture, TestNetTemporary)
 
 	// Verify that client has received the data
 	{
-		UTestReplicatedIrisObject* ClientObject = Cast<UTestReplicatedIrisObject>(Client->GetReplicationBridge()->GetReplicatedObject(ServerObject->NetHandle));
+		UTestReplicatedIrisObject* ClientObject = Cast<UTestReplicatedIrisObject>(Client->GetReplicationBridge()->GetReplicatedObject(ServerObject->NetRefHandle));
 
 		UE_NET_ASSERT_TRUE(ClientObject != nullptr);
 		UE_NET_ASSERT_EQ(ServerObject->IntA, ClientObject->IntA);
 	}
 
 	{
-		UTestReplicatedIrisObject* ClientObject = Cast<UTestReplicatedIrisObject>(Client->GetReplicationBridge()->GetReplicatedObject(ServerObject1->NetHandle));
+		UTestReplicatedIrisObject* ClientObject = Cast<UTestReplicatedIrisObject>(Client->GetReplicationBridge()->GetReplicatedObject(ServerObject1->NetRefHandle));
 
 		UE_NET_ASSERT_TRUE(ClientObject != nullptr);
 		UE_NET_ASSERT_EQ(ServerObject1->IntA, ClientObject->IntA);
@@ -800,7 +800,7 @@ UE_NET_TEST_FIXTURE(FReplicationSystemServerClientTestFixture, TestNetTemporary)
 	}
 
 	// Mark the object as a net temporary
-	ReplicationSystem->SetIsNetTemporary(ServerObject->NetHandle);
+	ReplicationSystem->SetIsNetTemporary(ServerObject->NetRefHandle);
 
 	// Modify the value
 	ServerObject->IntA = 2;
@@ -813,7 +813,7 @@ UE_NET_TEST_FIXTURE(FReplicationSystemServerClientTestFixture, TestNetTemporary)
 
 	// Verify that client has not received the data for changed temporary
 	{
-		UTestReplicatedIrisObject* ClientObject = Cast<UTestReplicatedIrisObject>(Client->GetReplicationBridge()->GetReplicatedObject(ServerObject->NetHandle));
+		UTestReplicatedIrisObject* ClientObject = Cast<UTestReplicatedIrisObject>(Client->GetReplicationBridge()->GetReplicatedObject(ServerObject->NetRefHandle));
 
 		UE_NET_ASSERT_TRUE(ClientObject != nullptr);
 		UE_NET_ASSERT_NE(ServerObject->IntA, ClientObject->IntA);
@@ -821,7 +821,7 @@ UE_NET_TEST_FIXTURE(FReplicationSystemServerClientTestFixture, TestNetTemporary)
 
 	// Verify that client has received the data for normal object
 	{
-		UTestReplicatedIrisObject* ClientObject = Cast<UTestReplicatedIrisObject>(Client->GetReplicationBridge()->GetReplicatedObject(ServerObject1->NetHandle));
+		UTestReplicatedIrisObject* ClientObject = Cast<UTestReplicatedIrisObject>(Client->GetReplicationBridge()->GetReplicatedObject(ServerObject1->NetRefHandle));
 
 		UE_NET_ASSERT_TRUE(ClientObject != nullptr);
 		UE_NET_ASSERT_EQ(ServerObject1->IntA, ClientObject->IntA);
@@ -838,7 +838,7 @@ UE_NET_TEST_FIXTURE(FReplicationSystemServerClientTestFixture, TestNetTemporary)
 
 	// We should now have the latest state for both objects
 	{
-		UTestReplicatedIrisObject* ClientObject = Cast<UTestReplicatedIrisObject>(Client2->GetReplicationBridge()->GetReplicatedObject(ServerObject->NetHandle));
+		UTestReplicatedIrisObject* ClientObject = Cast<UTestReplicatedIrisObject>(Client2->GetReplicationBridge()->GetReplicatedObject(ServerObject->NetRefHandle));
 
 		UE_NET_ASSERT_TRUE(ClientObject != nullptr);
 		UE_NET_ASSERT_EQ(ServerObject->IntA, ClientObject->IntA);
@@ -846,7 +846,7 @@ UE_NET_TEST_FIXTURE(FReplicationSystemServerClientTestFixture, TestNetTemporary)
 
 	// Verify that client has received the data for normal object
 	{
-		UTestReplicatedIrisObject* ClientObject = Cast<UTestReplicatedIrisObject>(Client2->GetReplicationBridge()->GetReplicatedObject(ServerObject1->NetHandle));
+		UTestReplicatedIrisObject* ClientObject = Cast<UTestReplicatedIrisObject>(Client2->GetReplicationBridge()->GetReplicatedObject(ServerObject1->NetRefHandle));
 
 		UE_NET_ASSERT_TRUE(ClientObject != nullptr);
 		UE_NET_ASSERT_EQ(ServerObject1->IntA, ClientObject->IntA);
@@ -875,7 +875,7 @@ UE_NET_TEST_FIXTURE(FReplicationSystemServerClientTestFixture, TestTearOffExisti
 	Server->PostSendUpdate();
 
 	// Store Pointer to object 
-	UTestReplicatedIrisObject* ClientObjectThatWillBeTornOff = Cast<UTestReplicatedIrisObject>(Client->GetReplicationBridge()->GetReplicatedObject(ServerObject->NetHandle));
+	UTestReplicatedIrisObject* ClientObjectThatWillBeTornOff = Cast<UTestReplicatedIrisObject>(Client->GetReplicationBridge()->GetReplicatedObject(ServerObject->NetRefHandle));
 
 	UE_NET_ASSERT_TRUE(ClientObjectThatWillBeTornOff != nullptr);
 	UE_NET_ASSERT_EQ(ServerObject->IntA, ClientObjectThatWillBeTornOff->IntA);
@@ -893,7 +893,7 @@ UE_NET_TEST_FIXTURE(FReplicationSystemServerClientTestFixture, TestTearOffExisti
 
 	// Verify that ClientObject is torn-off and that the final state was applied
 	UE_NET_ASSERT_EQ(ServerObject->IntA, ClientObjectThatWillBeTornOff->IntA);
-	UE_NET_ASSERT_TRUE(Cast<UTestReplicatedIrisObject>(Client->GetReplicationBridge()->GetReplicatedObject(ServerObject->NetHandle)) == nullptr);
+	UE_NET_ASSERT_TRUE(Cast<UTestReplicatedIrisObject>(Client->GetReplicationBridge()->GetReplicatedObject(ServerObject->NetRefHandle)) == nullptr);
 }
 
 // Test TearOff for new object
@@ -924,7 +924,7 @@ UE_NET_TEST_FIXTURE(FReplicationSystemServerClientTestFixture, TestTearOffOnNewl
 	// Client should have created a object
 	UE_NET_ASSERT_EQ(NumObjectsCreatedOnClientBeforeReplication + 1, Client->CreatedObjects.Num());
 
-	UE_NET_ASSERT_TRUE(Cast<UTestReplicatedIrisObject>(Client->GetReplicationBridge()->GetReplicatedObject(ServerObject->NetHandle)) == nullptr);
+	UE_NET_ASSERT_TRUE(Cast<UTestReplicatedIrisObject>(Client->GetReplicationBridge()->GetReplicatedObject(ServerObject->NetRefHandle)) == nullptr);
 
 	// We should be able to get the object from the created objects array to validate the state
 	UTestReplicatedIrisObject* ClientObjectThatWillBeTornOff = Cast<UTestReplicatedIrisObject>(Client->CreatedObjects[NumObjectsCreatedOnClientBeforeReplication].Get());
@@ -953,7 +953,7 @@ UE_NET_TEST_FIXTURE(FReplicationSystemServerClientTestFixture, TestTearOffResend
 	Server->PostSendUpdate();
 
 	// Store Pointer to object 
-	UTestReplicatedIrisObject* ClientObjectThatWillBeTornOff = Cast<UTestReplicatedIrisObject>(Client->GetReplicationBridge()->GetReplicatedObject(ServerObject->NetHandle));
+	UTestReplicatedIrisObject* ClientObjectThatWillBeTornOff = Cast<UTestReplicatedIrisObject>(Client->GetReplicationBridge()->GetReplicatedObject(ServerObject->NetRefHandle));
 
 	UE_NET_ASSERT_NE(ClientObjectThatWillBeTornOff, nullptr);
 	UE_NET_ASSERT_EQ(ServerObject->IntA, ClientObjectThatWillBeTornOff->IntA);
@@ -966,15 +966,15 @@ UE_NET_TEST_FIXTURE(FReplicationSystemServerClientTestFixture, TestTearOffResend
 	Server->SendAndDeliverTo(Client, DoNotDeliverPacket);
 	Server->PostSendUpdate();
 
-	// The ClientObject should still be found using the NetHandle
-	UE_NET_ASSERT_NE(Cast<UTestReplicatedIrisObject>(Client->GetReplicationBridge()->GetReplicatedObject(ServerObject->NetHandle)), nullptr);
+	// The ClientObject should still be found using the NetRefHandle
+	UE_NET_ASSERT_NE(Cast<UTestReplicatedIrisObject>(Client->GetReplicationBridge()->GetReplicatedObject(ServerObject->NetRefHandle)), nullptr);
 
 	Server->PreSendUpdate();
 	Server->SendAndDeliverTo(Client, DeliverPacket);
 	Server->PostSendUpdate();
 
 	// Verify that ClientObject is torn-off
-	UE_NET_ASSERT_EQ(Cast<UTestReplicatedIrisObject>(Client->GetReplicationBridge()->GetReplicatedObject(ServerObject->NetHandle)), nullptr);
+	UE_NET_ASSERT_EQ(Cast<UTestReplicatedIrisObject>(Client->GetReplicationBridge()->GetReplicatedObject(ServerObject->NetRefHandle)), nullptr);
 }
 
 // Test TearOff for new object and resend (should not work or is this what we want?)
@@ -1011,7 +1011,7 @@ UE_NET_TEST_FIXTURE(FReplicationSystemServerClientTestFixture, TestTearOffImmedi
 	// Client should have created a object
 	UE_NET_ASSERT_EQ(NumObjectsCreatedOnClientBeforeReplication + 1, Client->CreatedObjects.Num());
 
-	UE_NET_ASSERT_TRUE(Cast<UTestReplicatedIrisObject>(Client->GetReplicationBridge()->GetReplicatedObject(ServerObject->NetHandle)) == nullptr);
+	UE_NET_ASSERT_TRUE(Cast<UTestReplicatedIrisObject>(Client->GetReplicationBridge()->GetReplicatedObject(ServerObject->NetRefHandle)) == nullptr);
 
 	// We should be able to get the object from the created objects array to validate the state
 	UTestReplicatedIrisObject* ClientObjectThatWillBeTornOff = Cast<UTestReplicatedIrisObject>(Client->CreatedObjects[NumObjectsCreatedOnClientBeforeReplication].Get());
@@ -1039,7 +1039,7 @@ UE_NET_TEST_FIXTURE(FReplicationSystemServerClientTestFixture, TestDefferedTearO
 	ServerObject->IntA = 1;
 
 	// TearOff the object before first replication
-	Server->ReplicationSystem->TearOffNextUpdate(ServerObject->NetHandle);
+	Server->ReplicationSystem->TearOffNextUpdate(ServerObject->NetRefHandle);
 
 	// Send and deliver packet
 	Server->PreSendUpdate();
@@ -1057,7 +1057,7 @@ UE_NET_TEST_FIXTURE(FReplicationSystemServerClientTestFixture, TestDefferedTearO
 	// Client should have created a object
 	UE_NET_ASSERT_EQ(NumObjectsCreatedOnClientBeforeReplication + 1, Client->CreatedObjects.Num());
 
-	UE_NET_ASSERT_TRUE(Cast<UTestReplicatedIrisObject>(Client->GetReplicationBridge()->GetReplicatedObject(ServerObject->NetHandle)) == nullptr);
+	UE_NET_ASSERT_TRUE(Cast<UTestReplicatedIrisObject>(Client->GetReplicationBridge()->GetReplicatedObject(ServerObject->NetRefHandle)) == nullptr);
 
 	// We should be able to get the object from the created objects array to validate the state
 	UTestReplicatedIrisObject* ClientObjectThatWillBeTornOff = Cast<UTestReplicatedIrisObject>(Client->CreatedObjects[NumObjectsCreatedOnClientBeforeReplication].Get());
@@ -1092,7 +1092,7 @@ UE_NET_TEST_FIXTURE(FReplicationSystemServerClientTestFixture, TestTearOffObject
 	Server->DeliverTo(Client, true);
 
 	// Store Pointer to object and verify initial state
-	UTestReplicatedIrisObject* ClientObjectThatWillBeTornOff = Cast<UTestReplicatedIrisObject>(Client->GetReplicationBridge()->GetReplicatedObject(ServerObject->NetHandle));
+	UTestReplicatedIrisObject* ClientObjectThatWillBeTornOff = Cast<UTestReplicatedIrisObject>(Client->GetReplicationBridge()->GetReplicatedObject(ServerObject->NetRefHandle));
 
 	UE_NET_ASSERT_TRUE(ClientObjectThatWillBeTornOff != nullptr);
 	UE_NET_ASSERT_EQ(ServerObject->IntA, ClientObjectThatWillBeTornOff->IntA);
@@ -1104,7 +1104,7 @@ UE_NET_TEST_FIXTURE(FReplicationSystemServerClientTestFixture, TestTearOffObject
 
 	// Verify that ClientObject is TornOFf and that the final state was applied
 	UE_NET_ASSERT_EQ(ServerObject->IntA, ClientObjectThatWillBeTornOff->IntA);
-	UE_NET_ASSERT_TRUE(Cast<UTestReplicatedIrisObject>(Client->GetReplicationBridge()->GetReplicatedObject(ServerObject->NetHandle)) == nullptr);
+	UE_NET_ASSERT_TRUE(Cast<UTestReplicatedIrisObject>(Client->GetReplicationBridge()->GetReplicatedObject(ServerObject->NetRefHandle)) == nullptr);
 }
 
 // Test TearOff for existing object pending destroy (should do nothing)
@@ -1127,7 +1127,7 @@ UE_NET_TEST_FIXTURE(FReplicationSystemServerClientTestFixture, TestTearOffExisti
 	Server->PostSendUpdate();
 
 	// Store Pointer to object 
-	UTestReplicatedIrisObject* ClientObjectThatWillBeTornOff = Cast<UTestReplicatedIrisObject>(Client->GetReplicationBridge()->GetReplicatedObject(ServerObject->NetHandle));
+	UTestReplicatedIrisObject* ClientObjectThatWillBeTornOff = Cast<UTestReplicatedIrisObject>(Client->GetReplicationBridge()->GetReplicatedObject(ServerObject->NetRefHandle));
 
 	UE_NET_ASSERT_TRUE(ClientObjectThatWillBeTornOff != nullptr);
 	UE_NET_ASSERT_EQ(ServerObject->IntA, ClientObjectThatWillBeTornOff->IntA);
@@ -1148,7 +1148,7 @@ UE_NET_TEST_FIXTURE(FReplicationSystemServerClientTestFixture, TestTearOffExisti
 
 	// Verify that ClientObject is not tornOff and that the final state was not applied as we issued tearoff after ending replication
 	UE_NET_ASSERT_NE(ServerObject->IntA, ClientObjectThatWillBeTornOff->IntA);
-	UE_NET_ASSERT_TRUE(Cast<UTestReplicatedIrisObject>(Client->GetReplicationBridge()->GetReplicatedObject(ServerObject->NetHandle)) == nullptr);
+	UE_NET_ASSERT_TRUE(Cast<UTestReplicatedIrisObject>(Client->GetReplicationBridge()->GetReplicatedObject(ServerObject->NetRefHandle)) == nullptr);
 }
 
 // Test TearOff resend 
@@ -1171,7 +1171,7 @@ UE_NET_TEST_FIXTURE(FReplicationSystemServerClientTestFixture, TestTearOffResend
 	Server->PostSendUpdate();
 
 	// Store Pointer to object 
-	UTestReplicatedIrisObject* ClientObjectThatWillBeTornOff = Cast<UTestReplicatedIrisObject>(Client->GetReplicationBridge()->GetReplicatedObject(ServerObject->NetHandle));
+	UTestReplicatedIrisObject* ClientObjectThatWillBeTornOff = Cast<UTestReplicatedIrisObject>(Client->GetReplicationBridge()->GetReplicatedObject(ServerObject->NetRefHandle));
 
 	UE_NET_ASSERT_TRUE(ClientObjectThatWillBeTornOff != nullptr);
 	UE_NET_ASSERT_EQ(ServerObject->IntA, ClientObjectThatWillBeTornOff->IntA);
@@ -1180,7 +1180,7 @@ UE_NET_TEST_FIXTURE(FReplicationSystemServerClientTestFixture, TestTearOffResend
 	ServerObject->IntA = 2;
 
 	// TearOff the object
-	Server->ReplicationSystem->TearOffNextUpdate(ServerObject->NetHandle);
+	Server->ReplicationSystem->TearOffNextUpdate(ServerObject->NetRefHandle);
 
 	// Send and deliver packet, in this case the packet containing 2 was lost, but, we did not know that when we 
 	Server->PreSendUpdate();
@@ -1197,7 +1197,7 @@ UE_NET_TEST_FIXTURE(FReplicationSystemServerClientTestFixture, TestTearOffResend
 
 	// Verify that ClientObject is torn-off and that the expected final state was applied
 	UE_NET_ASSERT_EQ(2, ClientObjectThatWillBeTornOff->IntA);
-	UE_NET_ASSERT_TRUE(Cast<UTestReplicatedIrisObject>(Client->GetReplicationBridge()->GetReplicatedObject(ServerObject->NetHandle)) == nullptr);
+	UE_NET_ASSERT_TRUE(Cast<UTestReplicatedIrisObject>(Client->GetReplicationBridge()->GetReplicatedObject(ServerObject->NetRefHandle)) == nullptr);
 }
 
 // Test TearOff does not pickup statechanges after tear off
@@ -1220,7 +1220,7 @@ UE_NET_TEST_FIXTURE(FReplicationSystemServerClientTestFixture, TestTornedOffObje
 	Server->PostSendUpdate();
 
 	// Store Pointer to object 
-	UTestReplicatedIrisObject* ClientObjectThatWillBeTornOff = Cast<UTestReplicatedIrisObject>(Client->GetReplicationBridge()->GetReplicatedObject(ServerObject->NetHandle));
+	UTestReplicatedIrisObject* ClientObjectThatWillBeTornOff = Cast<UTestReplicatedIrisObject>(Client->GetReplicationBridge()->GetReplicatedObject(ServerObject->NetRefHandle));
 
 	UE_NET_ASSERT_TRUE(ClientObjectThatWillBeTornOff != nullptr);
 	UE_NET_ASSERT_EQ(ServerObject->IntA, ClientObjectThatWillBeTornOff->IntA);
@@ -1229,7 +1229,7 @@ UE_NET_TEST_FIXTURE(FReplicationSystemServerClientTestFixture, TestTornedOffObje
 	ServerObject->IntA = 2;
 
 	// TearOff the object
-	Server->ReplicationSystem->TearOffNextUpdate(ServerObject->NetHandle);
+	Server->ReplicationSystem->TearOffNextUpdate(ServerObject->NetRefHandle);
 
 	// Send and drop packet containing the value 2
 	Server->PreSendUpdate();
@@ -1246,7 +1246,7 @@ UE_NET_TEST_FIXTURE(FReplicationSystemServerClientTestFixture, TestTornedOffObje
 
 	// Verify that ClientObject is TornOFf and that the expected final state was applied
 	UE_NET_ASSERT_EQ(2, ClientObjectThatWillBeTornOff->IntA);
-	UE_NET_ASSERT_TRUE(Cast<UTestReplicatedIrisObject>(Client->GetReplicationBridge()->GetReplicatedObject(ServerObject->NetHandle)) == nullptr);
+	UE_NET_ASSERT_TRUE(Cast<UTestReplicatedIrisObject>(Client->GetReplicationBridge()->GetReplicatedObject(ServerObject->NetRefHandle)) == nullptr);
 }
 
 // Test TearOff and SubObjects, SubObjects must apply state?
@@ -1261,7 +1261,7 @@ UE_NET_TEST_FIXTURE(FReplicationSystemServerClientTestFixture, TestImmediateTear
 	UTestReplicatedIrisObject* ServerObject = Server->CreateObject(0,0);
 
 	// Spawn second object on server as a subobject
-	UTestReplicatedIrisObject* ServerSubObject = Server->CreateSubObject(ServerObject->NetHandle, 0, 0);
+	UTestReplicatedIrisObject* ServerSubObject = Server->CreateSubObject(ServerObject->NetRefHandle, 0, 0);
 
 	// Trigger replication
 	ServerObject->IntA = 1;
@@ -1273,11 +1273,11 @@ UE_NET_TEST_FIXTURE(FReplicationSystemServerClientTestFixture, TestImmediateTear
 	Server->PostSendUpdate();
 
 	// Store Pointer to objects
-	UTestReplicatedIrisObject* ClientObjectThatWillBeTornOff = Cast<UTestReplicatedIrisObject>(Client->GetReplicationBridge()->GetReplicatedObject(ServerObject->NetHandle));
+	UTestReplicatedIrisObject* ClientObjectThatWillBeTornOff = Cast<UTestReplicatedIrisObject>(Client->GetReplicationBridge()->GetReplicatedObject(ServerObject->NetRefHandle));
 	UE_NET_ASSERT_TRUE(ClientObjectThatWillBeTornOff != nullptr);
 	UE_NET_ASSERT_EQ(ServerObject->IntA, ClientObjectThatWillBeTornOff->IntA);
 
-	UTestReplicatedIrisObject* ClientSubObjectThatWillBeTornOff = Cast<UTestReplicatedIrisObject>(Client->GetReplicationBridge()->GetReplicatedObject(ServerSubObject->NetHandle));
+	UTestReplicatedIrisObject* ClientSubObjectThatWillBeTornOff = Cast<UTestReplicatedIrisObject>(Client->GetReplicationBridge()->GetReplicatedObject(ServerSubObject->NetRefHandle));
 	UE_NET_ASSERT_TRUE(ClientSubObjectThatWillBeTornOff != nullptr);
 	UE_NET_ASSERT_EQ(ServerSubObject->IntA, ClientSubObjectThatWillBeTornOff->IntA);
 
@@ -1294,7 +1294,7 @@ UE_NET_TEST_FIXTURE(FReplicationSystemServerClientTestFixture, TestImmediateTear
 
 	// Verify that ClientObject is torn-off and that the final state was applied to subObject 
 	UE_NET_ASSERT_EQ(ServerSubObject->IntA, ClientSubObjectThatWillBeTornOff->IntA);
-	UE_NET_ASSERT_TRUE(Cast<UTestReplicatedIrisObject>(Client->GetReplicationBridge()->GetReplicatedObject(ServerSubObject->NetHandle)) == nullptr);
+	UE_NET_ASSERT_TRUE(Cast<UTestReplicatedIrisObject>(Client->GetReplicationBridge()->GetReplicatedObject(ServerSubObject->NetRefHandle)) == nullptr);
 }
 
 // Test TearOff and SubObjects, SubObjects must apply state?
@@ -1309,7 +1309,7 @@ UE_NET_TEST_FIXTURE(FReplicationSystemServerClientTestFixture, TestImmediateTear
 	UTestReplicatedIrisObject* ServerObject = Server->CreateObject(0,0);
 
 	// Spawn second object on server as a subobject
-	UTestReplicatedIrisObject* ServerSubObject = Server->CreateSubObject(ServerObject->NetHandle, 0, 0);
+	UTestReplicatedIrisObject* ServerSubObject = Server->CreateSubObject(ServerObject->NetRefHandle, 0, 0);
 
 	// Trigger replication
 	ServerObject->IntA = 1;
@@ -1321,11 +1321,11 @@ UE_NET_TEST_FIXTURE(FReplicationSystemServerClientTestFixture, TestImmediateTear
 	Server->PostSendUpdate();
 
 	// Store Pointer to objects
-	UTestReplicatedIrisObject* ClientObjectThatWillBeTornOff = Cast<UTestReplicatedIrisObject>(Client->GetReplicationBridge()->GetReplicatedObject(ServerObject->NetHandle));
+	UTestReplicatedIrisObject* ClientObjectThatWillBeTornOff = Cast<UTestReplicatedIrisObject>(Client->GetReplicationBridge()->GetReplicatedObject(ServerObject->NetRefHandle));
 	UE_NET_ASSERT_TRUE(ClientObjectThatWillBeTornOff != nullptr);
 	UE_NET_ASSERT_EQ(ServerObject->IntA, ClientObjectThatWillBeTornOff->IntA);
 
-	UTestReplicatedIrisObject* ClientSubObjectThatWillBeTornOff = Cast<UTestReplicatedIrisObject>(Client->GetReplicationBridge()->GetReplicatedObject(ServerSubObject->NetHandle));
+	UTestReplicatedIrisObject* ClientSubObjectThatWillBeTornOff = Cast<UTestReplicatedIrisObject>(Client->GetReplicationBridge()->GetReplicatedObject(ServerSubObject->NetRefHandle));
 	UE_NET_ASSERT_TRUE(ClientSubObjectThatWillBeTornOff != nullptr);
 	UE_NET_ASSERT_EQ(ServerSubObject->IntA, ClientSubObjectThatWillBeTornOff->IntA);
 
@@ -1347,7 +1347,7 @@ UE_NET_TEST_FIXTURE(FReplicationSystemServerClientTestFixture, TestImmediateTear
 
 	// Verify that ClientObject is torn-off and that the final state was applied to subObject 
 	UE_NET_ASSERT_EQ(ServerSubObject->IntA, ClientSubObjectThatWillBeTornOff->IntA);
-	UE_NET_ASSERT_TRUE(Cast<UTestReplicatedIrisObject>(Client->GetReplicationBridge()->GetReplicatedObject(ServerSubObject->NetHandle)) == nullptr);
+	UE_NET_ASSERT_TRUE(Cast<UTestReplicatedIrisObject>(Client->GetReplicationBridge()->GetReplicatedObject(ServerSubObject->NetRefHandle)) == nullptr);
 }
 
 // Test TearOff and SubObjects, SubObjects must apply state?
@@ -1362,7 +1362,7 @@ UE_NET_TEST_FIXTURE(FReplicationSystemServerClientTestFixture, TestTearOffNextUp
 	UTestReplicatedIrisObject* ServerObject = Server->CreateObject(0,0);
 
 	// Spawn second object on server as a subobject
-	UTestReplicatedIrisObject* ServerSubObject = Server->CreateSubObject(ServerObject->NetHandle, 0, 0);
+	UTestReplicatedIrisObject* ServerSubObject = Server->CreateSubObject(ServerObject->NetRefHandle, 0, 0);
 
 	// Trigger replication
 	ServerObject->IntA = 1;
@@ -1374,11 +1374,11 @@ UE_NET_TEST_FIXTURE(FReplicationSystemServerClientTestFixture, TestTearOffNextUp
 	Server->PostSendUpdate();
 
 	// Store Pointer to objects
-	UTestReplicatedIrisObject* ClientObjectThatWillBeTornOff = Cast<UTestReplicatedIrisObject>(Client->GetReplicationBridge()->GetReplicatedObject(ServerObject->NetHandle));
+	UTestReplicatedIrisObject* ClientObjectThatWillBeTornOff = Cast<UTestReplicatedIrisObject>(Client->GetReplicationBridge()->GetReplicatedObject(ServerObject->NetRefHandle));
 	UE_NET_ASSERT_TRUE(ClientObjectThatWillBeTornOff != nullptr);
 	UE_NET_ASSERT_EQ(ServerObject->IntA, ClientObjectThatWillBeTornOff->IntA);
 
-	UTestReplicatedIrisObject* ClientSubObjectThatWillBeTornOff = Cast<UTestReplicatedIrisObject>(Client->GetReplicationBridge()->GetReplicatedObject(ServerSubObject->NetHandle));
+	UTestReplicatedIrisObject* ClientSubObjectThatWillBeTornOff = Cast<UTestReplicatedIrisObject>(Client->GetReplicationBridge()->GetReplicatedObject(ServerSubObject->NetRefHandle));
 	UE_NET_ASSERT_TRUE(ClientSubObjectThatWillBeTornOff != nullptr);
 	UE_NET_ASSERT_EQ(ServerSubObject->IntA, ClientSubObjectThatWillBeTornOff->IntA);
 
@@ -1386,7 +1386,7 @@ UE_NET_TEST_FIXTURE(FReplicationSystemServerClientTestFixture, TestTearOffNextUp
 	ServerSubObject->IntA = 2;
 
 	// TearOff the object using immediate tear-off
-	Server->ReplicationSystem->TearOffNextUpdate(ServerObject->NetHandle);
+	Server->ReplicationSystem->TearOffNextUpdate(ServerObject->NetRefHandle);
 
 	// Send and deliver packet
 	Server->PreSendUpdate();
@@ -1395,14 +1395,14 @@ UE_NET_TEST_FIXTURE(FReplicationSystemServerClientTestFixture, TestTearOffNextUp
 
 	// Verify that ClientObject is torn-off and that the final state was applied to subObject 
 	UE_NET_ASSERT_EQ(ServerSubObject->IntA, ClientSubObjectThatWillBeTornOff->IntA);
-	UE_NET_ASSERT_TRUE(Cast<UTestReplicatedIrisObject>(Client->GetReplicationBridge()->GetReplicatedObject(ServerSubObject->NetHandle)) == nullptr);
+	UE_NET_ASSERT_TRUE(Cast<UTestReplicatedIrisObject>(Client->GetReplicationBridge()->GetReplicatedObject(ServerSubObject->NetRefHandle)) == nullptr);
 }
 
 // Test TearOff and destroy of SubObjects that are still pending create/tearoff
 UE_NET_TEST_FIXTURE(FReplicationSystemServerClientTestFixture, TestTearOffNextUpdateExistingObjectWithSubObjectPendingCreation)
 {
 	UReplicationSystem* ReplicationSystem = Server->ReplicationSystem;
-	FNetHandleManager* NetHandleManager = &ReplicationSystem->GetReplicationSystemInternal()->GetNetHandleManager();
+	FNetRefHandleManager* NetRefHandleManager = &ReplicationSystem->GetReplicationSystemInternal()->GetNetRefHandleManager();
 
 	// Add a client
 	FReplicationSystemTestClient* Client = CreateClient();
@@ -1411,36 +1411,36 @@ UE_NET_TEST_FIXTURE(FReplicationSystemServerClientTestFixture, TestTearOffNextUp
 	UTestReplicatedIrisObject* ServerObject = Server->CreateObject(0,0);
 
 	// Spawn second object on server as a subobject
-	UTestReplicatedIrisObject* ServerSubObject = Server->CreateSubObject(ServerObject->NetHandle, 0, 0);
+	UTestReplicatedIrisObject* ServerSubObject = Server->CreateSubObject(ServerObject->NetRefHandle, 0, 0);
 
 	// Trigger replication
 	ServerObject->IntA = 1;
 	ServerSubObject->IntA = 1;
 
-	const FInternalNetHandle ServerObjectInternalIndex = NetHandleManager->GetInternalIndex(ServerObject->NetHandle);
-	const FInternalNetHandle SubObjectObjectInternalIndex = NetHandleManager->GetInternalIndex(ServerSubObject->NetHandle);
+	const FInternalNetRefIndex ServerObjectInternalIndex = NetRefHandleManager->GetInternalIndex(ServerObject->NetRefHandle);
+	const FInternalNetRefIndex SubObjectObjectInternalIndex = NetRefHandleManager->GetInternalIndex(ServerSubObject->NetRefHandle);
 
 	// Trigger presend without send to add the objects to scope
 	Server->PreSendUpdate();
 
-	UE_NET_ASSERT_EQ(uint16(1), NetHandleManager->GetNetObjectRefCount(ServerObjectInternalIndex));
-	UE_NET_ASSERT_EQ(uint16(1), NetHandleManager->GetNetObjectRefCount(SubObjectObjectInternalIndex));
+	UE_NET_ASSERT_EQ(uint16(1), NetRefHandleManager->GetNetObjectRefCount(ServerObjectInternalIndex));
+	UE_NET_ASSERT_EQ(uint16(1), NetRefHandleManager->GetNetObjectRefCount(SubObjectObjectInternalIndex));
 
 	// TearOff the object this will also tear-off subobject
-	Server->ReplicationSystem->TearOffNextUpdate(ServerObject->NetHandle);
+	Server->ReplicationSystem->TearOffNextUpdate(ServerObject->NetRefHandle);
 
 	// Update logic, object should be removed from scope but still exist as pending create in
 	Server->PreSendUpdate();
 
-	UE_NET_ASSERT_EQ(uint16(1), NetHandleManager->GetNetObjectRefCount(ServerObjectInternalIndex));
-	UE_NET_ASSERT_EQ(uint16(1), NetHandleManager->GetNetObjectRefCount(SubObjectObjectInternalIndex));
+	UE_NET_ASSERT_EQ(uint16(1), NetRefHandleManager->GetNetObjectRefCount(ServerObjectInternalIndex));
+	UE_NET_ASSERT_EQ(uint16(1), NetRefHandleManager->GetNetObjectRefCount(SubObjectObjectInternalIndex));
 
 	// Destroy the object
 	Server->DestroyObject(ServerObject);
 
 	// Verify that we no longer have any references to the object
-	UE_NET_ASSERT_EQ(uint16(0), NetHandleManager->GetNetObjectRefCount(ServerObjectInternalIndex));
-	UE_NET_ASSERT_EQ(uint16(0), NetHandleManager->GetNetObjectRefCount(SubObjectObjectInternalIndex));
+	UE_NET_ASSERT_EQ(uint16(0), NetRefHandleManager->GetNetObjectRefCount(ServerObjectInternalIndex));
+	UE_NET_ASSERT_EQ(uint16(0), NetRefHandleManager->GetNetObjectRefCount(SubObjectObjectInternalIndex));
 	
 	Server->SendAndDeliverTo(Client, true);
 	Server->PostSendUpdate();
@@ -1456,7 +1456,7 @@ UE_NET_TEST_FIXTURE(FReplicationSystemServerClientTestFixture, TestReplicatedObj
 
 	// Spawn object on server
 	UTestReplicatedIrisObjectWithNoReplicatedMembers* ServerObject = Server->CreateObject<UTestReplicatedIrisObjectWithNoReplicatedMembers>();
-	const FNetHandle ServerHandle = ServerObject->NetHandle;
+	const FNetRefHandle ServerHandle = ServerObject->NetRefHandle;
 
 	UE_NET_ASSERT_TRUE(ServerHandle.IsValid());
 
@@ -1500,8 +1500,8 @@ UE_NET_TEST_FIXTURE(FReplicationSystemServerClientTestFixture, TestObjectPollFra
 	Server->PostSendUpdate();
 
 	// Store Pointer to objects and verify state after initial replication
-	UTestReplicatedIrisObject* ClientObject = Cast<UTestReplicatedIrisObject>(Client->GetReplicationBridge()->GetReplicatedObject(ServerObject->NetHandle));
-	UTestReplicatedIrisObject* ClientObjectPolledEveryOtherFrame = Cast<UTestReplicatedIrisObject>(Client->GetReplicationBridge()->GetReplicatedObject(ServerObjectPolledEveryOtherFrame->NetHandle));
+	UTestReplicatedIrisObject* ClientObject = Cast<UTestReplicatedIrisObject>(Client->GetReplicationBridge()->GetReplicatedObject(ServerObject->NetRefHandle));
+	UTestReplicatedIrisObject* ClientObjectPolledEveryOtherFrame = Cast<UTestReplicatedIrisObject>(Client->GetReplicationBridge()->GetReplicatedObject(ServerObjectPolledEveryOtherFrame->NetRefHandle));
 
 	UE_NET_ASSERT_NE(ClientObjectPolledEveryOtherFrame, nullptr);
 	UE_NET_ASSERT_NE(ClientObject, nullptr);
diff --git a/Engine/Plugins/Runtime/ReplicationSystemTestPlugin/Source/Private/Tests/Serialization/TestObjectNetSerializer.cpp b/Engine/Plugins/Runtime/ReplicationSystemTestPlugin/Source/Private/Tests/Serialization/TestObjectNetSerializer.cpp
--- a/Engine/Plugins/Runtime/ReplicationSystemTestPlugin/Source/Private/Tests/Serialization/TestObjectNetSerializer.cpp
+++ b/Engine/Plugins/Runtime/ReplicationSystemTestPlugin/Source/Private/Tests/Serialization/TestObjectNetSerializer.cpp
@@ -37,12 +37,12 @@ UE_NET_TEST_FIXTURE(FReplicationSystemServerClientTestFixture, TestObjectReferen
 	Server->SendAndDeliverTo(Client, true);
 	Server->PostSendUpdate();
 
-	// Check state, we cheat a bit and use the server nethandle for lookup
+	// Check state, we cheat a bit and use the server NetRefHandle for lookup
 	// THIS will fail if we start using full handle compares
-	auto ClientObjectA = Client->GetObjectAs<UTestReplicatedIrisObjectWithObjectReference>(ObjectA->NetHandle);
-	auto ClientObjectB = Client->GetObjectAs<UTestReplicatedIrisObjectWithObjectReference>(ObjectB->NetHandle);
-	auto ClientObjectC = Client->GetObjectAs<UTestReplicatedIrisObjectWithObjectReference>(ObjectC->NetHandle);
-	auto ClientObjectD = Client->GetObjectAs<UTestReplicatedIrisObjectWithObjectReference>(ObjectD->NetHandle);
+	auto ClientObjectA = Client->GetObjectAs<UTestReplicatedIrisObjectWithObjectReference>(ObjectA->NetRefHandle);
+	auto ClientObjectB = Client->GetObjectAs<UTestReplicatedIrisObjectWithObjectReference>(ObjectB->NetRefHandle);
+	auto ClientObjectC = Client->GetObjectAs<UTestReplicatedIrisObjectWithObjectReference>(ObjectC->NetRefHandle);
+	auto ClientObjectD = Client->GetObjectAs<UTestReplicatedIrisObjectWithObjectReference>(ObjectD->NetRefHandle);
 	
 	UE_NET_ASSERT_TRUE(ClientObjectA != nullptr);
 	UE_NET_ASSERT_TRUE(ClientObjectB != nullptr);
@@ -77,12 +77,12 @@ UE_NET_TEST_FIXTURE(FReplicationSystemServerClientTestFixture, TestCircularRefer
 	Server->SendAndDeliverTo(Client, true);
 	Server->PostSendUpdate();
 
-	// Check state, we cheat a bit and use the server nethandle for lookup
+	// Check state, we cheat a bit and use the server NetRefHandle for lookup
 	// THIS will fail if we start using full handle compares
-	auto ClientObjectA = Client->GetObjectAs<UTestReplicatedIrisObjectWithObjectReference>(ObjectA->NetHandle);
-	auto ClientObjectB = Client->GetObjectAs<UTestReplicatedIrisObjectWithObjectReference>(ObjectB->NetHandle);
-	auto ClientObjectC = Client->GetObjectAs<UTestReplicatedIrisObjectWithObjectReference>(ObjectC->NetHandle);
-	auto ClientObjectD = Client->GetObjectAs<UTestReplicatedIrisObjectWithObjectReference>(ObjectD->NetHandle);
+	auto ClientObjectA = Client->GetObjectAs<UTestReplicatedIrisObjectWithObjectReference>(ObjectA->NetRefHandle);
+	auto ClientObjectB = Client->GetObjectAs<UTestReplicatedIrisObjectWithObjectReference>(ObjectB->NetRefHandle);
+	auto ClientObjectC = Client->GetObjectAs<UTestReplicatedIrisObjectWithObjectReference>(ObjectC->NetRefHandle);
+	auto ClientObjectD = Client->GetObjectAs<UTestReplicatedIrisObjectWithObjectReference>(ObjectD->NetRefHandle);
 	
 	UE_NET_ASSERT_TRUE(ClientObjectA != nullptr);
 	UE_NET_ASSERT_TRUE(ClientObjectB != nullptr);
@@ -140,12 +140,12 @@ UE_NET_TEST_FIXTURE(FReplicationSystemServerClientTestFixture, TestCircularRefer
 	Server->SendAndDeliverTo(Client, true);
 	Server->PostSendUpdate();
 
-	// Check state, we cheat a bit and use the server nethandle for lookup
+	// Check state, we cheat a bit and use the server NetRefHandle for lookup
 	// THIS will fail if we start using full handle compares
-	auto ClientObjectA = Client->GetObjectAs<UTestReplicatedIrisObjectWithObjectReference>(ObjectA->NetHandle);
-	auto ClientObjectB = Client->GetObjectAs<UTestReplicatedIrisObjectWithObjectReference>(ObjectB->NetHandle);
-	auto ClientObjectC = Client->GetObjectAs<UTestReplicatedIrisObjectWithObjectReference>(ObjectC->NetHandle);
-	auto ClientObjectD = Client->GetObjectAs<UTestReplicatedIrisObjectWithObjectReference>(ObjectD->NetHandle);
+	auto ClientObjectA = Client->GetObjectAs<UTestReplicatedIrisObjectWithObjectReference>(ObjectA->NetRefHandle);
+	auto ClientObjectB = Client->GetObjectAs<UTestReplicatedIrisObjectWithObjectReference>(ObjectB->NetRefHandle);
+	auto ClientObjectC = Client->GetObjectAs<UTestReplicatedIrisObjectWithObjectReference>(ObjectC->NetRefHandle);
+	auto ClientObjectD = Client->GetObjectAs<UTestReplicatedIrisObjectWithObjectReference>(ObjectD->NetRefHandle);
 	
 	UE_NET_ASSERT_TRUE(ClientObjectA != nullptr);
 	UE_NET_ASSERT_TRUE(ClientObjectB != nullptr);
diff --git a/Engine/Plugins/Runtime/ReplicationSystemTestPlugin/Source/Private/Tests/Serialization/TestPolymorphicStructNetSerializer.cpp b/Engine/Plugins/Runtime/ReplicationSystemTestPlugin/Source/Private/Tests/Serialization/TestPolymorphicStructNetSerializer.cpp
--- a/Engine/Plugins/Runtime/ReplicationSystemTestPlugin/Source/Private/Tests/Serialization/TestPolymorphicStructNetSerializer.cpp
+++ b/Engine/Plugins/Runtime/ReplicationSystemTestPlugin/Source/Private/Tests/Serialization/TestPolymorphicStructNetSerializer.cpp
@@ -718,7 +718,7 @@ UE_NET_TEST_FIXTURE(FTestPolymorphicArrayStructNetSerializerFixture, TestPolymor
 	Server->SendAndDeliverTo(Client, DeliverPacket);
 	Server->PostSendUpdate();
 	
-	auto ClientObject = Client->GetObjectAs<UTestPolymorphicStructNetSerializer_TestObject>(ServerObject->NetHandle);
+	auto ClientObject = Client->GetObjectAs<UTestPolymorphicStructNetSerializer_TestObject>(ServerObject->NetRefHandle);
 	UE_NET_ASSERT_NE(ClientObject, nullptr);
 	CA_ASSUME(ClientObject != nullptr);
 	UE_NET_ASSERT_EQ(ClientObject->PolyStruct.GetAs<FExamplePolymorphicStructB>().SomeFloat, 12.0f);
@@ -778,7 +778,7 @@ UE_NET_TEST_FIXTURE(FTestPolymorphicArrayStructNetSerializerFixture, TestPolymor
 	Server->SendAndDeliverTo(Client, DeliverPacket);
 	Server->PostSendUpdate();
 	
-	auto ClientObject = Client->GetObjectAs<UTestPolymorphicStructNetSerializer_TestObject>(ServerObject->NetHandle);
+	auto ClientObject = Client->GetObjectAs<UTestPolymorphicStructNetSerializer_TestObject>(ServerObject->NetRefHandle);
 	UE_NET_ASSERT_NE(ClientObject, nullptr);
 	CA_ASSUME(ClientObject != nullptr);
 	UE_NET_ASSERT_EQ(ClientObject->PolyStructFastArray.GetItemArray().Num(), ServerObject->PolyStructFastArray.GetItemArray().Num());
diff --git a/Engine/Plugins/Runtime/ReplicationSystemTestPlugin/Source/Private/Tests/Serialization/TestScriptInterfaceNetSerializer.cpp b/Engine/Plugins/Runtime/ReplicationSystemTestPlugin/Source/Private/Tests/Serialization/TestScriptInterfaceNetSerializer.cpp
--- a/Engine/Plugins/Runtime/ReplicationSystemTestPlugin/Source/Private/Tests/Serialization/TestScriptInterfaceNetSerializer.cpp
+++ b/Engine/Plugins/Runtime/ReplicationSystemTestPlugin/Source/Private/Tests/Serialization/TestScriptInterfaceNetSerializer.cpp
@@ -34,7 +34,7 @@ UE_NET_TEST_FIXTURE(FTestScriptInterfaceNetSerializer, TestNoScriptInterface)
 	Server->PostSendUpdate();
 
 	// Verify object was created
-	UTestObjectReferencingScriptInterface* ClientObject = Client->GetObjectAs<UTestObjectReferencingScriptInterface>(Object->NetHandle);	
+	UTestObjectReferencingScriptInterface* ClientObject = Client->GetObjectAs<UTestObjectReferencingScriptInterface>(Object->NetRefHandle);	
 	UE_NET_ASSERT_NE(ClientObject, nullptr);
 
 	// Verify the interface objects are identical, i.e. not set.
@@ -61,11 +61,11 @@ UE_NET_TEST_FIXTURE(FTestScriptInterfaceNetSerializer, TestInterfaceToReplicated
 	Server->PostSendUpdate();
 
 	// Verify object was created
-	UTestObjectReferencingScriptInterface* ClientObject = Client->GetObjectAs<UTestObjectReferencingScriptInterface>(Object->NetHandle);
+	UTestObjectReferencingScriptInterface* ClientObject = Client->GetObjectAs<UTestObjectReferencingScriptInterface>(Object->NetRefHandle);
 	UE_NET_ASSERT_NE(ClientObject, nullptr);
 
 	// Verify interface object was created
-	UTestScriptInterfaceReplicatedObject* ClientInterfaceObject = Client->GetObjectAs<UTestScriptInterfaceReplicatedObject>(InterfaceObject->NetHandle);
+	UTestScriptInterfaceReplicatedObject* ClientInterfaceObject = Client->GetObjectAs<UTestScriptInterfaceReplicatedObject>(InterfaceObject->NetRefHandle);
 	UE_NET_ASSERT_NE(ClientInterfaceObject, nullptr);
 
 	// Verify the interface object is set
@@ -93,11 +93,11 @@ UE_NET_TEST_FIXTURE(FTestScriptInterfaceNetSerializer, TestInterfaceToNamedObjec
 	Server->PostSendUpdate();
 
 	// Verify object was created
-	UTestObjectReferencingScriptInterface* ClientObject = Client->GetObjectAs<UTestObjectReferencingScriptInterface>(Object->NetHandle);
+	UTestObjectReferencingScriptInterface* ClientObject = Client->GetObjectAs<UTestObjectReferencingScriptInterface>(Object->NetRefHandle);
 	UE_NET_ASSERT_NE(ClientObject, nullptr);
 
 	// Verify the object with the interface subobject was created
-	UTestScriptInterfaceReplicatedObjectWithDefaultSubobject* ClientObjectWithInterfaceSubObject = Client->GetObjectAs<UTestScriptInterfaceReplicatedObjectWithDefaultSubobject>(ObjectWithInterfaceSubObject->NetHandle);
+	UTestScriptInterfaceReplicatedObjectWithDefaultSubobject* ClientObjectWithInterfaceSubObject = Client->GetObjectAs<UTestScriptInterfaceReplicatedObjectWithDefaultSubobject>(ObjectWithInterfaceSubObject->NetRefHandle);
 	UE_NET_ASSERT_NE(ClientObjectWithInterfaceSubObject, nullptr);
 
 	// Verify the interface object is set
diff --git a/Engine/Source/Runtime/Engine/Private/Actor.cpp b/Engine/Source/Runtime/Engine/Private/Actor.cpp
--- a/Engine/Source/Runtime/Engine/Private/Actor.cpp
+++ b/Engine/Source/Runtime/Engine/Private/Actor.cpp
@@ -53,6 +53,7 @@
 #include "Net/Core/PushModel/PushModel.h"
 #include "Engine/AutoDestroySubsystem.h"
 #if UE_WITH_IRIS
+#include "Iris/ReplicationSystem/ObjectReplicationBridge.h"
 #include "Iris/ReplicationSystem/ReplicationSystem.h"
 #include "Iris/ReplicationSystem/Conditionals/ReplicationCondition.h"
 #include "Net/Iris/ReplicationSystem/ReplicationSystemUtil.h"
@@ -3985,12 +3986,15 @@ void AActor::SetAutonomousProxy(const bool bInAutonomousProxy, const bool bAllow
 			// Update autonomous role condition
 			if (UReplicationSystem* ReplicationSystem = UE::Net::FReplicationSystemUtil::GetReplicationSystem(GetNetOwner()))
 			{
-				const UE::Net::FNetHandle NetHandle = UE::Net::FReplicationSystemUtil::GetNetHandle(this);
-				if (NetHandle.IsValid())
+				if (UObjectReplicationBridge* Bridge = ReplicationSystem->GetReplicationBridgeAs<UObjectReplicationBridge>())
 				{
-					const uint32 OwningNetConnectionId = ReplicationSystem->GetOwningNetConnection(NetHandle);
-					const bool bEnableAutonomousCondition = RemoteRole == ROLE_AutonomousProxy;
-					ReplicationSystem->SetReplicationConditionConnectionFilter(NetHandle, UE::Net::EReplicationCondition::RoleAutonomous, OwningNetConnectionId, bEnableAutonomousCondition);
+					const UE::Net::FNetRefHandle RefHandle = Bridge->GetReplicatedRefHandle(this);
+					if (RefHandle.IsValid())
+					{
+						const uint32 OwningNetConnectionId = ReplicationSystem->GetOwningNetConnection(RefHandle);
+						const bool bEnableAutonomousCondition = RemoteRole == ROLE_AutonomousProxy;
+						ReplicationSystem->SetReplicationConditionConnectionFilter(RefHandle, UE::Net::EReplicationCondition::RoleAutonomous, OwningNetConnectionId, bEnableAutonomousCondition);
+					}
 				}
 			}
 #endif // UE_WITH_IRIS
diff --git a/Engine/Source/Runtime/Engine/Private/ActorReplication.cpp b/Engine/Source/Runtime/Engine/Private/ActorReplication.cpp
--- a/Engine/Source/Runtime/Engine/Private/ActorReplication.cpp
+++ b/Engine/Source/Runtime/Engine/Private/ActorReplication.cpp
@@ -893,12 +893,8 @@ void AActor::RegisterReplicationFragments(UE::Net::FFragmentRegistrationContext&
 
 void AActor::BeginReplication(const FActorBeginReplicationParams& Params)
 {
-	// If we do not have a handle we create one
-	UE::Net::FNetHandle ActorHandle = UE::Net::FReplicationSystemUtil::BeginReplication(this, Params);
-	if (ActorHandle.IsValid())
-	{
-		UpdateOwningNetConnection();
-	}
+	UE::Net::FReplicationSystemUtil::BeginReplication(this, Params);
+	UpdateOwningNetConnection();
 }
 
 void AActor::BeginReplication()
@@ -918,7 +914,8 @@ void AActor::UpdateOwningNetConnection() const
 	using namespace UE::Net;
 
 	UReplicationSystem* ReplicationSystem = FReplicationSystemUtil::GetReplicationSystem(GetNetOwner());
-	if (ReplicationSystem == nullptr)
+	UObjectReplicationBridge* ObjectReplicationBridge = (ReplicationSystem ? ReplicationSystem->GetReplicationBridgeAs<UObjectReplicationBridge>() : nullptr);
+	if (ObjectReplicationBridge == nullptr)
 	{
 		return;
 	}
@@ -942,8 +939,7 @@ void AActor::UpdateOwningNetConnection() const
 #endif
 
 		FNetHandle NetHandle = FReplicationSystemUtil::GetNetHandle(this);
-		const uint32 CurrentOwningNetConnectionId = NetHandle.IsValid() ? ReplicationSystem->GetOwningNetConnection(NetHandle) : 0U;
-		bUpdateChildren = (NewOwningNetConnectionId != CurrentOwningNetConnectionId);
+		bUpdateChildren = NetHandle.IsValid();
 	}
 
 	if (!bUpdateChildren)
@@ -966,15 +962,15 @@ void AActor::UpdateOwningNetConnection() const
 
 		if (Actor->GetIsReplicated())
 		{
-			const FNetHandle NetHandle = FReplicationSystemUtil::GetNetHandle(Actor);
-			if (NetHandle.IsValid())
+			const UE::Net::FNetRefHandle RefHandle = ObjectReplicationBridge->GetReplicatedRefHandle(this);
+			if (RefHandle.IsValid())
 			{
-				ReplicationSystem->SetOwningNetConnection(NetHandle, NewOwningNetConnectionId);
+				ReplicationSystem->SetOwningNetConnection(RefHandle, NewOwningNetConnectionId);
 				// Update autonomous proxy condition
 				if (Actor->GetRemoteRole() == ROLE_AutonomousProxy)
 				{
 					const bool bEnableAutonomousCondition = true;
-					ReplicationSystem->SetReplicationConditionConnectionFilter(NetHandle, EReplicationCondition::RoleAutonomous, NewOwningNetConnectionId, bEnableAutonomousCondition);
+					ReplicationSystem->SetReplicationConditionConnectionFilter(RefHandle, EReplicationCondition::RoleAutonomous, NewOwningNetConnectionId, bEnableAutonomousCondition);
 				}
 			}
 		}
@@ -990,8 +986,7 @@ void AActor::UpdateReplicatePhysicsCondition()
 	const FNetHandle NetHandle = FReplicationSystemUtil::GetNetHandle(this);
 	if (NetHandle.IsValid())
 	{
-		UReplicationSystem* ReplicationSystem = FReplicationSystemUtil::GetReplicationSystem(this);
-		ReplicationSystem->SetReplicationCondition(NetHandle, EReplicationCondition::ReplicatePhysics, GetReplicatedMovement().bRepPhysics);
+		FReplicationSystemUtil::SetReplicationCondition(NetHandle, EReplicationCondition::ReplicatePhysics, GetReplicatedMovement().bRepPhysics);
 	}
 #endif // UE_WITH_IRIS
 }
diff --git a/Engine/Source/Runtime/Engine/Private/Net/Iris/ReplicationSystem/ActorReplicationBridge.cpp b/Engine/Source/Runtime/Engine/Private/Net/Iris/ReplicationSystem/ActorReplicationBridge.cpp
--- a/Engine/Source/Runtime/Engine/Private/Net/Iris/ReplicationSystem/ActorReplicationBridge.cpp
+++ b/Engine/Source/Runtime/Engine/Private/Net/Iris/ReplicationSystem/ActorReplicationBridge.cpp
@@ -24,6 +24,7 @@
 #include "Engine/Level.h"
 #include "Engine/NetConnection.h"
 #include "Engine/NetDriver.h"
+#include "Engine/Level.h"
 #include "Engine/World.h"
 #include "GameFramework/Actor.h"
 #include "Net/DataBunch.h"
@@ -33,11 +34,7 @@
 #include "UObject/Package.h"
 #include <limits>
 
-#if UE_NET_ALLOW_MULTIPLE_REPLICATION_SYSTEMS
 #define UE_LOG_ACTORREPLICATIONBRIDGE(Category, Format, ...)	UE_LOG(LogIrisBridge, Category, TEXT("ActorReplicationBridge(%u)::") Format, GetReplicationSystem()->GetId(), ##__VA_ARGS__)
-#else
-#define UE_LOG_ACTORREPLICATIONBRIDGE(Category, Format, ...)	UE_LOG(LogIrisBridge, Category, TEXT("ActorReplicationBridge::") Format, ##__VA_ARGS__)
-#endif
 
 extern bool GDefaultUseSubObjectReplicationList;
 
@@ -71,7 +68,7 @@ bool IsActorValidForIrisReplication(const AActor* Actor)
 	return IsValid(Actor) && !Actor->IsActorBeingDestroyed() && !Actor->IsUnreachable();
 }
 
-void ActorReplicationBridgePreUpdateFunction(FNetHandle Handle, UObject* Instance, const UReplicationBridge* Bridge)
+void ActorReplicationBridgePreUpdateFunction(FNetRefHandle Handle, UObject* Instance, const UReplicationBridge* Bridge)
 {
 	AActor* Actor = Cast<AActor>(Instance);
 	if (IsActorValidForIrisReplication(Actor))
@@ -85,7 +82,7 @@ void ActorReplicationBridgePreUpdateFunction(FNetHandle Handle, UObject* Instanc
 	}
 }
 
-FVector ActorReplicationBridgeGetActorWorldLocation(FNetHandle Handle, const UObject* Instance)
+FVector ActorReplicationBridgeGetActorWorldLocation(FNetRefHandle Handle, const UObject* Instance)
 {
 	if (const AActor* Actor = Cast<AActor>(Instance))
 	{
@@ -165,78 +162,78 @@ UActorReplicationBridge::~UActorReplicationBridge()
 	}
 }
 
-UE::Net::FNetHandle UActorReplicationBridge::BeginReplication(AActor* Actor, const FActorBeginReplicationParams& Params)
+UE::Net::FNetRefHandle UActorReplicationBridge::BeginReplication(AActor* Actor, const FActorBeginReplicationParams& Params)
 {
 	using namespace UE::Net;
 
 	if (!ShouldUseIrisReplication(Actor))
 	{
-		return FNetHandle();
+		return FNetRefHandle();
 	}
 
 	if (!ensureMsgf(Actor == nullptr || !(Actor->HasAnyFlags(RF_ClassDefaultObject | RF_ArchetypeObject)), TEXT("Actor %s is a CDO or Archetype and should not be replicated."), ToCStr(GetFullNameSafe(Actor))))
 	{
-		return FNetHandle();
+		return FNetRefHandle();
 	}
 
 	const bool bIsNetActor = ULevel::IsNetActor(Actor);
 	if (!bIsNetActor)
 	{
 		UE_LOG_ACTORREPLICATIONBRIDGE(VeryVerbose, TEXT("Actor %s doesn't have a NetRole."), ToCStr(GetFullNameSafe(Actor)));
-		return FNetHandle();
+		return FNetRefHandle();
 	}
 
 	if (Actor->GetLocalRole() != ROLE_Authority)
 	{
 		UE_LOG_ACTORREPLICATIONBRIDGE(VeryVerbose, TEXT("Actor %s NetRole isn't Authority."), ToCStr(GetFullNameSafe(Actor)));
-		return FNetHandle();
+		return FNetRefHandle();
 	}
 
 	if (Actor->IsActorBeingDestroyed() || !IsValid(Actor) || Actor->IsUnreachable())
 	{
 		UE_LOG_ACTORREPLICATIONBRIDGE(Verbose, TEXT("Actor %s is being destroyed or unreachable and can't be replicated."), ToCStr(GetFullNameSafe(Actor)));
-		return FNetHandle();
+		return FNetRefHandle();
 	}
 
 	if (!Actor->GetIsReplicated())
 	{
 		UE_LOG_ACTORREPLICATIONBRIDGE(Verbose, TEXT("Actor %s is not supposed to be replicated."), ToCStr(GetFullNameSafe(Actor)));
-		return FNetHandle();
+		return FNetRefHandle();
 	}
 
 	if (Actor->GetTearOff())
 	{
 		UE_LOG_ACTORREPLICATIONBRIDGE(Verbose, TEXT("Actor %s is torn off and should not be replicated."), ToCStr(GetFullNameSafe(Actor)));
-		return FNetHandle();
+		return FNetRefHandle();
 	}
 
 	if (!Actor->IsActorInitialized())
 	{
 		UE_LOG_ACTORREPLICATIONBRIDGE(Warning, TEXT("Actor %s is not initialized and won't be replicated."), ToCStr(GetFullNameSafe(Actor)));
-		return FNetHandle();
+		return FNetRefHandle();
 	}
 
 	if (!NetDriver)
 	{
 		UE_LOG_ACTORREPLICATIONBRIDGE(VeryVerbose, TEXT("There's no NetDriver so nothing can be replicated."));
-		return FNetHandle();
+		return FNetRefHandle();
 	}
 
 
 	if (!NetDriver->ShouldReplicateActor(Actor))
 	{
-		UE_LOG_ACTORREPLICATIONBRIDGE(Verbose, TEXT("Actor %s doesn't want to replicate with NetDriver %s."), ToCStr(GetFullNameSafe(Actor)), ToCStr(NetDriver->GetName()));
-		return FNetHandle();
+		UE_LOG_ACTORREPLICATIONBRIDGE(VeryVerbose, TEXT("Actor %s doesn't want to replicate with NetDriver %s."), ToCStr(GetFullNameSafe(Actor)), ToCStr(NetDriver->GetName()));
+		return FNetRefHandle();
 	}
 
 	// Initially dormant actors begin replication when their dormancy is flushed
 	const ENetDormancy Dormancy = Actor->NetDormancy;	
 	if (Actor->IsNetStartupActor() && (Dormancy == DORM_Initial))
 	{
-		return FNetHandle();
+		return FNetRefHandle();
 	}
 
-	FNetHandle ExistingHandle = GetReplicatedHandle(Actor);
+	FNetRefHandle ExistingHandle = GetReplicatedRefHandle(Actor);
 	if (ExistingHandle.IsValid())
 	{
 		return ExistingHandle;
@@ -249,32 +246,32 @@ UE::Net::FNetHandle UActorReplicationBridge::BeginReplication(AActor* Actor, con
 		UE_LOG_ACTORREPLICATIONBRIDGE(Warning, TEXT("Actor %s does not replicate subobjects using the registered SubObjectsLists, SubObjects will not replicate properly"), ToCStr(GetFullNameSafe(Actor)));
 	}
 
-	// Create NetHandle for the registered fragments
-	Super::FCreateNetHandleParams CreateNetHandleParams = DefaultCreateNetHandleParams;
-	CreateNetHandleParams.bNeedsPreUpdate = 1U;
-	CreateNetHandleParams.bNeedsWorldLocationUpdate = 1U;
-	CreateNetHandleParams.StaticPriority = (Actor->bAlwaysRelevant || Actor->bOnlyRelevantToOwner) ? Actor->NetPriority : 0.0f;
+	// Create handles for the registered fragments
+	Super::FCreateNetRefHandleParams CreateNetRefHandleParams = DefaultCreateNetRefHandleParams;
+	CreateNetRefHandleParams.bNeedsPreUpdate = 1U;
+	CreateNetRefHandleParams.bNeedsWorldLocationUpdate = 1U;
+	CreateNetRefHandleParams.StaticPriority = (Actor->bAlwaysRelevant || Actor->bOnlyRelevantToOwner) ? Actor->NetPriority : 0.0f;
 	const uint32 UnclampedPollPeriod = GetPollFramePeriod(Actor->NetUpdateFrequency);
 	const uint32 ClampedPollPeriod = FMath::Clamp<uint32>(UnclampedPollPeriod, 1U, static_cast<uint32>(std::numeric_limits<uint8>::max()) + 1U);
-	CreateNetHandleParams.PollFramePeriod = (ClampedPollPeriod - 1U) & 255U;
+	CreateNetRefHandleParams.PollFramePeriod = (ClampedPollPeriod - 1U) & 255U;
 #if !UE_BUILD_SHIPPING
 	ensureAlwaysMsgf(!UE::Net::Private::bIrisEnableLowNetUpdateFrequencyEnsure || (ClampedPollPeriod >= UnclampedPollPeriod), TEXT("Very low NetUpdateFrequency %f for Actor %s. Suggest setting it to %f or higher."), Actor->NetUpdateFrequency, ToCStr(Actor->GetName()), GetMinSupportedNetUpdateFrequency());
-	ensureAlwaysMsgf(!(Actor->bAlwaysRelevant || Actor->bOnlyRelevantToOwner) || CreateNetHandleParams.StaticPriority >= 1.0f, TEXT("Very low NetPriority %.02f for always relevant or owner relevant Actor %s. Set it to 1.0f or higher."), Actor->NetPriority, ToCStr(Actor->GetName()));
+	ensureAlwaysMsgf(!(Actor->bAlwaysRelevant || Actor->bOnlyRelevantToOwner) || CreateNetRefHandleParams.StaticPriority >= 1.0f, TEXT("Very low NetPriority %.02f for always relevant or owner relevant Actor %s. Set it to 1.0f or higher."), Actor->NetPriority, ToCStr(Actor->GetName()));
 #endif
 
-	FNetHandle ActorHandle = Super::BeginReplication(Actor, CreateNetHandleParams);
+	FNetRefHandle ActorRefHandle = Super::BeginReplication(Actor, CreateNetRefHandleParams);
 
-	if (!ActorHandle.IsValid())
+	if (!ActorRefHandle.IsValid())
 	{
-		ensureMsgf(false, TEXT("Failed to create NetHandle for Actor Named %s"), ToCStr(Actor->GetName()));
-		return FNetHandle();
+		ensureMsgf(false, TEXT("Failed to create NetRefHandle for Actor Named %s"), ToCStr(Actor->GetName()));
+		return FNetRefHandle();
 	}
 
 	// Set owning connection filtering if actor is only relevant to owner
 	{
 		if (Actor->bOnlyRelevantToOwner && !Actor->bAlwaysRelevant)
 		{
-			GetReplicationSystem()->SetFilter(ActorHandle, ToOwnerFilterHandle);
+			GetReplicationSystem()->SetFilter(ActorRefHandle, ToOwnerFilterHandle);
 		}
 	}
 
@@ -282,7 +279,7 @@ UE::Net::FNetHandle UActorReplicationBridge::BeginReplication(AActor* Actor, con
 	{
 		if (Actor->bNetTemporary)
 		{
-			GetReplicationSystem()->SetIsNetTemporary(ActorHandle);
+			GetReplicationSystem()->SetIsNetTemporary(ActorRefHandle);
 		}
 	}
 
@@ -290,7 +287,7 @@ UE::Net::FNetHandle UActorReplicationBridge::BeginReplication(AActor* Actor, con
 	{
 		if (Dormancy > DORM_Awake)
 		{
-			SetObjectWantsToBeDormant(ActorHandle, true);
+			SetObjectWantsToBeDormant(ActorRefHandle, true);
 		}
 	}
 
@@ -314,8 +311,8 @@ UE::Net::FNetHandle UActorReplicationBridge::BeginReplication(AActor* Actor, con
 			}
 
 			// Add object to group
-			UE_LOG_ACTORREPLICATIONBRIDGE(Verbose, TEXT("Added %s to GroupIndex: %u Level: %s"), *ActorHandle.ToString(), LevelGroup, ToCStr(PackageName.ToString()));
-			GetReplicationSystem()->AddToGroup(LevelGroup, ActorHandle);	
+			UE_LOG_ACTORREPLICATIONBRIDGE(Verbose, TEXT("Added %s to GroupIndex: %u Level: %s"), *ActorRefHandle.ToString(), LevelGroup, ToCStr(PackageName.ToString()));
+			GetReplicationSystem()->AddToGroup(LevelGroup, ActorRefHandle);	
 		}
 	}
 
@@ -330,10 +327,10 @@ UE::Net::FNetHandle UActorReplicationBridge::BeginReplication(AActor* Actor, con
 		{
 			if (IsValid(SubObjectInfo.SubObject) && SubObjectInfo.NetCondition != ELifetimeCondition::COND_Never)
 			{
-				FNetHandle SubObjectHandle = UObjectReplicationBridge::BeginReplication(ActorHandle, SubObjectInfo.SubObject);
-				if (SubObjectHandle.IsValid() && SubObjectInfo.NetCondition != ELifetimeCondition::COND_None)
+				FNetRefHandle SubObjectRefHandle = UObjectReplicationBridge::BeginReplication(ActorRefHandle, SubObjectInfo.SubObject);
+				if (SubObjectRefHandle.IsValid() && SubObjectInfo.NetCondition != ELifetimeCondition::COND_None)
 				{
-					UObjectReplicationBridge::SetSubObjectNetCondition(SubObjectHandle, SubObjectInfo.NetCondition);
+					UObjectReplicationBridge::SetSubObjectNetCondition(SubObjectRefHandle, SubObjectInfo.NetCondition);
 				}
 			}
 		}
@@ -350,21 +347,21 @@ UE::Net::FNetHandle UActorReplicationBridge::BeginReplication(AActor* Actor, con
 		}
 	}
 
-	return ActorHandle;
+	return ActorRefHandle;
 }
 
-UE::Net::FNetHandle UActorReplicationBridge::BeginReplication(FNetHandle OwnerHandle, UActorComponent* SubObject)
+UE::Net::FNetRefHandle UActorReplicationBridge::BeginReplication(FNetRefHandle OwnerHandle, UActorComponent* SubObject)
 {
 	using namespace UE::Net;
 
 	if (!OwnerHandle.IsValid() || !ShouldUseIrisReplication(SubObject))
 	{
-		return FNetHandle();
+		return FNetRefHandle();
 	}
 
 	AActor* Owner = SubObject->GetOwner();
 
-	FNetHandle ReplicatedComponentHandle = GetReplicatedHandle(SubObject);
+	FNetRefHandle ReplicatedComponentHandle = GetReplicatedRefHandle(SubObject);
 	const FReplicatedComponentInfo* RepComponentInfo = FSubObjectRegistryGetter::GetReplicatedComponentInfoForComponent(Owner, SubObject);
 
 	if (!ReplicatedComponentHandle.IsValid())
@@ -375,7 +372,7 @@ UE::Net::FNetHandle UActorReplicationBridge::BeginReplication(FNetHandle OwnerHa
 			|| SubObject->HasAnyFlags(RF_ArchetypeObject | RF_ClassDefaultObject)
 		)
 		{
-			return FNetHandle();
+			return FNetRefHandle();
 		}
 
 		if (!SubObject->IsUsingRegisteredSubObjectList())
@@ -385,17 +382,17 @@ UE::Net::FNetHandle UActorReplicationBridge::BeginReplication(FNetHandle OwnerHa
 
 		if (RepComponentInfo == nullptr || RepComponentInfo->NetCondition == ELifetimeCondition::COND_Never)
 		{
-			return FNetHandle();
+			return FNetRefHandle();
 		}
 
-		// Create NetHandle and attach for the sub object
+		// Start replicating the subobject with its owner.
 		ReplicatedComponentHandle = Super::BeginReplication(OwnerHandle, SubObject);
 	}
 
 	if (!ReplicatedComponentHandle.IsValid())
 	{
-		ensureMsgf(false, TEXT("Failed to create or find NetHandle for ActorComponent Named %s"), ToCStr(SubObject->GetName()));
-		return FNetHandle();
+		ensureMsgf(false, TEXT("Failed to create or find NetRefHandle for ActorComponent Named %s"), ToCStr(SubObject->GetName()));
+		return FNetRefHandle();
 	}
 
 	// Update or set any conditionals
@@ -409,7 +406,7 @@ UE::Net::FNetHandle UActorReplicationBridge::BeginReplication(FNetHandle OwnerHa
 	{
 		if (IsValid(SubObjectInfo.SubObject) && SubObjectInfo.NetCondition != ELifetimeCondition::COND_Never)
 		{
-			FNetHandle SubObjectHandle = UObjectReplicationBridge::BeginReplication(OwnerHandle, SubObjectInfo.SubObject, ReplicatedComponentHandle, UReplicationBridge::ESubObjectInsertionOrder::ReplicateWith);
+			FNetRefHandle SubObjectHandle = UObjectReplicationBridge::BeginReplication(OwnerHandle, SubObjectInfo.SubObject, ReplicatedComponentHandle, UReplicationBridge::ESubObjectInsertionOrder::ReplicateWith);
 			if (SubObjectHandle.IsValid() && SubObjectInfo.NetCondition != ELifetimeCondition::COND_None)
 			{
 				SetSubObjectNetCondition(SubObjectHandle, SubObjectInfo.NetCondition);
@@ -424,19 +421,19 @@ void UActorReplicationBridge::EndReplication(AActor* Actor, EEndPlayReason::Type
 {
 	using namespace UE::Net;
 
-	FNetHandle NetHandle = GetReplicatedHandle(Actor);
-	if (NetHandle.IsValid())
+	FNetRefHandle RefHandle = GetReplicatedRefHandle(Actor);
+	if (RefHandle.IsValid())
 	{
-		UE_LOG(LogIrisBridge, Verbose, TEXT("EndReplication for %s %s. Reason %s "), *Actor->GetName(), *NetHandle.ToString(), *UEnum::GetValueAsString(TEXT("Engine.EEndPlayReason"), EndPlayReason));
+		UE_LOG(LogIrisBridge, Verbose, TEXT("EndReplication for %s %s. Reason %s "), *Actor->GetName(), *RefHandle.ToString(), *UEnum::GetValueAsString(TEXT("Engine.EEndPlayReason"), EndPlayReason));
 	
 		EEndReplicationFlags Flags = EEndReplicationFlags::None;
 		const bool bShouldDestroyObject = EndPlayReason == EEndPlayReason::Destroyed;		
 		if (bShouldDestroyObject)
 		{ 
-			Flags |= EEndReplicationFlags::Destroy;
+			Flags |= EEndReplicationFlags::Destroy | EEndReplicationFlags::DestroyNetHandle | EEndReplicationFlags::ClearNetPushId;
 		}
 					
-		const bool bShouldCreateDestructionInfo = NetHandle.IsStatic() && bShouldDestroyObject && GetReplicationSystem()->IsServer();
+		const bool bShouldCreateDestructionInfo = RefHandle.IsStatic() && bShouldDestroyObject && GetReplicationSystem()->IsServer();
 		if (bShouldCreateDestructionInfo)
 		{
 			UObjectReplicationBridge::FEndReplicationParameters EndReplicationParameters;
@@ -445,36 +442,36 @@ void UActorReplicationBridge::EndReplication(AActor* Actor, EEndPlayReason::Type
 			EndReplicationParameters.Level = Actor->GetLevel();
 			EndReplicationParameters.bUseDistanceBasedPrioritization = Actor->bAlwaysRelevant == false;
 
-			UObjectReplicationBridge::EndReplication(NetHandle, Flags, &EndReplicationParameters);		
+			UObjectReplicationBridge::EndReplication(RefHandle, Flags, &EndReplicationParameters);		
 		}
 		else
 		{
-			UObjectReplicationBridge::EndReplication(NetHandle, Flags, nullptr);
+			UObjectReplicationBridge::EndReplication(RefHandle, Flags, nullptr);
 		}
 	}
 }
 
-void UActorReplicationBridge::EndReplicationForActorComponent(UActorComponent* ActorComponent)
+void UActorReplicationBridge::EndReplicationForActorComponent(UActorComponent* ActorComponent, EEndReplicationFlags EndReplicationFlags)
 {
 	using namespace UE::Net;
 
-	FNetHandle ComponentHandle = GetReplicatedHandle(ActorComponent);
-	AActor* Actor = ActorComponent->GetOwner();
+	FNetRefHandle ComponentHandle = GetReplicatedRefHandle(ActorComponent);
+	const AActor* Actor = ActorComponent->GetOwner();
 	if (ComponentHandle.IsValid() && Actor)
 	{
 		UE_LOG(LogIrisBridge, Verbose, TEXT("EndReplicationForActorComponent for %s %s."), *ActorComponent->GetName(), *ComponentHandle.ToString());
 		
-		UObjectReplicationBridge::EndReplication(ComponentHandle, EEndReplicationFlags::None, nullptr);
+		UObjectReplicationBridge::EndReplication(ComponentHandle, EndReplicationFlags, nullptr);
 	}
 }
 
-bool UActorReplicationBridge::WriteCreationHeader(UE::Net::FNetSerializationContext& Context, FNetHandle Handle)
+bool UActorReplicationBridge::WriteCreationHeader(UE::Net::FNetSerializationContext& Context, FNetRefHandle Handle)
 {
 	using namespace UE::Net;
 	using namespace UE::Net::Private;
 
 	FNetBitStreamWriter* Writer = Context.GetBitStreamWriter();
-
+	ensureAlways(IsReplicatedHandle(Handle));
 	const UObject* Object = GetReplicatedObject(Handle);
 	if (const AActor* Actor = Cast<AActor>(Object))
 	{
@@ -503,7 +500,7 @@ bool UActorReplicationBridge::WriteCreationHeader(UE::Net::FNetSerializationCont
 		return !Writer->IsOverflown();
 	}
 
-	ensureAlwaysMsgf(false, TEXT("UActorReplicationBridge::WriteCreationHeader Failed to write creationHeader for NetHandleIndex: %u"), Handle.GetId());
+	ensureAlwaysMsgf(false, TEXT("UActorReplicationBridge::WriteCreationHeader Failed to write creationHeader for NetRefHandle (Id=%u)"), Handle.GetId());
 
 	return false;
 }
@@ -539,7 +536,7 @@ UObjectReplicationBridge::FCreationHeader* UActorReplicationBridge::ReadCreation
 	return nullptr;
 }
 
-UObject* UActorReplicationBridge::BeginInstantiateFromRemote(FNetHandle SubObjectOwnerNetHandle, const UE::Net::FNetObjectResolveContext& ResolveContext, const UObjectReplicationBridge::FCreationHeader* InHeader)
+UObject* UActorReplicationBridge::BeginInstantiateFromRemote(FNetRefHandle SubObjectOwnerNetHandle, const UE::Net::FNetObjectResolveContext& ResolveContext, const UObjectReplicationBridge::FCreationHeader* InHeader)
 {
 	using namespace UE::Net::Private;
 
@@ -718,7 +715,7 @@ bool UActorReplicationBridge::OnInstantiatedFromRemote(UObject* Instance, const
 	return true;
 }
 
-void UActorReplicationBridge::EndInstantiateFromRemote(FNetHandle Handle)
+void UActorReplicationBridge::EndInstantiateFromRemote(FNetRefHandle Handle)
 {
 	if (AActor* Actor = Cast<AActor>(GetReplicatedObject(Handle)))
 	{
@@ -771,7 +768,7 @@ void UActorReplicationBridge::DestroyInstanceFromRemote(UObject* Instance, bool
 	}
 }
 
-void UActorReplicationBridge::GetInitialDependencies(FNetHandle Handle, FNetDependencyInfoArray& OutDependencies) const
+void UActorReplicationBridge::GetInitialDependencies(FNetRefHandle Handle, FNetDependencyInfoArray& OutDependencies) const
 {
 	using namespace UE::Net;
 
