Iris - Multiple replication systems tests.

diff --git a/Engine/Plugins/Runtime/ReplicationSystemTestPlugin/Source/Private/NetworkAutomationTest.h b/Engine/Plugins/Runtime/ReplicationSystemTestPlugin/Source/Private/NetworkAutomationTest.h
--- a/Engine/Plugins/Runtime/ReplicationSystemTestPlugin/Source/Private/NetworkAutomationTest.h
+++ b/Engine/Plugins/Runtime/ReplicationSystemTestPlugin/Source/Private/NetworkAutomationTest.h
@@ -4,7 +4,6 @@
 
 #include "CoreTypes.h"
 #include "Misc/AutomationTest.h"
-#include "ReplicationSystemTestPlugin/NetworkAutomationTest.h"
 
 namespace UE::Net
 {
diff --git a/Engine/Plugins/Runtime/ReplicationSystemTestPlugin/Source/Private/TestMessage.cpp b/Engine/Plugins/Runtime/ReplicationSystemTestPlugin/Source/Private/TestMessage.cpp
new file mode 100644
--- /dev/null
+++ b/Engine/Plugins/Runtime/ReplicationSystemTestPlugin/Source/Private/TestMessage.cpp
@@ -0,0 +1,17 @@
+// Copyright Epic Games, Inc. All Rights Reserved.
+
+#include "TestMessage.h"
+#include "Misc/StringBuilder.h"
+#include "Net/Core/NetHandle/NetHandle.h"
+
+namespace UE::Net
+{
+
+FTestMessage& operator<<(FTestMessage& Message, const FNetHandle& NetHandle)
+{
+	TStringBuilder<64> StringBuilder;
+	StringBuilder << NetHandle;
+	return Message << StringBuilder;
+}
+
+}
diff --git a/Engine/Plugins/Runtime/ReplicationSystemTestPlugin/Source/Private/TestMessage.h b/Engine/Plugins/Runtime/ReplicationSystemTestPlugin/Source/Private/TestMessage.h
--- a/Engine/Plugins/Runtime/ReplicationSystemTestPlugin/Source/Private/TestMessage.h
+++ b/Engine/Plugins/Runtime/ReplicationSystemTestPlugin/Source/Private/TestMessage.h
@@ -7,6 +7,11 @@
 #include "Math/Vector.h"
 #include <cinttypes>
 
+namespace UE::Net
+{
+	class FNetHandle;
+}
+
 namespace UE::Net
 {
 
@@ -19,6 +24,7 @@ public:
 	inline FTestMessage& operator<<(const FString& InString) { String.Append(InString); return *this; }
 	inline FTestMessage& operator<<(const TCHAR* InString) { String.Append(InString); return *this; }
 	inline FTestMessage& operator<<(TCHAR Char) { String.AppendChar(Char); return *this; }
+	inline FTestMessage& operator<<(FStringBuilderBase& StringBuilder) { String.Append(StringBuilder); return *this; }
 
 	inline FTestMessage& operator<<(const FTestMessage& Message) { String.Append(Message.String); return *this; }
 
@@ -33,7 +39,7 @@ inline FTestMessage& operator<<(FTestMessage& Message, int Value)
 
 inline FTestMessage& operator<<(FTestMessage& Message, unsigned Value)
 {
-	return Message << TStringBuilder<32>().Appendf(TEXT("%u"), Value).ToString();
+	return Message << TStringBuilder<32>().Appendf(TEXT("%u"), Value);
 }
 
 inline FTestMessage& operator<<(FTestMessage& Message, int8 Value)
@@ -125,6 +131,8 @@ inline FTestMessage& operator<<(FTestMessage& Message, const FVector& Vector)
 	return Message << Vector.ToString();
 }
 
+FTestMessage& operator<<(FTestMessage& Message, const FNetHandle& NetHandle);
+
 #if defined(_WIN32) && !defined(_WIN64)
 inline FTestMessage& operator<<(FTestMessage& Message, SIZE_T Value)
 {
diff --git a/Engine/Plugins/Runtime/ReplicationSystemTestPlugin/Source/Private/Tests/ReplicationSystem/MultiReplicationSystemsTestFixture.cpp b/Engine/Plugins/Runtime/ReplicationSystemTestPlugin/Source/Private/Tests/ReplicationSystem/MultiReplicationSystemsTestFixture.cpp
new file mode 100644
--- /dev/null
+++ b/Engine/Plugins/Runtime/ReplicationSystemTestPlugin/Source/Private/Tests/ReplicationSystem/MultiReplicationSystemsTestFixture.cpp
@@ -0,0 +1,167 @@
+// Copyright Epic Games, Inc. All Rights Reserved.
+
+#include "Tests/ReplicationSystem/MultiReplicationSystemsTestFixture.h"
+#include "Containers/Array.h"
+#include "Containers/Map.h"
+
+namespace UE::Net
+{
+
+void FMultiReplicationSystemsTestFixture::SetUp()
+{
+	FNetworkAutomationTestSuiteFixture::SetUp();
+
+	// Fake what we normally get from config
+	DataStreamUtil.SetUp();
+	DataStreamUtil.AddDataStreamDefinition(TEXT("NetToken"), TEXT("/Script/IrisCore.NetTokenDataStream"));
+	DataStreamUtil.AddDataStreamDefinition(TEXT("Replication"), TEXT("/Script/IrisCore.ReplicationDataStream"));
+	DataStreamUtil.FixupDefinitions();
+}
+
+void FMultiReplicationSystemsTestFixture::TearDown()
+{
+	for (const TMap<FReplicationSystemTestServer*, TArray<FReplicationSystemTestClient*>>::ElementType& ServerAndClient : ServerClients)
+	{
+		FReplicationSystemTestServer* Server = ServerAndClient.Key;
+		delete Server;
+		for (FReplicationSystemTestClient* Client : ServerAndClient.Value)
+		{
+			delete Client;
+		}
+	}
+
+	Servers.Empty();
+	ServerClients.Empty();
+
+	DataStreamUtil.TearDown();
+
+	FNetworkAutomationTestSuiteFixture::TearDown();
+}
+
+void FMultiReplicationSystemsTestFixture::CreateServers(unsigned ServerCount)
+{
+	for (unsigned It = 0, EndIt = ServerCount; It != EndIt; ++It)
+	{
+		CreateServer();
+	}
+}
+
+void FMultiReplicationSystemsTestFixture::CreateSomeServers()
+{
+	CreateServers(DefaultServerCount);
+}
+
+TArrayView<FReplicationSystemTestServer*> FMultiReplicationSystemsTestFixture::GetAllServers()
+{
+	return MakeArrayView(Servers);
+}
+
+FReplicationSystemTestServer* FMultiReplicationSystemsTestFixture::CreateServer()
+{
+	FReplicationSystemTestServer* Server = new FReplicationSystemTestServer(GetName());
+	Servers.Add(Server);
+	ServerClients.Emplace(Server);
+
+	return Server;
+}
+
+FReplicationSystemTestClient* FMultiReplicationSystemsTestFixture::CreateClientForServer(FReplicationSystemTestServer* Server)
+{
+	if (TMap<FReplicationSystemTestServer*, TArray<FReplicationSystemTestClient*>>::ValueType* Clients = ServerClients.Find(Server))
+	{
+		FReplicationSystemTestClient* Client = new FReplicationSystemTestClient(GetName()); 
+		Clients->Add(Client);
+
+		// The client needs a connection
+		Client->LocalConnectionId = Client->AddConnection();
+
+		// Auto connect to server
+		Client->ConnectionIdOnServer = Server->AddConnection();
+
+		return Client;
+	}
+
+	return nullptr;
+}
+
+TArrayView<FReplicationSystemTestClient*> FMultiReplicationSystemsTestFixture::GetClients(FReplicationSystemTestServer* Server)
+{
+	if (TMap<FReplicationSystemTestServer*, TArray<FReplicationSystemTestClient*>>::ValueType* Clients = ServerClients.Find(Server))
+	{
+		return MakeArrayView(*Clients);
+	}
+
+	return TArrayView<FReplicationSystemTestClient*>();
+}
+
+UTestReplicatedIrisObject* FMultiReplicationSystemsTestFixture::CreateObject(const UTestReplicatedIrisObject::FComponents& Components)
+{
+	UTestReplicatedIrisObject* CreatedObject = NewObject<UTestReplicatedIrisObject>();
+	CreatedObject->AddComponents(Components);
+
+	CreatedObjects.Add(TStrongObjectPtr<UObject>(CreatedObject));
+	return CreatedObject;
+}
+
+void FMultiReplicationSystemsTestFixture::BeginReplication(UTestReplicatedIrisObject* Object)
+{
+	for (FReplicationSystemTestServer* Server : GetAllServers())
+	{
+		if (UReplicatedTestObjectBridge* ReplicationBridge = Server->GetReplicationBridge())
+		{
+			ReplicationBridge->BeginReplication(Object);
+		}
+	}
+}
+
+void FMultiReplicationSystemsTestFixture::EndReplication(UTestReplicatedIrisObject* Object)
+{
+	for (FReplicationSystemTestServer* Server : GetAllServers())
+	{
+		if (UReplicatedTestObjectBridge* ReplicationBridge = Server->GetReplicationBridge())
+		{
+			constexpr EEndReplicationFlags EndReplicationFlags = EEndReplicationFlags::Destroy | EEndReplicationFlags::DestroyNetHandle | EEndReplicationFlags::ClearNetPushId;
+			ReplicationBridge->EndReplication(Object, EndReplicationFlags);
+		}
+	}
+}
+
+void FMultiReplicationSystemsTestFixture::BeginReplication(FReplicationSystemTestServer* Server, UTestReplicatedIrisObject* Object)
+{
+	if (UReplicatedTestObjectBridge* ReplicationBridge = Server->GetReplicationBridge())
+	{
+		ReplicationBridge->BeginReplication(Object);
+	}
+}
+
+void FMultiReplicationSystemsTestFixture::EndReplication(FReplicationSystemTestServer* Server, UTestReplicatedIrisObject* Object)
+{
+	if (UReplicatedTestObjectBridge* ReplicationBridge = Server->GetReplicationBridge())
+	{
+		ReplicationBridge->EndReplication(Object);
+	}
+}
+
+void FMultiReplicationSystemsTestFixture::FullSendAndDeliverUpdate()
+{
+	for (FReplicationSystemTestServer* Server : Servers)
+	{
+		Server->PreSendUpdate();
+	}
+
+	for (TMap<FReplicationSystemTestServer*, TArray<FReplicationSystemTestClient*>>::ElementType& ServerWithClients : ServerClients)
+	{
+		FReplicationSystemTestServer* Server = ServerWithClients.Key;
+		for (FReplicationSystemTestClient* Client : ServerWithClients.Value)
+		{
+			Server->SendAndDeliverTo(Client, DeliverPacket);
+		}
+	}
+
+	for (FReplicationSystemTestServer* Server : Servers)
+	{
+		Server->PostSendUpdate();
+	}
+}
+
+}
diff --git a/Engine/Plugins/Runtime/ReplicationSystemTestPlugin/Source/Private/Tests/ReplicationSystem/MultiReplicationSystemsTestFixture.h b/Engine/Plugins/Runtime/ReplicationSystemTestPlugin/Source/Private/Tests/ReplicationSystem/MultiReplicationSystemsTestFixture.h
new file mode 100644
--- /dev/null
+++ b/Engine/Plugins/Runtime/ReplicationSystemTestPlugin/Source/Private/Tests/ReplicationSystem/MultiReplicationSystemsTestFixture.h
@@ -0,0 +1,77 @@
+// Copyright Epic Games, Inc. All Rights Reserved.
+
+#pragma once
+
+#include "NetworkAutomationTest.h"
+#include "Tests/ReplicationSystem/ReplicationSystemServerClientTestFixture.h"
+#include "Containers/ContainersFwd.h"
+
+namespace UE::Net
+{
+
+class FMultiReplicationSystemsTestFixture : public FNetworkAutomationTestSuiteFixture
+{
+public:
+	FMultiReplicationSystemsTestFixture() : FNetworkAutomationTestSuiteFixture() {}
+
+protected:
+	enum EConstants : unsigned
+	{
+		/** It's common for things to stop working beyond 2. */
+		DefaultServerCount = 3U,
+	};
+
+	enum : bool
+	{
+		DoNotDeliverPacket = false,
+		DeliverPacket = true,
+	};
+
+protected:
+	virtual void SetUp() override;
+	virtual void TearDown() override;
+
+	/** Create a specific amount of servers. */
+	void CreateServers(unsigned ServerCount);
+
+	/** Creates the default amount of servers. */
+	void CreateSomeServers();
+
+	/** Returns a view of all created servers. */
+	TArrayView<FReplicationSystemTestServer*> GetAllServers();
+
+	/** Creates a server and returns it. */
+	FReplicationSystemTestServer* CreateServer();
+
+	/** Creates a client that is connected to a specific server. */
+	FReplicationSystemTestClient* CreateClientForServer(FReplicationSystemTestServer*);
+
+	/** Returns a view of the clients connected to the server. */
+	TArrayView<FReplicationSystemTestClient*> GetClients(FReplicationSystemTestServer*);
+
+	/** Creates a UTestReplicatedIrisObject. It will not begin replicating unless a BeginReplication method is called. */
+	UTestReplicatedIrisObject* CreateObject(const UTestReplicatedIrisObject::FComponents& Components);
+	
+	/** Begins replication of the object on all replication systems. */
+	void BeginReplication(UTestReplicatedIrisObject* Object);
+
+	/** Ends replication of the object on all replication systems. */
+	void EndReplication(UTestReplicatedIrisObject* Object);
+
+	/** Begins replicatiion of the object on a specific replication system. */
+	void BeginReplication(FReplicationSystemTestServer* Server, UTestReplicatedIrisObject* Object);
+
+	/** Ends replication of the object on a specific replication system. */
+	void EndReplication(FReplicationSystemTestServer* Server, UTestReplicatedIrisObject* Object);
+
+	/** Performs PreSendUpdate, SendAndDeliverTo and PostSendUpdate on all servers. Packets are delivered to all clients connected to the servers. */
+	void FullSendAndDeliverUpdate();
+
+private:
+	FDataStreamTestUtil DataStreamUtil;
+	TArray<FReplicationSystemTestServer*> Servers;
+	TMap<FReplicationSystemTestServer*, TArray<FReplicationSystemTestClient*>> ServerClients;
+	TArray<TStrongObjectPtr<UObject>> CreatedObjects;
+};
+
+}
diff --git a/Engine/Plugins/Runtime/ReplicationSystemTestPlugin/Source/Private/Tests/ReplicationSystem/TestMultiReplicationSystems.cpp b/Engine/Plugins/Runtime/ReplicationSystemTestPlugin/Source/Private/Tests/ReplicationSystem/TestMultiReplicationSystems.cpp
new file mode 100644
--- /dev/null
+++ b/Engine/Plugins/Runtime/ReplicationSystemTestPlugin/Source/Private/Tests/ReplicationSystem/TestMultiReplicationSystems.cpp
@@ -0,0 +1,214 @@
+// Copyright Epic Games, Inc. All Rights Reserved.
+
+#include "Tests/ReplicationSystem/MultiReplicationSystemsTestFixture.h"
+#include "Net/Core/NetHandle/NetHandleManager.h"
+
+namespace UE::Net::Private
+{
+
+UE_NET_TEST_FIXTURE(FMultiReplicationSystemsTestFixture, CanCreateMultipleReplicationSystems)
+{
+	CreateSomeServers();
+
+	TArrayView<FReplicationSystemTestServer*> ServersView = GetAllServers();
+	UE_NET_ASSERT_EQ(unsigned(ServersView.Num()), unsigned(DefaultServerCount));
+
+	const unsigned ExpectedReplicationSystemCount = DefaultServerCount;
+	unsigned ValidReplicationSystemCount = 0U;
+	for (FReplicationSystemTestServer* Server : ServersView)
+	{
+		if (UReplicationSystem* ReplicationSystem = Server->GetReplicationSystem())
+		{
+			++ValidReplicationSystemCount;
+		}
+	}
+
+	UE_NET_ASSERT_EQ(ValidReplicationSystemCount, ExpectedReplicationSystemCount);
+}
+
+UE_NET_TEST_FIXTURE(FMultiReplicationSystemsTestFixture, CanReplicateObjectOnMultipleReplicationSystems)
+{
+	CreateSomeServers();
+
+	UTestReplicatedIrisObject* Object = CreateObject(UTestReplicatedIrisObject::FComponents{});
+	BeginReplication(Object);
+
+	for (FReplicationSystemTestServer* Server : GetAllServers())
+	{
+		UReplicatedTestObjectBridge* ReplicationBridge = Server->GetReplicationBridge();
+		FNetRefHandle RefHandle = ReplicationBridge->GetReplicatedRefHandle(Object);
+		UE_NET_ASSERT_TRUE(RefHandle.IsValid());
+	}
+}
+
+UE_NET_TEST_FIXTURE(FMultiReplicationSystemsTestFixture, ReplicatedObjectIsAssignedGlobalNetHandle)
+{
+	CreateSomeServers();
+
+	UTestReplicatedIrisObject* Object = CreateObject(UTestReplicatedIrisObject::FComponents{});
+	BeginReplication(Object);
+
+	FNetHandle NetHandle = FNetHandleManager::GetNetHandle(Object);
+	UE_NET_ASSERT_TRUE(NetHandle.IsValid());
+}
+
+UE_NET_TEST_FIXTURE(FMultiReplicationSystemsTestFixture, ReplicatedObjectLosesGlobalNetHandleAfterEndReplication)
+{
+	CreateSomeServers();
+
+	UTestReplicatedIrisObject* Object = CreateObject(UTestReplicatedIrisObject::FComponents{});
+	BeginReplication(Object);
+	EndReplication(Object);
+
+	// The object should no longer be associated with a NetHandle when ending replication on all systems.
+	FNetHandle NetHandle = FNetHandleManager::GetNetHandle(Object);
+	UE_NET_ASSERT_FALSE(NetHandle.IsValid());
+}
+
+UE_NET_TEST_FIXTURE(FMultiReplicationSystemsTestFixture, ReplicatedObjectKeepsGlobalNetHandleAfterEndReplicationOnSingleSystem)
+{
+	CreateSomeServers();
+
+	UTestReplicatedIrisObject* Object = CreateObject(UTestReplicatedIrisObject::FComponents{});
+	BeginReplication(Object);
+
+	FNetHandle NetHandlePriorToEndReplication = FNetHandleManager::GetNetHandle(Object);
+
+	// End replication on a single system
+	FReplicationSystemTestServer* Server = GetAllServers()[0];
+	UReplicatedTestObjectBridge* ReplicationBridge = Server->GetReplicationBridge();
+	UE_NET_ASSERT_NE(ReplicationBridge, nullptr);
+	ReplicationBridge->EndReplication(Object);
+
+	// Make sure replication was ended on the system.
+	FNetRefHandle RefHandle = ReplicationBridge->GetReplicatedRefHandle(Object);
+	UE_NET_ASSERT_FALSE(RefHandle.IsValid());
+
+	// Validate there's still a global NetHandle assigned.
+	FNetHandle NetHandleAfterSingleEndReplication = FNetHandleManager::GetNetHandle(Object);
+	UE_NET_ASSERT_EQ(NetHandlePriorToEndReplication, NetHandleAfterSingleEndReplication);
+}
+
+UE_NET_TEST_FIXTURE(FMultiReplicationSystemsTestFixture, ObjectIsReplicatedToAllClientsOnAllSystems)
+{
+	CreateSomeServers();
+
+	// Create clients for all systems
+	for (FReplicationSystemTestServer* Server : GetAllServers())
+	{
+		constexpr unsigned ClientCountPerServer = 3;
+		for (unsigned ClientIt = 0, ClientEndIt = ClientCountPerServer; ClientIt != ClientEndIt; ++ClientIt)
+		{
+			CreateClientForServer(Server);
+		}
+	}
+
+	UTestReplicatedIrisObject* ServerObject = CreateObject(UTestReplicatedIrisObject::FComponents{});
+	BeginReplication(ServerObject);
+
+	FullSendAndDeliverUpdate();
+
+	// Verify the object was created on all clients
+	for (FReplicationSystemTestServer* Server : GetAllServers())
+	{
+		FNetRefHandle RefHandleOnServer = Server->GetReplicationBridge()->GetReplicatedRefHandle(ServerObject);
+		for (FReplicationSystemTestClient* Client : GetClients(Server))
+		{
+			UE_NET_ASSERT_NE(Client->GetReplicationBridge()->GetReplicatedObject(RefHandleOnServer), nullptr);
+		}
+	}
+}
+
+UE_NET_TEST_FIXTURE(FMultiReplicationSystemsTestFixture, ReplicatedObjectIsDestroyedOnAllClientsAfterEndReplication)
+{
+	CreateSomeServers();
+
+	// Create clients for all systems
+	for (FReplicationSystemTestServer* Server : GetAllServers())
+	{
+		constexpr unsigned ClientCountPerServer = 3;
+		for (unsigned ClientIt = 0, ClientEndIt = ClientCountPerServer; ClientIt != ClientEndIt; ++ClientIt)
+		{
+			CreateClientForServer(Server);
+		}
+	}
+
+	UTestReplicatedIrisObject* ServerObject = CreateObject(UTestReplicatedIrisObject::FComponents{});
+	BeginReplication(ServerObject);
+
+	TArray<FNetRefHandle> ServerRefHandles;
+	ServerRefHandles.Reserve(GetAllServers().Num());
+	for (FReplicationSystemTestServer* Server : GetAllServers())
+	{
+		FNetRefHandle RefHandle = Server->GetReplicationBridge()->GetReplicatedRefHandle(ServerObject);
+		ServerRefHandles.Add(RefHandle);
+	}
+
+	FullSendAndDeliverUpdate();
+
+	EndReplication(ServerObject);
+
+	FullSendAndDeliverUpdate();
+
+	// Verify the object was destroyed on all clients
+	unsigned ServerIt = 0;
+	for (FReplicationSystemTestServer* Server : GetAllServers())
+	{
+		FNetRefHandle RefHandleOnServer = ServerRefHandles[ServerIt++];
+		for (FReplicationSystemTestClient* Client : GetClients(Server))
+		{
+			UE_NET_ASSERT_EQ(Client->GetReplicationBridge()->GetReplicatedObject(RefHandleOnServer), nullptr);
+		}
+	}
+}
+
+UE_NET_TEST_FIXTURE(FMultiReplicationSystemsTestFixture, ReplicatedObjectCanStopReplicatingOnSingleSystem)
+{
+	CreateSomeServers();
+
+	// Create clients for all systems
+	for (FReplicationSystemTestServer* Server : GetAllServers())
+	{
+		CreateClientForServer(Server);
+	}
+
+	UTestReplicatedIrisObject* ServerObject = CreateObject(UTestReplicatedIrisObject::FComponents{});
+	BeginReplication(ServerObject);
+
+	FullSendAndDeliverUpdate();
+
+	// End replication on single system
+	constexpr SIZE_T SystemIndexToEndReplicationOn = 0;
+	FReplicationSystemTestServer* ServerToEndReplicationOn = GetAllServers()[SystemIndexToEndReplicationOn];
+	FNetRefHandle RefHandleOnServerToEndReplicationOn;
+	{
+		UReplicatedTestObjectBridge* ReplicationBridge = ServerToEndReplicationOn->GetReplicationBridge();
+		RefHandleOnServerToEndReplicationOn = ReplicationBridge->GetReplicatedRefHandle(ServerObject);
+		ReplicationBridge->EndReplication(ServerObject);
+	}
+
+	FullSendAndDeliverUpdate();
+
+	// Verify object was destroyed on client connected to system where replication was ended
+	for (FReplicationSystemTestServer* Server : GetAllServers())
+	{
+		if (Server == ServerToEndReplicationOn)
+		{
+			for (FReplicationSystemTestClient* Client : GetClients(ServerToEndReplicationOn))
+			{
+				UE_NET_ASSERT_EQ(Client->GetReplicationBridge()->GetReplicatedObject(RefHandleOnServerToEndReplicationOn), nullptr);
+			}
+		}
+		else
+		{
+			FNetRefHandle RefHandleOnServer = Server->GetReplicationBridge()->GetReplicatedRefHandle(ServerObject);
+			for (FReplicationSystemTestClient* Client : GetClients(Server))
+			{
+				UE_NET_ASSERT_NE(Client->GetReplicationBridge()->GetReplicatedObject(RefHandleOnServer), nullptr);
+			}
+		}
+	}
+}
+
+
+}
