Iris
- Added proper propagation of RepNotify_Always flag

diff --git a/Engine/Plugins/Runtime/ReplicationSystemTestPlugin/Source/Private/Tests/ReplicationState/TestPropertyReplicationState.cpp b/Engine/Plugins/Runtime/ReplicationSystemTestPlugin/Source/Private/Tests/ReplicationState/TestPropertyReplicationState.cpp
--- a/Engine/Plugins/Runtime/ReplicationSystemTestPlugin/Source/Private/Tests/ReplicationState/TestPropertyReplicationState.cpp
+++ b/Engine/Plugins/Runtime/ReplicationSystemTestPlugin/Source/Private/Tests/ReplicationState/TestPropertyReplicationState.cpp
@@ -21,9 +21,14 @@ void UTestPropertyReplicationState_TestClassWithRepNotify::OnRep_IntA(int32 OldI
 {
 }
 
+void UTestPropertyReplicationState_TestClassWithRepNotify::OnRep_IntB(int32 OldInt)
+{
+}
+
 void UTestPropertyReplicationState_TestClassWithRepNotify::GetLifetimeReplicatedProps( TArray< class FLifetimeProperty > & OutLifetimeProps ) const
 {
-	DOREPLIFETIME_CONDITION_NOTIFY( UTestPropertyReplicationState_TestClassWithRepNotify, IntA, COND_None, REPNOTIFY_Always );
+	DOREPLIFETIME_CONDITION_NOTIFY(UTestPropertyReplicationState_TestClassWithRepNotify, IntA, COND_None, REPNOTIFY_Always);
+	DOREPLIFETIME_CONDITION_NOTIFY(UTestPropertyReplicationState_TestClassWithRepNotify, IntB, COND_None, REPNOTIFY_OnChanged);
 }
 
 void UTestPropertyReplicationState_TestClassWithInitAndCArrays::GetLifetimeReplicatedProps(TArray<class FLifetimeProperty>& OutLifetimeProps) const
@@ -315,7 +320,12 @@ UE_NET_TEST_FIXTURE(FTestPropertyReplicationStateContext, PropertyReplicationSta
 {
 	InitDescriptorsFromClass(UTestPropertyReplicationState_TestClassWithRepNotify::StaticClass());
 
-	UE_NET_ASSERT_EQ(true, EnumHasAllFlags(Descriptor->Traits, EReplicationStateTraits::HasRepNotifies | EReplicationStateTraits::KeepPreviousState));
+	UE_NET_ASSERT_TRUE(EnumHasAllFlags(Descriptor->Traits, EReplicationStateTraits::HasRepNotifies | EReplicationStateTraits::KeepPreviousState));
+	UE_NET_ASSERT_TRUE(EnumHasAllFlags(Descriptor->MemberTraitsDescriptors[0].Traits, EReplicationStateMemberTraits::HasRepNotifyAlways));
+	UE_NET_ASSERT_NE(Descriptor->MemberPropertyDescriptors[0].RepNotifyFunction, nullptr);
+
+	UE_NET_ASSERT_FALSE(EnumHasAllFlags(Descriptor->MemberTraitsDescriptors[1].Traits, EReplicationStateMemberTraits::HasRepNotifyAlways));
+	UE_NET_ASSERT_NE(Descriptor->MemberPropertyDescriptors[1].RepNotifyFunction, nullptr);
 }
 
 UE_NET_TEST_FIXTURE(FTestPropertyReplicationStateContext, PropertyReplicationState_SetInitArray)
diff --git a/Engine/Plugins/Runtime/ReplicationSystemTestPlugin/Source/Private/Tests/ReplicationState/TestPropertyReplicationState.h b/Engine/Plugins/Runtime/ReplicationSystemTestPlugin/Source/Private/Tests/ReplicationState/TestPropertyReplicationState.h
--- a/Engine/Plugins/Runtime/ReplicationSystemTestPlugin/Source/Private/Tests/ReplicationState/TestPropertyReplicationState.h
+++ b/Engine/Plugins/Runtime/ReplicationSystemTestPlugin/Source/Private/Tests/ReplicationState/TestPropertyReplicationState.h
@@ -34,8 +34,14 @@ public:
 	UFUNCTION()
 	void OnRep_IntA(int32 OldInt);
 
+	UFUNCTION()
+	void OnRep_IntB(int32 OldInt);
+
 	UPROPERTY(Transient, ReplicatedUsing=OnRep_IntA)
 	int32 IntA;
+
+	UPROPERTY(Transient, ReplicatedUsing=OnRep_IntB)
+	int32 IntB;
 };
 
 
diff --git a/Engine/Source/Runtime/Experimental/Iris/Core/Private/Iris/ReplicationState/ReplicationStateDescriptorBuilder.cpp b/Engine/Source/Runtime/Experimental/Iris/Core/Private/Iris/ReplicationState/ReplicationStateDescriptorBuilder.cpp
--- a/Engine/Source/Runtime/Experimental/Iris/Core/Private/Iris/ReplicationState/ReplicationStateDescriptorBuilder.cpp
+++ b/Engine/Source/Runtime/Experimental/Iris/Core/Private/Iris/ReplicationState/ReplicationStateDescriptorBuilder.cpp
@@ -966,6 +966,7 @@ void FPropertyReplicationStateDescriptorBuilder::BuildMemberTraitsDescriptors(FR
 	{
 		CurrentMemberTraitsDescriptor->Traits |= (EnumHasAnyFlags(Info.Traits, EMemberPropertyTraits::HasDynamicState) ? EReplicationStateMemberTraits::HasDynamicState : EReplicationStateMemberTraits::None);
 		CurrentMemberTraitsDescriptor->Traits |= (EnumHasAnyFlags(Info.Traits, EMemberPropertyTraits::HasObjectReference) ? EReplicationStateMemberTraits::HasObjectReference : EReplicationStateMemberTraits::None);
+		CurrentMemberTraitsDescriptor->Traits |= (EnumHasAnyFlags(Info.Traits, EMemberPropertyTraits::RepNotifyAlways) ? EReplicationStateMemberTraits::HasRepNotifyAlways : EReplicationStateMemberTraits::None);
 		++CurrentMemberTraitsDescriptor;
 	}
 
diff --git a/Engine/Source/Runtime/Experimental/Iris/Core/Public/Iris/ReplicationState/ReplicationStateDescriptor.h b/Engine/Source/Runtime/Experimental/Iris/Core/Public/Iris/ReplicationState/ReplicationStateDescriptor.h
--- a/Engine/Source/Runtime/Experimental/Iris/Core/Public/Iris/ReplicationState/ReplicationStateDescriptor.h
+++ b/Engine/Source/Runtime/Experimental/Iris/Core/Public/Iris/ReplicationState/ReplicationStateDescriptor.h
@@ -46,6 +46,7 @@ enum class EReplicationStateMemberTraits : uint16
 	HasDynamicState = 1U << 0U,
 	HasObjectReference = HasDynamicState << 1U,
 	HasConnectionSpecificSerialization = HasObjectReference << 1U,
+	HasRepNotifyAlways = HasConnectionSpecificSerialization << 1U,
 };
 ENUM_CLASS_FLAGS(EReplicationStateMemberTraits);
 
