Iris
- Added scheduling order hint to allow dependent objects to be scheduled before parent to mimic some ordering quirks of old replication system

diff --git a/Engine/Plugins/Runtime/ReplicationSystemTestPlugin/Source/Private/Tests/ReplicationSystem/TestDependentObjects.cpp b/Engine/Plugins/Runtime/ReplicationSystemTestPlugin/Source/Private/Tests/ReplicationSystem/TestDependentObjects.cpp
--- a/Engine/Plugins/Runtime/ReplicationSystemTestPlugin/Source/Private/Tests/ReplicationSystem/TestDependentObjects.cpp
+++ b/Engine/Plugins/Runtime/ReplicationSystemTestPlugin/Source/Private/Tests/ReplicationSystem/TestDependentObjects.cpp
@@ -376,4 +376,322 @@ UE_NET_TEST_FIXTURE(FReplicationSystemServerClientTestFixture, TestDependentObje
 	UE_NET_ASSERT_EQ(ClientDependentObject->IntA, ServerDependentObject->IntA);
 }
 
+UE_NET_TEST_FIXTURE(FReplicationSystemServerClientTestFixture, TestDependentObjectScheduledAfterParent)
+{
+	UReplicationSystem* ReplicationSystem = Server->ReplicationSystem;
+	UReplicatedTestObjectBridge* Bridge = Server->GetReplicationBridge();
+
+	// Add a client
+	FReplicationSystemTestClient* Client = CreateClient();
+
+	// Spawn object on server
+	UReplicatedSubObjectOrderObject* ServerObject = Server->CreateObject<UReplicatedSubObjectOrderObject>();
+
+	// Spawn dependent object
+	UReplicatedSubObjectOrderObject* ServerDependentObject = Server->CreateObject<UReplicatedSubObjectOrderObject>();
+	Bridge->AddDependentObject(ServerObject->NetRefHandle, ServerDependentObject->NetRefHandle);
+
+	// Set prio to only replicate with parent
+	ReplicationSystem->SetStaticPriority(ServerDependentObject->NetRefHandle, 0.f);
+
+	// Reset RepOrderCounter
+	UReplicatedSubObjectOrderObject::RepOrderCounter = 0U;
+
+	// Send and deliver packet
+	Server->PreSendUpdate();
+	Server->SendAndDeliverTo(Client, true);
+	Server->PostSendUpdate();
+
+	// Verify that objects have replicated
+	UReplicatedSubObjectOrderObject* ClientObject = Cast<UReplicatedSubObjectOrderObject>(Client->GetReplicationBridge()->GetReplicatedObject(ServerObject->NetRefHandle));
+	UReplicatedSubObjectOrderObject* ClientDependentObject = Cast<UReplicatedSubObjectOrderObject>(Client->GetReplicationBridge()->GetReplicatedObject(ServerDependentObject->NetRefHandle));
+
+	UE_NET_ASSERT_NE(ClientObject, nullptr);
+	UE_NET_ASSERT_NE(ClientDependentObject, nullptr);
+
+	// Verify that they have replicated in expected order for initial objects
+	UE_NET_ASSERT_GT(ClientDependentObject->LastRepOrderCounter, ClientObject->LastRepOrderCounter);
+
+	// Modify both parent and dependent
+	ServerObject->IntA = 1;
+	ServerDependentObject->IntA = 1;
+
+	// Send and deliver packet
+	Server->PreSendUpdate();
+	Server->SendAndDeliverTo(Client, true);
+	Server->PostSendUpdate();
+
+	// Verify that they have replicated in expected order 
+	UE_NET_ASSERT_GT(ClientDependentObject->LastRepOrderCounter, ClientObject->LastRepOrderCounter);
+}
+
+UE_NET_TEST_FIXTURE(FReplicationSystemServerClientTestFixture, TestDependentObjectScheduledBeforeParent)
+{
+	UReplicationSystem* ReplicationSystem = Server->ReplicationSystem;
+	UReplicatedTestObjectBridge* Bridge = Server->GetReplicationBridge();
+
+	// Add a client
+	FReplicationSystemTestClient* Client = CreateClient();
+
+	// Spawn object on server
+	UReplicatedSubObjectOrderObject* ServerObject = Server->CreateObject<UReplicatedSubObjectOrderObject>();
+
+	// Spawn dependent object
+	UReplicatedSubObjectOrderObject* ServerDependentObject = Server->CreateObject<UReplicatedSubObjectOrderObject>();
+	Bridge->AddDependentObject(ServerObject->NetRefHandle, ServerDependentObject->NetRefHandle, EDependentObjectSchedulingHint::ScheduleBeforeParent);
+
+	// Set prio to only replicate with parent
+	ReplicationSystem->SetStaticPriority(ServerDependentObject->NetRefHandle, 0.f);
+
+	// Reset RepOrderCounter
+	UReplicatedSubObjectOrderObject::RepOrderCounter = 0U;
+
+	// Send and deliver packet
+	Server->PreSendUpdate();
+	Server->SendAndDeliverTo(Client, true);
+	Server->PostSendUpdate();
+
+	// Verify that objects have replicated
+	UReplicatedSubObjectOrderObject* ClientObject = Cast<UReplicatedSubObjectOrderObject>(Client->GetReplicationBridge()->GetReplicatedObject(ServerObject->NetRefHandle));
+	UReplicatedSubObjectOrderObject* ClientDependentObject = Cast<UReplicatedSubObjectOrderObject>(Client->GetReplicationBridge()->GetReplicatedObject(ServerDependentObject->NetRefHandle));
+
+	UE_NET_ASSERT_NE(ClientObject, nullptr);
+	UE_NET_ASSERT_NE(ClientDependentObject, nullptr);
+
+	// Verify that they have replicated in expected order for initial objects
+	UE_NET_ASSERT_LT(ClientDependentObject->LastRepOrderCounter, ClientObject->LastRepOrderCounter);
+
+	// Modify both parent and dependent
+	ServerObject->IntA = 1;
+	ServerDependentObject->IntA = 1;
+
+	// Send and deliver packet
+	Server->PreSendUpdate();
+	Server->SendAndDeliverTo(Client, true);
+	Server->PostSendUpdate();
+
+	// Verify that they have replicated in expected order
+	UE_NET_ASSERT_LT(ClientDependentObject->LastRepOrderCounter, ClientObject->LastRepOrderCounter);
+}
+
+UE_NET_TEST_FIXTURE(FReplicationSystemServerClientTestFixture, TestDependentObjectScheduledBeforeParentIfInitialState)
+{
+	UReplicationSystem* ReplicationSystem = Server->ReplicationSystem;
+	UReplicatedTestObjectBridge* Bridge = Server->GetReplicationBridge();
+
+	// Add a client
+	FReplicationSystemTestClient* Client = CreateClient();
+
+	// Spawn object on server
+	UReplicatedSubObjectOrderObject* ServerObject = Server->CreateObject<UReplicatedSubObjectOrderObject>();
+
+	// Spawn dependent object
+	UReplicatedSubObjectOrderObject* ServerDependentObject = Server->CreateObject<UReplicatedSubObjectOrderObject>();
+	Bridge->AddDependentObject(ServerObject->NetRefHandle, ServerDependentObject->NetRefHandle, EDependentObjectSchedulingHint::ScheduleBeforeParentIfInitialState);
+
+	// Set prio to only replicate with parent
+	ReplicationSystem->SetStaticPriority(ServerDependentObject->NetRefHandle, 0.f);
+
+	// Reset RepOrderCounter
+	UReplicatedSubObjectOrderObject::RepOrderCounter = 0U;
+
+	// Send and deliver packet
+	Server->PreSendUpdate();
+	Server->SendAndDeliverTo(Client, true);
+	Server->PostSendUpdate();
+
+	// Verify that objects have replicated
+	UReplicatedSubObjectOrderObject* ClientObject = Cast<UReplicatedSubObjectOrderObject>(Client->GetReplicationBridge()->GetReplicatedObject(ServerObject->NetRefHandle));
+	UReplicatedSubObjectOrderObject* ClientDependentObject = Cast<UReplicatedSubObjectOrderObject>(Client->GetReplicationBridge()->GetReplicatedObject(ServerDependentObject->NetRefHandle));
+
+	UE_NET_ASSERT_NE(ClientObject, nullptr);
+	UE_NET_ASSERT_NE(ClientDependentObject, nullptr);
+
+	// Verify that they have replicated in expected order for initial objects
+	UE_NET_ASSERT_LT(ClientDependentObject->LastRepOrderCounter, ClientObject->LastRepOrderCounter);
+
+	// Modify both parent and dependent
+	ServerObject->IntA = 1;
+	ServerDependentObject->IntA = 1;
+
+	// Send and deliver packet
+	Server->PreSendUpdate();
+	Server->SendAndDeliverTo(Client, true);
+	Server->PostSendUpdate();
+
+	// Verify that they have replicated in expected order
+	UE_NET_ASSERT_GT(ClientDependentObject->LastRepOrderCounter, ClientObject->LastRepOrderCounter);
+}
+
+UE_NET_TEST_FIXTURE(FReplicationSystemServerClientTestFixture, TestNestedDependentObjectScheduledBeforeParent)
+{
+	UReplicationSystem* ReplicationSystem = Server->ReplicationSystem;
+	UReplicatedTestObjectBridge* Bridge = Server->GetReplicationBridge();
+
+	// Add a client
+	FReplicationSystemTestClient* Client = CreateClient();
+
+	// Spawn objects on server, abusing some assumptions about scheduling order based on assigned internal indices to try out ordering of dependent objects
+	UReplicatedSubObjectOrderObject* ServerNestedDependentObject = Server->CreateObject<UReplicatedSubObjectOrderObject>();
+	UReplicatedSubObjectOrderObject* ServerDependentObject = Server->CreateObject<UReplicatedSubObjectOrderObject>();
+	UReplicatedSubObjectOrderObject* ServerObject = Server->CreateObject<UReplicatedSubObjectOrderObject>();
+
+	// Setup dependent object and nested dependent object to both replicate before the parent
+	Bridge->AddDependentObject(ServerObject->NetRefHandle, ServerDependentObject->NetRefHandle, EDependentObjectSchedulingHint::ScheduleBeforeParent);
+	Bridge->AddDependentObject(ServerDependentObject->NetRefHandle, ServerNestedDependentObject->NetRefHandle, EDependentObjectSchedulingHint::ScheduleBeforeParent);
+
+	// Set static prio to only replicate with parent
+	ReplicationSystem->SetStaticPriority(ServerDependentObject->NetRefHandle, 0.f);
+	ReplicationSystem->SetStaticPriority(ServerNestedDependentObject->NetRefHandle, 0.f);
+
+	// Reset RepOrderCounter
+	UReplicatedSubObjectOrderObject::RepOrderCounter = 0U;
+
+	// Send and deliver packet
+	Server->PreSendUpdate();
+	Server->SendAndDeliverTo(Client, true);
+	Server->PostSendUpdate();
+
+	// Verify that objects have replicated
+	UReplicatedSubObjectOrderObject* ClientObject = Cast<UReplicatedSubObjectOrderObject>(Client->GetReplicationBridge()->GetReplicatedObject(ServerObject->NetRefHandle));
+	UReplicatedSubObjectOrderObject* ClientDependentObject = Cast<UReplicatedSubObjectOrderObject>(Client->GetReplicationBridge()->GetReplicatedObject(ServerDependentObject->NetRefHandle));
+	UReplicatedSubObjectOrderObject* ClientNestedDependentObject = Cast<UReplicatedSubObjectOrderObject>(Client->GetReplicationBridge()->GetReplicatedObject(ServerNestedDependentObject->NetRefHandle));
+
+	UE_NET_ASSERT_NE(ClientObject, nullptr);
+	UE_NET_ASSERT_NE(ClientDependentObject, nullptr);
+	UE_NET_ASSERT_NE(ClientNestedDependentObject, nullptr);
+
+	// Verify that they have replicated in expected order for initial objects
+	UE_NET_ASSERT_LT(ClientNestedDependentObject->LastRepOrderCounter, ClientDependentObject->LastRepOrderCounter);
+	UE_NET_ASSERT_LT(ClientDependentObject->LastRepOrderCounter, ClientObject->LastRepOrderCounter);
+
+	// Modify both parent and dependent
+	ServerObject->IntA = 1;
+	ServerDependentObject->IntA = 1;
+	ServerNestedDependentObject->IntA = 1;
+
+	// Send and deliver packet
+	Server->PreSendUpdate();
+	Server->SendAndDeliverTo(Client, true);
+	Server->PostSendUpdate();
+
+	// Verify that they have replicated in expected order 
+	UE_NET_ASSERT_LT(ClientNestedDependentObject->LastRepOrderCounter, ClientDependentObject->LastRepOrderCounter);
+	UE_NET_ASSERT_LT(ClientDependentObject->LastRepOrderCounter, ClientObject->LastRepOrderCounter);
+}
+
+UE_NET_TEST_FIXTURE(FReplicationSystemServerClientTestFixture, TestNestedDependentObjectScheduledBeforeParents)
+{
+	UReplicationSystem* ReplicationSystem = Server->ReplicationSystem;
+	UReplicatedTestObjectBridge* Bridge = Server->GetReplicationBridge();
+
+	// Add a client
+	FReplicationSystemTestClient* Client = CreateClient();
+
+	// Spawn objects on server, abusing some assumptions about scheduling order based on assigned internal indices to try out ordering of dependent objects
+	UReplicatedSubObjectOrderObject* ServerNestedDependentObject = Server->CreateObject<UReplicatedSubObjectOrderObject>();
+	UReplicatedSubObjectOrderObject* ServerDependentObject = Server->CreateObject<UReplicatedSubObjectOrderObject>();
+	UReplicatedSubObjectOrderObject* ServerObject = Server->CreateObject<UReplicatedSubObjectOrderObject>();
+
+	// Setup dependent object with a nested dependent object that is replicating before its parent
+	Bridge->AddDependentObject(ServerObject->NetRefHandle, ServerDependentObject->NetRefHandle);
+	Bridge->AddDependentObject(ServerDependentObject->NetRefHandle, ServerNestedDependentObject->NetRefHandle, EDependentObjectSchedulingHint::ScheduleBeforeParent);
+
+	// Set static prio to only replicate with parent
+	ReplicationSystem->SetStaticPriority(ServerDependentObject->NetRefHandle, 0.f);
+	ReplicationSystem->SetStaticPriority(ServerNestedDependentObject->NetRefHandle, 0.f);
+
+	// Reset RepOrderCounter
+	UReplicatedSubObjectOrderObject::RepOrderCounter = 0U;
+
+	// Send and deliver packet
+	Server->PreSendUpdate();
+	Server->SendAndDeliverTo(Client, true);
+	Server->PostSendUpdate();
+
+	// Verify that objects have replicated
+	UReplicatedSubObjectOrderObject* ClientObject = Cast<UReplicatedSubObjectOrderObject>(Client->GetReplicationBridge()->GetReplicatedObject(ServerObject->NetRefHandle));
+	UReplicatedSubObjectOrderObject* ClientDependentObject = Cast<UReplicatedSubObjectOrderObject>(Client->GetReplicationBridge()->GetReplicatedObject(ServerDependentObject->NetRefHandle));
+	UReplicatedSubObjectOrderObject* ClientNestedDependentObject = Cast<UReplicatedSubObjectOrderObject>(Client->GetReplicationBridge()->GetReplicatedObject(ServerNestedDependentObject->NetRefHandle));
+
+	UE_NET_ASSERT_NE(ClientObject, nullptr);
+	UE_NET_ASSERT_NE(ClientDependentObject, nullptr);
+	UE_NET_ASSERT_NE(ClientNestedDependentObject, nullptr);
+
+	// Verify that they have replicated in expected order for initial objects
+	UE_NET_ASSERT_LT(ClientNestedDependentObject->LastRepOrderCounter, ClientDependentObject->LastRepOrderCounter);
+	UE_NET_ASSERT_LT(ClientObject->LastRepOrderCounter, ClientDependentObject->LastRepOrderCounter);
+
+	// Modify both parent and dependent
+	ServerObject->IntA = 1;
+	ServerDependentObject->IntA = 1;
+	ServerNestedDependentObject->IntA = 1;
+
+	// Send and deliver packet
+	Server->PreSendUpdate();
+	Server->SendAndDeliverTo(Client, true);
+	Server->PostSendUpdate();
+
+	// Verify that they have replicated in expected order 
+	UE_NET_ASSERT_LT(ClientNestedDependentObject->LastRepOrderCounter, ClientDependentObject->LastRepOrderCounter);
+	UE_NET_ASSERT_LT(ClientObject->LastRepOrderCounter, ClientDependentObject->LastRepOrderCounter);
+}
+
+UE_NET_TEST_FIXTURE(FReplicationSystemServerClientTestFixture, TestLateAddedNestedDependentObjectScheduledBeforeParents)
+{
+	UReplicationSystem* ReplicationSystem = Server->ReplicationSystem;
+	UReplicatedTestObjectBridge* Bridge = Server->GetReplicationBridge();
+
+	// Add a client
+	FReplicationSystemTestClient* Client = CreateClient();
+
+	// Spawn objects on server, abusing some assumptions about scheduling order based on assigned internal indices to try out ordering of dependent objects
+	UReplicatedSubObjectOrderObject* ServerDependentObject = Server->CreateObject<UReplicatedSubObjectOrderObject>();
+	UReplicatedSubObjectOrderObject* ServerObject = Server->CreateObject<UReplicatedSubObjectOrderObject>();
+
+	// Setup dependent object with a nested dependent object that is replicating before its parent
+	Bridge->AddDependentObject(ServerObject->NetRefHandle, ServerDependentObject->NetRefHandle);
+
+	// Set static prio to only replicate with parent
+	ReplicationSystem->SetStaticPriority(ServerDependentObject->NetRefHandle, 0.f);
+
+	// Reset RepOrderCounter
+	UReplicatedSubObjectOrderObject::RepOrderCounter = 0U;
+
+	// Send and deliver packet
+	Server->PreSendUpdate();
+	Server->SendAndDeliverTo(Client, true);
+	Server->PostSendUpdate();
+
+	// Verify that objects have replicated
+	UReplicatedSubObjectOrderObject* ClientObject = Cast<UReplicatedSubObjectOrderObject>(Client->GetReplicationBridge()->GetReplicatedObject(ServerObject->NetRefHandle));
+	UReplicatedSubObjectOrderObject* ClientDependentObject = Cast<UReplicatedSubObjectOrderObject>(Client->GetReplicationBridge()->GetReplicatedObject(ServerDependentObject->NetRefHandle));
+
+	UE_NET_ASSERT_NE(ClientObject, nullptr);
+	UE_NET_ASSERT_NE(ClientDependentObject, nullptr);
+
+	// Verify that they have replicated in expected order for initial objects
+	UE_NET_ASSERT_LT(ClientObject->LastRepOrderCounter, ClientDependentObject->LastRepOrderCounter);
+
+	// Create new dependent object
+	UReplicatedSubObjectOrderObject* ServerNestedDependentObject = Server->CreateObject<UReplicatedSubObjectOrderObject>();
+	Bridge->AddDependentObject(ServerDependentObject->NetRefHandle, ServerNestedDependentObject->NetRefHandle, EDependentObjectSchedulingHint::ScheduleBeforeParent);
+	ReplicationSystem->SetStaticPriority(ServerNestedDependentObject->NetRefHandle, 0.f);
+
+	// Modify parent to trigger replication of new dependent object
+	ServerObject->IntA = 1;
+
+	// Send and deliver packet
+	Server->PreSendUpdate();
+	Server->SendAndDeliverTo(Client, true);
+	Server->PostSendUpdate();
+
+	// Verify that new dependent object has been created
+	UReplicatedSubObjectOrderObject* ClientNestedDependentObject = Cast<UReplicatedSubObjectOrderObject>(Client->GetReplicationBridge()->GetReplicatedObject(ServerNestedDependentObject->NetRefHandle));
+	UE_NET_ASSERT_NE(ClientNestedDependentObject, nullptr);
+
+	// Verify that they have replicated in expected order 
+	UE_NET_ASSERT_LT(ClientNestedDependentObject->LastRepOrderCounter, ClientObject->LastRepOrderCounter);
+}
+
 }
diff --git a/Engine/Source/Runtime/Engine/Private/Net/Iris/ReplicationSystem/ReplicationSystemUtil.cpp b/Engine/Source/Runtime/Engine/Private/Net/Iris/ReplicationSystem/ReplicationSystemUtil.cpp
--- a/Engine/Source/Runtime/Engine/Private/Net/Iris/ReplicationSystem/ReplicationSystemUtil.cpp
+++ b/Engine/Source/Runtime/Engine/Private/Net/Iris/ReplicationSystem/ReplicationSystemUtil.cpp
@@ -313,7 +313,7 @@ void FReplicationSystemUtil::EndReplicationForActorComponentSubObject(UActorComp
 	EndReplicationForActorSubObject(ActorComponent->GetOwner(), SubObject);
 }
 
-void FReplicationSystemUtil::AddDependentActor(const AActor* Parent, AActor* Child)
+void FReplicationSystemUtil::AddDependentActor(const AActor* Parent, AActor* Child, EDependentObjectSchedulingHint SchedulingHint)
 {
 	// Can only add dependent actors on already replicating actors
 	FNetHandle ParentHandle = GetNetHandle(Parent);
@@ -324,7 +324,7 @@ void FReplicationSystemUtil::AddDependentActor(const AActor* Parent, AActor* Chi
 
 	if (const UWorld* World = Parent->GetWorld())
 	{
-		ReplicationSystemUtil::ForEachReplicationSystem(GEngine, World, [ParentHandle, Child](UReplicationSystem* ReplicationSystem)
+		ReplicationSystemUtil::ForEachReplicationSystem(GEngine, World, [ParentHandle, Child, SchedulingHint](UReplicationSystem* ReplicationSystem)
 		{
 			if (ReplicationSystem->IsServer())
 			{
@@ -336,7 +336,7 @@ void FReplicationSystemUtil::AddDependentActor(const AActor* Parent, AActor* Chi
 						const FNetRefHandle ChildRefHandle = Bridge->BeginReplication(Child);
 						if (ensureMsgf(ChildRefHandle.IsValid(), TEXT("FReplicationSystemUtil::AddDependentActor Child %s must be replicated"), *GetPathNameSafe(Child)))
 						{
-							Bridge->AddDependentObject(ParentRefHandle, ChildRefHandle);
+							Bridge->AddDependentObject(ParentRefHandle, ChildRefHandle, SchedulingHint);
 						}
 					}
 				}
@@ -345,6 +345,11 @@ void FReplicationSystemUtil::AddDependentActor(const AActor* Parent, AActor* Chi
 	}
 }
 
+void FReplicationSystemUtil::AddDependentActor(const AActor* Parent, AActor* Child)
+{
+	AddDependentActor(Parent, Child, EDependentObjectSchedulingHint::Default);
+}
+
 void FReplicationSystemUtil::RemoveDependentActor(const AActor* Parent, AActor* Child)
 {
 	const FNetHandle ParentHandle = GetNetHandle(Parent);
diff --git a/Engine/Source/Runtime/Engine/Public/Net/Iris/ReplicationSystem/ReplicationSystemUtil.h b/Engine/Source/Runtime/Engine/Public/Net/Iris/ReplicationSystem/ReplicationSystemUtil.h
--- a/Engine/Source/Runtime/Engine/Public/Net/Iris/ReplicationSystem/ReplicationSystemUtil.h
+++ b/Engine/Source/Runtime/Engine/Public/Net/Iris/ReplicationSystem/ReplicationSystemUtil.h
@@ -20,6 +20,11 @@ class UReplicationSystem;
 class UWorld;
 class UNetConnection;
 
+namespace UE::Net
+{
+	enum class EDependentObjectSchedulingHint : uint8;
+}
+
 namespace UE::Net
 {
 
@@ -92,6 +97,8 @@ struct FReplicationSystemUtil
 	 * Dependent actors cannot be filtered out by dynamic filtering unless the parent is also filtered out.
 	 * @note There is no guarantee that the data will end up in the same packet so it is a very loose form of dependency.
 	 */
+	ENGINE_API static void AddDependentActor(const AActor* Parent, AActor* Child, EDependentObjectSchedulingHint SchedulingHint);
+
 	ENGINE_API static void AddDependentActor(const AActor* Parent, AActor* Child);
 
 	/** Remove dependent actor from parent. The dependent actor will function as a standard standalone replicated actor. */
diff --git a/Engine/Source/Runtime/Experimental/Iris/Core/Private/Iris/ReplicationSystem/NetDependencyData.cpp b/Engine/Source/Runtime/Experimental/Iris/Core/Private/Iris/ReplicationSystem/NetDependencyData.cpp
--- a/Engine/Source/Runtime/Experimental/Iris/Core/Private/Iris/ReplicationSystem/NetDependencyData.cpp
+++ b/Engine/Source/Runtime/Experimental/Iris/Core/Private/Iris/ReplicationSystem/NetDependencyData.cpp
@@ -29,6 +29,11 @@ void FNetDependencyData::FreeStoredDependencyDataForObject(FInternalNetRefIndex
 			SubObjectConditionalsStorage.RemoveAt(Entry->SubObjectConditionalArrayIndex);
 		}
 
+		if (Entry->DependentObjectsInfoArrayIndex != FDependencyInfo::InvalidCacheIndex)
+		{
+			DependentObjectInfosStorage.RemoveAt(Entry->DependentObjectsInfoArrayIndex);
+		}
+
 		DependencyInfos.Remove(InternalIndex);
 	}
 }
@@ -46,6 +51,25 @@ FNetDependencyData::FDependencyInfo& FNetDependencyData::GetOrCreateCacheEntry(F
 	return *Entry;
 }
 
+FNetDependencyData::FDependentObjectInfoArray& FNetDependencyData::GetOrCreateDependentObjectInfoArray(FInternalNetRefIndex OwnerIndex)
+{
+	FDependencyInfo& Entry = GetOrCreateCacheEntry(OwnerIndex);
+	
+	if (Entry.DependentObjectsInfoArrayIndex == FDependencyInfo::InvalidCacheIndex)
+	{
+		FSparseArrayAllocationInfo AllocInfo = DependentObjectInfosStorage.AddUninitialized();
+		Entry.DependentObjectsInfoArrayIndex = AllocInfo.Index;
+
+		FDependentObjectInfoArray* DependentObjectsInfoArrayIndex = new (AllocInfo.Pointer) FDependentObjectInfoArray();
+		
+		return *DependentObjectsInfoArrayIndex;
+	}
+	else
+	{
+		return DependentObjectInfosStorage[Entry.DependentObjectsInfoArrayIndex];
+	}
+}
+
 FNetDependencyData::FSubObjectConditionalsArray& FNetDependencyData::GetOrCreateSubObjectConditionalsArray(FInternalNetRefIndex OwnerIndex)
 {
 	FDependencyInfo& Entry = GetOrCreateCacheEntry(OwnerIndex);
diff --git a/Engine/Source/Runtime/Experimental/Iris/Core/Private/Iris/ReplicationSystem/NetDependencyData.h b/Engine/Source/Runtime/Experimental/Iris/Core/Private/Iris/ReplicationSystem/NetDependencyData.h
--- a/Engine/Source/Runtime/Experimental/Iris/Core/Private/Iris/ReplicationSystem/NetDependencyData.h
+++ b/Engine/Source/Runtime/Experimental/Iris/Core/Private/Iris/ReplicationSystem/NetDependencyData.h
@@ -6,6 +6,7 @@
 #include "Containers/ArrayView.h"
 #include "Containers/Map.h"
 #include "Containers/SparseArray.h"
+#include "Iris/ReplicationSystem/ReplicationSystemTypes.h"
 
 namespace UE::Net::Private
 {
@@ -23,21 +24,26 @@ struct FChildSubObjectsInfo
 	uint32 NumSubObjects = 0U;
 };
 
+struct FDependentObjectInfo
+{
+	FInternalNetRefIndex NetRefIndex = 0U;
+	EDependentObjectSchedulingHint SchedulingHint = EDependentObjectSchedulingHint::Default;
+};
+
 class FNetDependencyData
 {
 public:
 	FNetDependencyData();
 
 	typedef TArray<FInternalNetRefIndex, TInlineAllocator<8>> FInternalNetRefIndexArray;
-
 	typedef TArray<FLifeTimeConditionStorage, TInlineAllocator<8>> FSubObjectConditionalsArray;
+	typedef TArray<FDependentObjectInfo, TInlineAllocator<8>> FDependentObjectInfoArray;
 
 	enum EArrayType
 	{
 		SubObjects = 0U,
 		ChildSubObjects,
-		DependentObjects,
-		ParentObjects,
+		DependentParentObjects,
 		Count
 	};
 
@@ -48,6 +54,29 @@ public:
 		return GetOrCreateInternalIndexArray(InternalIndex, TypeIndex);
 	};
 
+	FDependentObjectInfoArray& GetOrCreateDependentObjectInfoArray(FInternalNetRefIndex InternalIndex);
+	FDependentObjectInfoArray* GetDependentObjectInfoArray(FInternalNetRefIndex InternalIndex)
+	{
+		const FDependencyInfo* Entry = DependencyInfos.Find(InternalIndex);
+		const uint32 ArrayIndex = Entry ? Entry->DependentObjectsInfoArrayIndex : FDependencyInfo::InvalidCacheIndex;
+		if (ArrayIndex != FDependencyInfo::InvalidCacheIndex)
+		{
+			return &DependentObjectInfosStorage[ArrayIndex];
+		}
+		return nullptr;
+	}
+
+	TArrayView<const FDependentObjectInfo> GetDependentObjectInfoArray(FInternalNetRefIndex InternalIndex) const
+	{
+		const FDependencyInfo* Entry = DependencyInfos.Find(InternalIndex);
+		const uint32 ArrayIndex = Entry ? Entry->DependentObjectsInfoArrayIndex : FDependencyInfo::InvalidCacheIndex;
+		if (ArrayIndex != FDependencyInfo::InvalidCacheIndex)
+		{
+			return MakeArrayView(DependentObjectInfosStorage[ArrayIndex]);
+		}
+		return MakeArrayView<const FDependentObjectInfo>(nullptr, 0);
+	}
+
 	FSubObjectConditionalsArray& GetOrCreateSubObjectConditionalsArray(FInternalNetRefIndex InternalIndex);
 
 	FInternalNetRefIndexArray& GetOrCreateInternalChildSubObjectsArray(FInternalNetRefIndex InternalIndex, FSubObjectConditionalsArray*& OutSubObjectConditionals);
@@ -82,7 +111,6 @@ public:
 		return true;
 	}
 
-
 	template<EArrayType TypeIndex>
 	FInternalNetRefIndexArray* GetInternalIndexArray(FInternalNetRefIndex InternalIndex)
 	{
@@ -117,8 +145,9 @@ private:
 	struct FDependencyInfo
 	{
 		constexpr static uint32 InvalidCacheIndex = ~(0U);
-		uint32 ArrayIndices[EArrayType::Count] = { InvalidCacheIndex, InvalidCacheIndex, InvalidCacheIndex, InvalidCacheIndex };
+		uint32 ArrayIndices[EArrayType::Count] = { InvalidCacheIndex, InvalidCacheIndex, InvalidCacheIndex };
 		uint32 SubObjectConditionalArrayIndex = InvalidCacheIndex;
+		uint32 DependentObjectsInfoArrayIndex = InvalidCacheIndex;
 	};
 
 private:
@@ -131,6 +160,8 @@ private:
 	TSparseArray<FInternalNetRefIndexArray> DependentObjectsStorage;
 	// Storage for SubObject conditionals
 	TSparseArray<FSubObjectConditionalsArray> SubObjectConditionalsStorage;
+	// Storage for DependentObjects
+	TSparseArray<FDependentObjectInfoArray> DependentObjectInfosStorage;
 };
 
 }
diff --git a/Engine/Source/Runtime/Experimental/Iris/Core/Private/Iris/ReplicationSystem/NetRefHandleManager.cpp b/Engine/Source/Runtime/Experimental/Iris/Core/Private/Iris/ReplicationSystem/NetRefHandleManager.cpp
--- a/Engine/Source/Runtime/Experimental/Iris/Core/Private/Iris/ReplicationSystem/NetRefHandleManager.cpp
+++ b/Engine/Source/Runtime/Experimental/Iris/Core/Private/Iris/ReplicationSystem/NetRefHandleManager.cpp
@@ -562,7 +562,7 @@ FNetRefHandle FNetRefHandleManager::GetSubObjectOwner(FNetRefHandle SubObjectRef
 	return OwnerInternalIndex != InvalidInternalIndex ? ReplicatedObjectData[OwnerInternalIndex].RefHandle : FNetRefHandle();
 }
 
-bool FNetRefHandleManager::AddDependentObject(FNetRefHandle ParentRefHandle, FNetRefHandle DependentObjectRefHandle, EAddDependentObjectFlags Flags)
+bool FNetRefHandleManager::AddDependentObject(FNetRefHandle ParentRefHandle, FNetRefHandle DependentObjectRefHandle, EDependentObjectSchedulingHint SchedulingHint, EAddDependentObjectFlags Flags)
 {
 	check(ParentRefHandle != DependentObjectRefHandle);
 
@@ -587,21 +587,26 @@ bool FNetRefHandleManager::AddDependentObject(FNetRefHandle ParentRefHandle, FNe
 	check(!SubObjectInternalIndices.GetBit(ParentInternalIndex));
 
 	// Add dependent to parents dependent object list
-	FNetDependencyData::FInternalNetRefIndexArray& ParentDependentObjectArray = SubObjects.GetOrCreateInternalIndexArray<FNetDependencyData::DependentObjects>(ParentInternalIndex);
-	const bool bDependentIsAlreadyDependant = ParentDependentObjectArray.Find(DependentObjectInternalIndex) != INDEX_NONE;
+	FNetDependencyData::FDependentObjectInfoArray& ParentDependentObjectsArray = SubObjects.GetOrCreateDependentObjectInfoArray(ParentInternalIndex);
+	const bool bDependentIsAlreadyDependant = ParentDependentObjectsArray.FindByPredicate([DependentObjectInternalIndex](const FDependentObjectInfo& Entry) { return Entry.NetRefIndex == DependentObjectInternalIndex;}) != nullptr;
 	if (!bDependentIsAlreadyDependant)
 	{
-		ParentDependentObjectArray.Add(DependentObjectInternalIndex);
+		FDependentObjectInfo DependentObjectInfo;
+
+		DependentObjectInfo.NetRefIndex = DependentObjectInternalIndex;
+		DependentObjectInfo.SchedulingHint = SchedulingHint;
+
+		ParentDependentObjectsArray.Add(DependentObjectInfo);
 	}
 
-	FNetDependencyData::FInternalNetRefIndexArray& DependentParentObjectArray = SubObjects.GetOrCreateInternalIndexArray<FNetDependencyData::ParentObjects>(DependentObjectInternalIndex);
+	FNetDependencyData::FInternalNetRefIndexArray& DependentParentObjectArray = SubObjects.GetOrCreateInternalIndexArray<FNetDependencyData::DependentParentObjects>(DependentObjectInternalIndex);
 	const bool bDependentHadParentAlready = DependentParentObjectArray.Find(ParentInternalIndex) != INDEX_NONE;
 	if (!bDependentHadParentAlready)
 	{
 		DependentParentObjectArray.Add(ParentInternalIndex);
 	}
 
-	// Update cached info to avoid to do map lookups to find out if we have are a dependent object or have dependent objects
+	// Update cached info to avoid to do map lookups to find out if we are a dependent object or have dependent objects
 	DependentObjectData.bIsDependentObject = true;
 	ParentObjectData.bHasDependentObjects = true;
 	ObjectsWithDependentObjectsInternalIndices.SetBit(ParentInternalIndex);
@@ -631,7 +636,7 @@ void FNetRefHandleManager::InternalRemoveDependentObject(FInternalNetRefIndex Pa
 {
 	if (EnumHasAnyFlags(Flags, ERemoveDependentObjectFlags::RemoveFromDependentParentObjects))
 	{
-		if (FNetDependencyData::FInternalNetRefIndexArray* ParentObjectArray = SubObjects.GetInternalIndexArray<FNetDependencyData::ParentObjects>(DependentInternalIndex))
+		if (FNetDependencyData::FInternalNetRefIndexArray* ParentObjectArray = SubObjects.GetInternalIndexArray<FNetDependencyData::DependentParentObjects>(DependentInternalIndex))
 		{
 			ParentObjectArray->Remove(ParentInternalIndex);
 			if (ParentObjectArray->Num() == 0)
@@ -645,12 +650,18 @@ void FNetRefHandleManager::InternalRemoveDependentObject(FInternalNetRefIndex Pa
 
 	if (EnumHasAnyFlags(Flags, ERemoveDependentObjectFlags::RemoveFromParentDependentObjects))
 	{
-		if (FNetDependencyData::FInternalNetRefIndexArray* ParentDependentObjectArray = SubObjects.GetInternalIndexArray<FNetDependencyData::DependentObjects>(ParentInternalIndex))
+		FReplicatedObjectData& ParentObjectData = GetReplicatedObjectDataNoCheck(ParentInternalIndex);
+		FNetDependencyData::FDependentObjectInfoArray* ParentDependentObjectsArray = ParentObjectData.bHasDependentObjects ? SubObjects.GetDependentObjectInfoArray(ParentInternalIndex) : nullptr;
+		if (ParentDependentObjectsArray)
 		{
-			ParentDependentObjectArray->Remove(DependentInternalIndex);
-			if (ParentDependentObjectArray->Num() == 0)
+			const int32 ArrayIndex =  ParentDependentObjectsArray->FindLastByPredicate([DependentInternalIndex](const FDependentObjectInfo& Entry) { return Entry.NetRefIndex == DependentInternalIndex;});
+			if (ArrayIndex != INDEX_NONE)
+			{
+				ParentDependentObjectsArray->RemoveAt(ArrayIndex);
+			}
+
+			if (ParentDependentObjectsArray->Num() == 0)
 			{
-				FReplicatedObjectData& ParentObjectData = GetReplicatedObjectDataNoCheck(ParentInternalIndex);
 				ParentObjectData.bHasDependentObjects = false;
 				ObjectsWithDependentObjectsInternalIndices.ClearBit(ParentInternalIndex);
 			}
@@ -661,7 +672,7 @@ void FNetRefHandleManager::InternalRemoveDependentObject(FInternalNetRefIndex Pa
 void FNetRefHandleManager::InternalRemoveDependentObject(FInternalNetRefIndex DependentInternalIndex)
 {
 	// Remove from all parents
-	if (FNetDependencyData::FInternalNetRefIndexArray* ParentObjectArray = SubObjects.GetInternalIndexArray<FNetDependencyData::ParentObjects>(DependentInternalIndex))
+	if (FNetDependencyData::FInternalNetRefIndexArray* ParentObjectArray = SubObjects.GetInternalIndexArray<FNetDependencyData::DependentParentObjects>(DependentInternalIndex))
 	{
 		for (FInternalNetRefIndex ParentInternalIndex : *ParentObjectArray)
 		{
@@ -672,12 +683,12 @@ void FNetRefHandleManager::InternalRemoveDependentObject(FInternalNetRefIndex De
 	}
 
 	// Remove from our dependents
-	if (FNetDependencyData::FInternalNetRefIndexArray* DependentObjectArray = SubObjects.GetInternalIndexArray<FNetDependencyData::DependentObjects>(DependentInternalIndex))
+	if (FNetDependencyData::FDependentObjectInfoArray* DependentObjectArray = SubObjects.GetDependentObjectInfoArray(DependentInternalIndex))
 	{
-		for (FInternalNetRefIndex ChildDepedentInternalIndex : *DependentObjectArray)
+		for (const FDependentObjectInfo& ChildDependentObjectInfo : *DependentObjectArray)
 		{
 			// Flag is set to only update data on the childDependentObject to avoid modifying the array we iterate over
-			InternalRemoveDependentObject(DependentInternalIndex, ChildDepedentInternalIndex, ERemoveDependentObjectFlags::RemoveFromDependentParentObjects);
+			InternalRemoveDependentObject(DependentInternalIndex, ChildDependentObjectInfo.NetRefIndex, ERemoveDependentObjectFlags::RemoveFromDependentParentObjects);
 		}
 		DependentObjectArray->Reset();		
 	}
diff --git a/Engine/Source/Runtime/Experimental/Iris/Core/Private/Iris/ReplicationSystem/NetRefHandleManager.h b/Engine/Source/Runtime/Experimental/Iris/Core/Private/Iris/ReplicationSystem/NetRefHandleManager.h
--- a/Engine/Source/Runtime/Experimental/Iris/Core/Private/Iris/ReplicationSystem/NetRefHandleManager.h
+++ b/Engine/Source/Runtime/Experimental/Iris/Core/Private/Iris/ReplicationSystem/NetRefHandleManager.h
@@ -202,7 +202,7 @@ public:
 	// DependentObjects
 	const FNetBitArray& GetDependentObjectInternalIndices() const { return DependentObjectInternalIndices; }
 	const FNetBitArray& GetObjectsWithDependentObjectsInternalIndices() const { return ObjectsWithDependentObjectsInternalIndices; }
-	bool AddDependentObject(FNetRefHandle ParentHandle, FNetRefHandle DependentHandle, EAddDependentObjectFlags Flags = EAddDependentObjectFlags::WarnIfAlreadyDependentObject);
+	bool AddDependentObject(FNetRefHandle ParentHandle, FNetRefHandle DependentHandle, EDependentObjectSchedulingHint SchedulingHint, EAddDependentObjectFlags Flags = EAddDependentObjectFlags::WarnIfAlreadyDependentObject);
 	void RemoveDependentObject(FNetRefHandle ParentHandle, FNetRefHandle DependentHandle);
 
 	// Remove DependentHandles from all dependent object tracking
@@ -221,7 +221,7 @@ public:
 	uint16 GetNetObjectRefCount(FInternalNetRefIndex ObjectInternalIndex) const { return ReplicatedObjectRefCount[ObjectInternalIndex]; }
 
 	// Get dependent objects for the given ParentIndex
-	inline TArrayView<const FInternalNetRefIndex> GetDependentObjects(FInternalNetRefIndex ParentIndex) const;
+	TArrayView<const FDependentObjectInfo> GetDependentObjectInfos(FInternalNetRefIndex ParentIndex) const { return SubObjects.GetDependentObjectInfoArray(ParentIndex); }
 
 	// Get all parents of the given DependentIndex
 	inline TArrayView<const FInternalNetRefIndex> GetDependentObjectParents(FInternalNetRefIndex DependentIndex) const;
@@ -255,10 +255,10 @@ public:
 	{
 		if (ObjectsWithDependentObjectsInternalIndices.GetBit(ObjectIndex))
 		{
-			for (const FInternalNetRefIndex DependentObjectIndex : GetDependentObjects(ObjectIndex))
+			for (const FDependentObjectInfo& DependentObjectInfo : GetDependentObjectInfos(ObjectIndex))
 			{
-				Functor(DependentObjectIndex);
-				ForAllDependentObjectsRecursive(DependentObjectIndex, Functor);
+				Functor(DependentObjectInfo.NetRefIndex);
+				ForAllDependentObjectsRecursive(DependentObjectInfo.NetRefIndex, Functor);
 			}
 			return;
 		}
@@ -439,14 +439,9 @@ TArrayView<const FInternalNetRefIndex> FNetRefHandleManager::GetChildSubObjects(
 	return SubObjects.GetInternalIndexArray<FNetDependencyData::EArrayType::ChildSubObjects>(OwnerIndex);
 }
 
-TArrayView<const FInternalNetRefIndex> FNetRefHandleManager::GetDependentObjects(FInternalNetRefIndex ParentIndex) const
-{
-	return SubObjects.GetInternalIndexArray<FNetDependencyData::EArrayType::DependentObjects>(ParentIndex);
-}
-
 TArrayView<const FInternalNetRefIndex> FNetRefHandleManager::GetDependentObjectParents(FInternalNetRefIndex DependentIndex) const
 {
-	return SubObjects.GetInternalIndexArray<FNetDependencyData::EArrayType::ParentObjects>(DependentIndex);
+	return SubObjects.GetInternalIndexArray<FNetDependencyData::EArrayType::DependentParentObjects>(DependentIndex);
 }
 	
 }
diff --git a/Engine/Source/Runtime/Experimental/Iris/Core/Private/Iris/ReplicationSystem/ObjectReplicationBridge.cpp b/Engine/Source/Runtime/Experimental/Iris/Core/Private/Iris/ReplicationSystem/ObjectReplicationBridge.cpp
--- a/Engine/Source/Runtime/Experimental/Iris/Core/Private/Iris/ReplicationSystem/ObjectReplicationBridge.cpp
+++ b/Engine/Source/Runtime/Experimental/Iris/Core/Private/Iris/ReplicationSystem/ObjectReplicationBridge.cpp
@@ -415,15 +415,14 @@ void UObjectReplicationBridge::SetSubObjectNetCondition(FNetRefHandle SubObjectR
 	}
 }
 
-void UObjectReplicationBridge::AddDependentObject(FNetRefHandle ParentHandle, FNetRefHandle DependentHandle, EDependentWarnFlags WarnFlags)
+void UObjectReplicationBridge::AddDependentObject(FNetRefHandle ParentHandle, FNetRefHandle DependentHandle, UE::Net::EDependentObjectSchedulingHint SchedulingHint)
 {
 	using namespace UE::Net::Private;
 
 	FReplicationSystemInternal* ReplicationSystemInternal = GetReplicationSystem()->GetReplicationSystemInternal();
 	FNetRefHandleManager& LocalNetRefHandleManager = ReplicationSystemInternal->GetNetRefHandleManager();
 
-	// New logic to treat Dependent objects 
-	if (LocalNetRefHandleManager.AddDependentObject(ParentHandle, DependentHandle))
+	if (LocalNetRefHandleManager.AddDependentObject(ParentHandle, DependentHandle, SchedulingHint))
 	{
 		FReplicationFiltering& Filtering = ReplicationSystemInternal->GetFiltering();
 		const FInternalNetRefIndex DependentInternalIndex = LocalNetRefHandleManager.GetInternalIndex(DependentHandle);
diff --git a/Engine/Source/Runtime/Experimental/Iris/Core/Private/Iris/ReplicationSystem/ReplicationWriter.cpp b/Engine/Source/Runtime/Experimental/Iris/Core/Private/Iris/ReplicationSystem/ReplicationWriter.cpp
--- a/Engine/Source/Runtime/Experimental/Iris/Core/Private/Iris/ReplicationSystem/ReplicationWriter.cpp
+++ b/Engine/Source/Runtime/Experimental/Iris/Core/Private/Iris/ReplicationSystem/ReplicationWriter.cpp
@@ -372,7 +372,7 @@ void FReplicationWriter::Init(const FReplicationParameters& InParameters)
 	ObjectsPendingDestroy.Init(Parameters.MaxActiveReplicatedObjectCount);
 	ObjectsWithDirtyChanges.Init(Parameters.MaxActiveReplicatedObjectCount);
 	ObjectsInScope.Init(Parameters.MaxActiveReplicatedObjectCount);	
-	WriteContext.ObjectsWrittenThisTick.Init(Parameters.MaxActiveReplicatedObjectCount);
+	WriteContext.ObjectsWrittenThisPacket.Init(Parameters.MaxActiveReplicatedObjectCount);
 
 	// Attachments init
 	SetupReplicationInfoForAttachmentsToObjectsNotInScope();
@@ -799,6 +799,44 @@ void FReplicationWriter::UpdatePriorities(const float* UpdatedPriorities)
 	ObjectsWithDirtyChanges.ForAllSetBits(UpdatePriority);
 }
 
+void FReplicationWriter::ScheduleDependentObjects(uint32 Index, float ParentPriority, float* LocalPriorities, FScheduleObjectInfo* ScheduledObjectIndices, uint32& OutScheduledObjectCount)
+{
+	const float DependentObjectPriorityBump = UE_KINDA_SMALL_NUMBER;
+
+	for (const FDependentObjectInfo& DependentObjectInfo : NetRefHandleManager->GetDependentObjectInfos(Index))
+	{
+		const FInternalNetRefIndex DependentInternalIndex = DependentObjectInfo.NetRefIndex;
+		float UpdatedPriority = ParentPriority;
+
+		if (ObjectsWithDirtyChanges.GetBit(DependentInternalIndex))
+		{
+			const FReplicationInfo& DependentInfo = this->GetReplicationInfo(DependentInternalIndex);
+
+			const bool bReplicateBeforeParent =	(DependentObjectInfo.SchedulingHint == EDependentObjectSchedulingHint::ScheduleBeforeParent) || ((DependentObjectInfo.SchedulingHint == EDependentObjectSchedulingHint::ScheduleBeforeParentIfInitialState) && IsInitialState(DependentInfo.GetState()));
+
+			if (bReplicateBeforeParent)
+			{
+				// Bump prio of dependent object to be scheduled before its parent.
+				UpdatedPriority = FMath::Max(ParentPriority + DependentObjectPriorityBump, LocalPriorities[DependentInternalIndex]);
+				LocalPriorities[DependentInternalIndex] = UpdatedPriority;
+
+				// Schedule it, it does not matter if we add it to the scheduled list multiple times
+				FScheduleObjectInfo& ScheduledObjectInfo = ScheduledObjectIndices[OutScheduledObjectCount];
+				ScheduledObjectInfo.Index = DependentInternalIndex;
+				ScheduledObjectInfo.SortKey = UpdatedPriority;
+				++OutScheduledObjectCount;								
+			}
+		}
+		
+		// We go through all dependent objects here even though it might not be 100% correct, but it will make sure that we respect
+		// the scheduling order hint at least in relation to the parent, but a dependent object might also end up replicating before its parent`s parent
+		if (NetRefHandleManager->GetObjectsWithDependentObjectsInternalIndices().GetBit(DependentInternalIndex))
+		{
+			ScheduleDependentObjects(DependentInternalIndex, UpdatedPriority, LocalPriorities, ScheduledObjectIndices, OutScheduledObjectCount);
+		}
+	}
+}
+
 uint32 FReplicationWriter::ScheduleObjects(FScheduleObjectInfo* OutScheduledObjectIndices)
 {
 	IRIS_PROFILER_SCOPE(FReplicationWriter_ScheduleObjects);
@@ -814,8 +852,7 @@ uint32 FReplicationWriter::ScheduleObjects(FScheduleObjectInfo* OutScheduledObje
 	const FNetBitArray& UpdatedObjects = ObjectsWithDirtyChanges;
 	const FNetBitArray& SubObjects = NetRefHandleManager->GetSubObjectInternalIndices();
 
-	// Fill IndexList with all objects with dirty state that have positive priority excluding subobjects
-	auto FillIndexListFunc = [&LocalPriorities, &ScheduledObjectIndices, &ScheduledObjectCount](uint32 Index)
+	auto FillIndexListFunc = [&LocalPriorities, &ScheduledObjectIndices, &ScheduledObjectCount, this](uint32 Index)
 	{
 		const float UpdatedPriority = LocalPriorities[Index];
 
@@ -823,7 +860,16 @@ uint32 FReplicationWriter::ScheduleObjects(FScheduleObjectInfo* OutScheduledObje
 		ScheduledObjectInfo.Index = Index;
 		ScheduledObjectInfo.SortKey = UpdatedPriority;
 
-		ScheduledObjectCount += (UpdatedPriority >= FReplicationWriter::SchedulingThresholdPriority) ? 1 : 0;
+		if (UpdatedPriority >= FReplicationWriter::SchedulingThresholdPriority)
+		{
+			++ScheduledObjectCount;
+
+			// If we have dependent objects that needs to replicate before parent we need to schedule them as well.
+			if (NetRefHandleManager->GetObjectsWithDependentObjectsInternalIndices().GetBit(Index))
+			{
+				ScheduleDependentObjects(Index, UpdatedPriority, LocalPriorities, ScheduledObjectIndices, ScheduledObjectCount);
+			}
+		}
 	};
 
 	// Invoke functor for all updated objects that not are sub objects.
@@ -1595,8 +1641,6 @@ bool FReplicationWriter::CanSendObject(uint32 InternalIndex) const
 	const FReplicationInfo& Info = GetReplicationInfo(InternalIndex);
 	const EReplicatedObjectState State = Info.GetState();
 
-	const bool bIsInitialState = IsInitialState(State);
-
 	// Currently we do wait for CreateConfirmation before sending more data
 	// We might want to change this and allow "bombing" creation info until we get confirmation to minimize latency
 	// We also prevent objects from being transmitted if they are waiting on destroy/tear-off confirmation or cancelling destroy.
@@ -1625,16 +1669,36 @@ bool FReplicationWriter::CanSendObject(uint32 InternalIndex) const
 		}
 	}
 
-	// Currently we need to enforce a strict dependency on the initial dependencies
-	for (FInternalNetRefIndex DependentInternalIndex : NetRefHandleManager->GetDependentObjects(InternalIndex))
+	// Currently we enforce a strict dependency on the state of initial dependent objects unless they are already serialized in the same packet
+	if (NetRefHandleManager->GetObjectsWithDependentObjectsInternalIndices().GetBit(InternalIndex))
 	{
-		const FReplicationInfo& DependentInfo = GetReplicationInfo(DependentInternalIndex);
-		const bool bIsDependentObjectInitialState = IsInitialState(DependentInfo.GetState());
-
-		if (bIsDependentObjectInitialState && !CanSendObject(DependentInternalIndex))
+		for (const FDependentObjectInfo DependentObjectInfo : NetRefHandleManager->GetDependentObjectInfos(InternalIndex))
 		{
-			UE_LOG(LogIris, Log, TEXT("ReplicationWriter: Cannot send internal index (%u) due to waiting on init dependency internal index (%d)"), InternalIndex, DependentInternalIndex);
-			return false;
+			const FInternalNetRefIndex DependentInternalIndex = DependentObjectInfo.NetRefIndex;
+
+			// If the dependent object already has been written in this packet we are fine.
+			if (WriteContext.ObjectsWrittenThisPacket.GetBit(DependentInternalIndex))
+			{
+				continue;
+			}
+
+			const FReplicationInfo& DependentReplicationInfo = GetReplicationInfo(DependentInternalIndex);
+			if (IsInitialState(DependentReplicationInfo.GetState()))
+			{
+				// if we cannot send the initial dependent object we must wait until we can.
+				if (!CanSendObject(DependentInternalIndex))
+				{
+					UE_LOG(LogIris, Log, TEXT("ReplicationWriter: Cannot send internal index (%u) due to waiting on init dependency internal index (%d)"), InternalIndex, DependentInternalIndex);
+					return false;
+				}
+
+				// if the dependent object are scheduled before parent and did not fit in this packet, we cannot write the parent either and have to wait until creation is confirmed
+				if ((DependentObjectInfo.SchedulingHint == EDependentObjectSchedulingHint::ScheduleBeforeParent) && ObjectsWithDirtyChanges.GetBit(DependentInternalIndex))
+				{
+					UE_LOG(LogIris, Log, TEXT("ReplicationWriter: Cannot send internal index (%u) due to waiting on ScheduleBefore dependency internal index (%d)"), InternalIndex, DependentInternalIndex);
+					return false;
+				}
+			}
 		}
 	}
 
@@ -2047,10 +2111,11 @@ FReplicationWriter::EWriteObjectStatus FReplicationWriter::WriteObjectInBatch(FN
 
 	// If we have any dependent objects pending creation we currently include them in the batch in order to guarantee that they get sent in the same packet
 	// When we implement support for proper dependencies this can be changed.
-	for (FInternalNetRefIndex DependentInternalIndex : NetRefHandleManager->GetDependentObjects(InternalIndex))
+	for (const FDependentObjectInfo DependentObjectInfo : NetRefHandleManager->GetDependentObjectInfos(InternalIndex))
 	{
+		const FInternalNetRefIndex DependentInternalIndex = DependentObjectInfo.NetRefIndex;
 		const bool bIsDependentInitialState = IsInitialState(GetReplicationInfo(DependentInternalIndex).GetState());
-		if (bIsDependentInitialState && !WriteContext.ObjectsWrittenThisTick.GetBit(DependentInternalIndex))
+		if (bIsDependentInitialState && !WriteContext.ObjectsWrittenThisPacket.GetBit(DependentInternalIndex))
 		{
 			UE_NET_TRACE_SCOPE(DependentObjectData, *Context.GetBitStreamWriter(), Context.GetTraceCollector(), ENetTraceVerbosity::VeryVerbose);
 			EWriteObjectStatus DependentObjectWriteStatus = WriteObjectInBatch(Context, DependentInternalIndex, WriteObjectFlags, OutBatchInfo);
@@ -2397,7 +2462,7 @@ uint32 FReplicationWriter::WriteObjects(FNetSerializationContext& Context)
 
 	auto SendObjectFunction = [this, &Context, &WrittenObjectCount](FInternalNetRefIndex InternalIndex)
 	{
-		if (!this->WriteContext.ObjectsWrittenThisTick.GetBit(InternalIndex) && this->CanSendObject(InternalIndex))
+		if (!this->WriteContext.ObjectsWrittenThisPacket.GetBit(InternalIndex) && this->CanSendObject(InternalIndex))
 		{
 #if UE_NET_IRIS_CSV_STATS && CSV_PROFILER
 			FIrisStatsTimer Timer;
@@ -2476,6 +2541,9 @@ uint32 FReplicationWriter::WriteObjects(FNetSerializationContext& Context)
 	WriteContext.CurrentIndex = ObjectListIt;
 	WriteContext.SortedObjectCount = SortedCount;
 
+	// Reset objects written this packet
+	WriteContext.ObjectsWrittenThisPacket.Reset();
+
 	return WrittenObjectCount;
 }
 
@@ -2534,12 +2602,13 @@ int FReplicationWriter::HandleObjectBatchSuccess(const FBatchInfo& BatchInfo, FR
 		}
 
 		// Schedule rest of dependent objects for replication, note there is no guarantee that they will replicate in same packet
-		for (FInternalNetRefIndex DependentInternalIndex : NetRefHandleManager->GetDependentObjects(BatchObjectInfo.InternalIndex))
+		for (const FDependentObjectInfo DependentObjectInfo : NetRefHandleManager->GetDependentObjectInfos(BatchObjectInfo.InternalIndex))
 		{
-			if (ObjectsWithDirtyChanges.GetBit(DependentInternalIndex) && !WriteContext.ObjectsWrittenThisTick.GetBit(DependentInternalIndex))
+			const FInternalNetRefIndex DependentInternalIndex = DependentObjectInfo.NetRefIndex;
+			if (ObjectsWithDirtyChanges.GetBit(DependentInternalIndex) && !WriteContext.ObjectsWrittenThisPacket.GetBit(DependentInternalIndex))
 			{
 				WriteContext.DependentObjectsPendingSend.Push(DependentInternalIndex);
-				// Bumping the scheduling priority here will make sure that we will be scheduled the next update if we are not allowed to replicate this frame
+				// Bumping the scheduling priority here will make sure that they will be scheduled the next update if we are not allowed to replicate this frame
 				SchedulingPriorities[DependentInternalIndex] = FMath::Max(SchedulingPriorities[BatchObjectInfo.InternalIndex], SchedulingPriorities[DependentInternalIndex]);
 			}
 		}			
@@ -2568,7 +2637,7 @@ int FReplicationWriter::HandleObjectBatchSuccess(const FBatchInfo& BatchInfo, FR
 		if (BatchObjectInfo.InternalIndex != ObjectIndexForOOBAttachment)
 		{
 			// Mark this object as written this tick to avoid sending it multiple times
-			WriteContext.ObjectsWrittenThisTick.SetBit(BatchObjectInfo.InternalIndex);
+			WriteContext.ObjectsWrittenThisPacket.SetBit(BatchObjectInfo.InternalIndex);
 		}
 
 #if UE_NET_IRIS_CSV_STATS && CSV_PROFILER
@@ -2665,9 +2734,6 @@ UDataStream::EWriteResult FReplicationWriter::BeginWrite()
 	// Reset dependent object array
 	WriteContext.DependentObjectsPendingSend.Reset();
 
-	// Reset objects written this tick
-	WriteContext.ObjectsWrittenThisTick.Reset();
-
 	// $IRIS TODO: LinearAllocator/ScratchPad?
 	// Allocate space for indices to send
 	// This should be allocated from frame temp allocator and be cleaned up end of frame, we might want this data to persist over multiple write calls but not over multiple frames 
diff --git a/Engine/Source/Runtime/Experimental/Iris/Core/Private/Iris/ReplicationSystem/ReplicationWriter.h b/Engine/Source/Runtime/Experimental/Iris/Core/Private/Iris/ReplicationSystem/ReplicationWriter.h
--- a/Engine/Source/Runtime/Experimental/Iris/Core/Private/Iris/ReplicationSystem/ReplicationWriter.h
+++ b/Engine/Source/Runtime/Experimental/Iris/Core/Private/Iris/ReplicationSystem/ReplicationWriter.h
@@ -207,7 +207,7 @@ private:
 		FWriteContext() : bIsValid(0) {}
 
 		// Objects we have written in this packet batch to avoid sending same object multiple times
-		FNetBitArray ObjectsWrittenThisTick;
+		FNetBitArray ObjectsWrittenThisPacket;
 		// DependentObjets that we should try to write this packet batch, aka. allow overcommit if we have pending DependentObjects when the packet is full
 		TArray<uint32, TInlineAllocator<32>> DependentObjectsPendingSend;
 		// Scheduled objects
@@ -370,6 +370,8 @@ private:
 
 	void CommitBatchRecord(const FBatchRecord& BatchRecord);
 
+	void ScheduleDependentObjects(uint32 Index, float ParentPriority, float* LocalPriorities, FScheduleObjectInfo* ScheduledObjectIndices, uint32& OutScheduledObjectCount);
+
 	uint32 ScheduleObjects(FScheduleObjectInfo* ScheduledObjectIndices);
 	
 	// Partial sort of OutScheduledObjectIndices, will sort at most PartialSortObjectCount objects
diff --git a/Engine/Source/Runtime/Experimental/Iris/Core/Public/Iris/Core/NetObjectReference.h b/Engine/Source/Runtime/Experimental/Iris/Core/Public/Iris/Core/NetObjectReference.h
--- a/Engine/Source/Runtime/Experimental/Iris/Core/Public/Iris/Core/NetObjectReference.h
+++ b/Engine/Source/Runtime/Experimental/Iris/Core/Public/Iris/Core/NetObjectReference.h
@@ -32,7 +32,7 @@ enum class ENetObjectReferenceTraits : uint32
 ENUM_CLASS_FLAGS(ENetObjectReferenceTraits);
 
 /**
- * A refererence to a network addressable object.
+ * A reference to a network addressable object.
  */
 class FNetObjectReference
 {
diff --git a/Engine/Source/Runtime/Experimental/Iris/Core/Public/Iris/ReplicationSystem/ObjectReplicationBridge.h b/Engine/Source/Runtime/Experimental/Iris/Core/Public/Iris/ReplicationSystem/ObjectReplicationBridge.h
--- a/Engine/Source/Runtime/Experimental/Iris/Core/Public/Iris/ReplicationSystem/ObjectReplicationBridge.h
+++ b/Engine/Source/Runtime/Experimental/Iris/Core/Public/Iris/ReplicationSystem/ObjectReplicationBridge.h
@@ -7,6 +7,7 @@
 #include "Net/Core/NetHandle/NetHandle.h"
 #include "Iris/ReplicationSystem/ReplicationBridge.h"
 #include "Delegates/IDelegateInstance.h"
+#include "Iris/ReplicationSystem/ReplicationSystemTypes.h"
 
 #include "ObjectReplicationBridge.generated.h"
 
@@ -114,21 +115,13 @@ public:
 	{
 		virtual ~FCreationHeader() {};
 	};
-
-	// Dependent objects supprt 
-
-	enum class EDependentWarnFlags : uint32
-	{
-		None = 1 << 0U,
-		WarnAlreadyDependent = 1 << None, 
-	};
 	
 	/**
 	 * Adds a dependent object. A dependent object can replicate separately or if a parent replicates.
 	 * Dependent objects cannot be filtered out by dynamic filtering unless the parent is also filtered out.
 	 * @note: There is no guarantee that the data will end up in the same packet so it is a very loose form of dependency.
 	 */
-	IRISCORE_API void AddDependentObject(FNetRefHandle Parent, FNetRefHandle DependentObject, EDependentWarnFlags WarnFlags = EDependentWarnFlags::WarnAlreadyDependent);
+	IRISCORE_API void AddDependentObject(FNetRefHandle Parent, FNetRefHandle DependentObject, UE::Net::EDependentObjectSchedulingHint SchedulingHint = UE::Net::EDependentObjectSchedulingHint::Default);
 
 	/** Remove dependent object from parent. The dependent object will function as a standard standalone replicated object. */
 	IRISCORE_API void RemoveDependentObject(FNetRefHandle Parent, FNetRefHandle DependentObject);
diff --git a/Engine/Source/Runtime/Experimental/Iris/Core/Public/Iris/ReplicationSystem/ReplicationSystemTypes.h b/Engine/Source/Runtime/Experimental/Iris/Core/Public/Iris/ReplicationSystem/ReplicationSystemTypes.h
--- a/Engine/Source/Runtime/Experimental/Iris/Core/Public/Iris/ReplicationSystem/ReplicationSystemTypes.h
+++ b/Engine/Source/Runtime/Experimental/Iris/Core/Public/Iris/ReplicationSystem/ReplicationSystemTypes.h
@@ -13,4 +13,16 @@ enum class ENetObjectDeltaCompressionStatus : unsigned
 	Allow,
 };
 
+enum class EDependentObjectSchedulingHint : uint8
+{
+	// Default behavior, dependent object will be scheduled to replicate if parent is replicated, if the dependent object has not yet been replicated it will be replicated in the same batch as the parent
+	Default = 0,
+
+	// Dependent object will be scheduled to replicate before parent is replicated, if the dependent has data to send and has not yet been replicated the parent will only be scheduled if they both fit in same packet
+	ScheduleBeforeParent,
+
+	// Not yet replicated dependent object will behave as ReplicateBeforeParent otherwise it will be scheduled to replicate if the parent is replicated and scheduled after the parent
+	ScheduleBeforeParentIfInitialState,
+};
+
 }
