Iris
- Fixed warning in TestBoolNetSerializer

diff --git a/Engine/Plugins/Runtime/ReplicationSystemTestPlugin/Source/Private/Tests/Serialization/TestBoolNetSerializer.cpp b/Engine/Plugins/Runtime/ReplicationSystemTestPlugin/Source/Private/Tests/Serialization/TestBoolNetSerializer.cpp
--- a/Engine/Plugins/Runtime/ReplicationSystemTestPlugin/Source/Private/Tests/Serialization/TestBoolNetSerializer.cpp
+++ b/Engine/Plugins/Runtime/ReplicationSystemTestPlugin/Source/Private/Tests/Serialization/TestBoolNetSerializer.cpp
@@ -145,7 +145,8 @@ void FTestBoolNetSerializer::TestSerialize()
 
 	FBoolNetSerializerConfig Config;
 	constexpr bool bQuantizedCompare = false;
-	TFunctionRef<bool(NetSerializerValuePointer Value0, NetSerializerValuePointer Value1)> CompareFunc = [](NetSerializerValuePointer Value0, NetSerializerValuePointer Value1) -> bool { return *reinterpret_cast<bool*>(Value0) == *reinterpret_cast<bool*>(Value1); };
+	auto&& CompareFunction = [](NetSerializerValuePointer Value0, NetSerializerValuePointer Value1) -> bool { return *reinterpret_cast<bool*>(Value0) == *reinterpret_cast<bool*>(Value1); };
+	TFunctionRef<bool(NetSerializerValuePointer Value0, NetSerializerValuePointer Value1)> CompareFunc = CompareFunction;
 	TTestNetSerializerFixture<PrintBoolNetSerializerConfig, bool>::TestSerialize(Values, ExpectedValues, ValueCount, Config, bQuantizedCompare, CompareFunc);
 }
 
