Iris - Fixed bug with splitting very large NetBlobs, larger than 128KiB.

Thanks to [REDACTED_PERSON] for reporting the issue.

diff --git a/Engine/Plugins/Runtime/ReplicationSystemTestPlugin/Source/Private/Tests/ReplicationSystem/NetBlob/MockNetBlob.cpp b/Engine/Plugins/Runtime/ReplicationSystemTestPlugin/Source/Private/Tests/ReplicationSystem/NetBlob/MockNetBlob.cpp
--- a/Engine/Plugins/Runtime/ReplicationSystemTestPlugin/Source/Private/Tests/ReplicationSystem/NetBlob/MockNetBlob.cpp
+++ b/Engine/Plugins/Runtime/ReplicationSystemTestPlugin/Source/Private/Tests/ReplicationSystem/NetBlob/MockNetBlob.cpp
@@ -84,6 +84,13 @@ void FMockNetBlob::Deserialize(FNetSerializationContext& Context)
 
 }
 
+// UMockSequentialPartialNetBlobHandlerConfig
+UMockSequentialPartialNetBlobHandlerConfig::UMockSequentialPartialNetBlobHandlerConfig()
+{
+	// Overriding MaxPartCount
+	MaxPartCount = 32768U;
+}
+
 // UMockNetBlobHandler
 UMockNetBlobHandler::UMockNetBlobHandler()
 : CallCounts({})
diff --git a/Engine/Plugins/Runtime/ReplicationSystemTestPlugin/Source/Private/Tests/ReplicationSystem/NetBlob/MockNetBlob.h b/Engine/Plugins/Runtime/ReplicationSystemTestPlugin/Source/Private/Tests/ReplicationSystem/NetBlob/MockNetBlob.h
--- a/Engine/Plugins/Runtime/ReplicationSystemTestPlugin/Source/Private/Tests/ReplicationSystem/NetBlob/MockNetBlob.h
+++ b/Engine/Plugins/Runtime/ReplicationSystemTestPlugin/Source/Private/Tests/ReplicationSystem/NetBlob/MockNetBlob.h
@@ -64,6 +64,15 @@ private:
 	mutable FCallCounts CallCounts;
 };
 
+UCLASS()
+class UMockSequentialPartialNetBlobHandlerConfig : public USequentialPartialNetBlobHandlerConfig
+{
+	GENERATED_BODY()
+
+public:
+	UMockSequentialPartialNetBlobHandlerConfig();
+};
+
 UCLASS()
 class UMockSequentialPartialNetBlobHandler final : public USequentialPartialNetBlobHandler
 {
diff --git a/Engine/Plugins/Runtime/ReplicationSystemTestPlugin/Source/Private/Tests/ReplicationSystem/NetBlob/PartialNetBlobTestFixture.cpp b/Engine/Plugins/Runtime/ReplicationSystemTestPlugin/Source/Private/Tests/ReplicationSystem/NetBlob/PartialNetBlobTestFixture.cpp
--- a/Engine/Plugins/Runtime/ReplicationSystemTestPlugin/Source/Private/Tests/ReplicationSystem/NetBlob/PartialNetBlobTestFixture.cpp
+++ b/Engine/Plugins/Runtime/ReplicationSystemTestPlugin/Source/Private/Tests/ReplicationSystem/NetBlob/PartialNetBlobTestFixture.cpp
@@ -30,7 +30,7 @@ void FPartialNetBlobTestFixture::RegisterNetBlobHandlers(FReplicationSystemTestN
 	{
 		FSequentialPartialNetBlobHandlerInitParams InitParams = {};
 		InitParams.ReplicationSystem = RepSys;
-		InitParams.Config = GetDefault<USequentialPartialNetBlobHandlerConfig>();
+		InitParams.Config = GetDefault<UMockSequentialPartialNetBlobHandlerConfig>();
 		UMockSequentialPartialNetBlobHandler* BlobHandler = NewObject<UMockSequentialPartialNetBlobHandler>();
 		BlobHandler->Init(InitParams);
 		const bool bPartialNetBlobHandlerWasRegistered = RegisterNetBlobHandler(RepSys, BlobHandler);
diff --git a/Engine/Plugins/Runtime/ReplicationSystemTestPlugin/Source/Private/Tests/ReplicationSystem/NetBlob/TestPartialNetBlob.cpp b/Engine/Plugins/Runtime/ReplicationSystemTestPlugin/Source/Private/Tests/ReplicationSystem/NetBlob/TestPartialNetBlob.cpp
--- a/Engine/Plugins/Runtime/ReplicationSystemTestPlugin/Source/Private/Tests/ReplicationSystem/NetBlob/TestPartialNetBlob.cpp
+++ b/Engine/Plugins/Runtime/ReplicationSystemTestPlugin/Source/Private/Tests/ReplicationSystem/NetBlob/TestPartialNetBlob.cpp
@@ -53,6 +53,29 @@ UE_NET_TEST_FIXTURE(FPartialNetBlobTestFixture, CanSplitNetBlobInNeedOfSplitting
 	UE_NET_ASSERT_EQ(MockNetBlobHandler->GetFunctionCallCounts().OnNetBlobReceived, ExpectedMockNetBlobReceiveCount);
 }
 
+UE_NET_TEST_FIXTURE(FPartialNetBlobTestFixture, CanSplitHugeNetBlob)
+{
+	FTestContext ServerContext;
+	SetupTestContext(ServerContext, Server);
+
+	constexpr uint32 PayloadBitCount = 384U*1024U*8U; // 384KiB
+	const TRefCountPtr<FNetBlob>& Blob = CreateUnreliableMockNetBlob(PayloadBitCount);
+
+	TArray<TRefCountPtr<FNetBlob>> PartialNetBlobs;
+	MockSequentialPartialNetBlobHandler->SplitNetBlob(Blob, PartialNetBlobs);
+	UE_NET_ASSERT_GE(PartialNetBlobs.Num(), 2);
+
+	for (const TRefCountPtr<FNetBlob>& PartialNetBlob : PartialNetBlobs)
+	{
+		ServerContext.SerializationContext.GetNetBlobReceiver()->OnNetBlobReceived(ServerContext.SerializationContext, PartialNetBlob);
+	}
+
+	const uint32 ExpectedPartialNetBlobReceiveCount = PartialNetBlobs.Num();
+	constexpr uint32 ExpectedMockNetBlobReceiveCount = 1U;
+	UE_NET_ASSERT_EQ(MockSequentialPartialNetBlobHandler->GetFunctionCallCounts().OnNetBlobReceived, ExpectedPartialNetBlobReceiveCount);
+	UE_NET_ASSERT_EQ(MockNetBlobHandler->GetFunctionCallCounts().OnNetBlobReceived, ExpectedMockNetBlobReceiveCount);
+}
+
 UE_NET_TEST_FIXTURE(FPartialNetBlobTestFixture, MissingFirstPartialNetBlobCausesError)
 {
 	FTestContext ServerContext;
diff --git a/Engine/Source/Runtime/Experimental/Iris/Core/Private/Iris/ReplicationSystem/NetBlob/PartialNetBlob.cpp b/Engine/Source/Runtime/Experimental/Iris/Core/Private/Iris/ReplicationSystem/NetBlob/PartialNetBlob.cpp
--- a/Engine/Source/Runtime/Experimental/Iris/Core/Private/Iris/ReplicationSystem/NetBlob/PartialNetBlob.cpp
+++ b/Engine/Source/Runtime/Experimental/Iris/Core/Private/Iris/ReplicationSystem/NetBlob/PartialNetBlob.cpp
@@ -203,6 +203,7 @@ bool FPartialNetBlob::SplitNetBlob(const FNetSerializationContext& Context, cons
 			}
 
 			CurrentPayloadBitCount = NewPayloadBitCount;
+			continue;
 		}
 
 		// Adjust the payload bit count to the final value
diff --git a/Engine/Source/Runtime/Experimental/Iris/Core/Public/Iris/ReplicationSystem/NetBlob/SequentialPartialNetBlobHandler.h b/Engine/Source/Runtime/Experimental/Iris/Core/Public/Iris/ReplicationSystem/NetBlob/SequentialPartialNetBlobHandler.h
--- a/Engine/Source/Runtime/Experimental/Iris/Core/Public/Iris/ReplicationSystem/NetBlob/SequentialPartialNetBlobHandler.h
+++ b/Engine/Source/Runtime/Experimental/Iris/Core/Public/Iris/ReplicationSystem/NetBlob/SequentialPartialNetBlobHandler.h
@@ -29,7 +29,7 @@ public:
 	uint32 GetMaxPartCount() const { return MaxPartCount; }
 	uint64 GetTotalMaxPayloadBitCount() const { return GetMaxPartBitCount()*uint64(GetMaxPartCount()); }
 
-private:
+protected:
 	/** How many bits a PartialNetBlob payload can hold at most. Cannot exceed 65535, but anything near the max packet size is discouraged as it is unlikely to fit. Keep it a power of two. */
 	UPROPERTY(Config)
 	uint32 MaxPartBitCount = 128*8;
