Iris
- Implemented delta variant of WriteSparseBitArray

diff --git a/Engine/Plugins/Runtime/ReplicationSystemTestPlugin/Source/Private/Tests/Serialization/TestCompressedChangeMask.cpp b/Engine/Plugins/Runtime/ReplicationSystemTestPlugin/Source/Private/Tests/Serialization/TestCompressedChangeMask.cpp
--- a/Engine/Plugins/Runtime/ReplicationSystemTestPlugin/Source/Private/Tests/Serialization/TestCompressedChangeMask.cpp
+++ b/Engine/Plugins/Runtime/ReplicationSystemTestPlugin/Source/Private/Tests/Serialization/TestCompressedChangeMask.cpp
@@ -19,6 +19,9 @@ public:
 
 	typedef void (*WriteSparseBitArrayFunc)(FNetBitStreamWriter* Writer, const uint32* Data, uint32 BitCount, ESparseBitArraySerializationHint Hint);
 	typedef void (*ReadSparseBitArrayFunc)(FNetBitStreamReader* Reader, uint32* OutData, uint32 BitCount, ESparseBitArraySerializationHint Hint);
+
+	typedef void (*WriteSparseBitArrayDeltaFunc)(FNetBitStreamWriter* Writer, const uint32* Data, const uint32* BaseData, uint32 BitCount);
+	typedef void (*ReadSparseBitArrayDeltaFunc)(FNetBitStreamReader* Reader, uint32* OutData, const uint32* BaseData, uint32 BitCount);
 	
 protected:
 	class FChangeMaskBuffer
@@ -76,6 +79,58 @@ protected:
 		}
 	}
 
+	template <typename T>
+	bool ValidateDeltaTestData(const T& DeltaTestData, const T& BaseTestData)
+	{
+		const uint32 TestDataWordCount = UE_ARRAY_COUNT(DeltaTestData);
+		if (TestDataWordCount != UE_ARRAY_COUNT(BaseTestData))
+		{
+			return false;
+		}
+
+		bool bResult = true;
+		SIZE_T Offset = 0;
+
+		while (Offset + 2U < TestDataWordCount)
+		{
+			const FWordType BitCount = DeltaTestData[Offset];
+			const FWordType Invert = DeltaTestData[Offset + 1];
+
+			if (BitCount != BaseTestData[Offset])
+			{
+				return false;
+			}
+
+			Offset += FNetBitArrayView::CalculateRequiredWordCount(BitCount) + 2U;
+		}
+
+		return Offset == TestDataWordCount;
+	}
+
+
+	template<typename FunctorT, typename T> void ForAllEntriesDelta(FunctorT&& Functor, const T& DeltaTestData, const T& BaseTestData) const
+	{
+		const uint32 TestDataWordCount = UE_ARRAY_COUNT(DeltaTestData);
+
+		check(TestDataWordCount == UE_ARRAY_COUNT(BaseTestData));
+
+		bool bResult = true;
+		SIZE_T Offset = 0;
+
+		while (Offset + 2U < TestDataWordCount)
+		{
+			const FWordType BitCount = DeltaTestData[Offset];
+			const FWordType Invert = BaseTestData[Offset + 1];
+			const FWordType* DeltaData = &DeltaTestData[Offset + 2];
+			const FWordType* BaseData = &BaseTestData[Offset + 2];
+
+			Functor(BitCount, Invert, DeltaData, BaseData);
+
+			Offset += FNetBitArrayView::CalculateRequiredWordCount(BitCount) + 2U;
+		}
+	}
+
+
 	FChangeMaskBuffer WriteBuffer;
 	FChangeMaskBuffer TempChangeMaskData;
 
@@ -205,7 +260,7 @@ UE_NET_TEST_FIXTURE(FCompressedChangeMaskTest, WriteAndReadChangeMasks)
 		// Read changemask
 		ReadFunc(&Reader, TempChangeMaskData.GetBuffer(), BitCount, Invert ? ESparseBitArraySerializationHint::ContainsMostlyOnes : ESparseBitArraySerializationHint::None);
 
-		UE_NET_ASSERT_EQ(true, Reader.IsOverflown() == false);
+		UE_NET_ASSERT_FALSE(Reader.IsOverflown());
 		UE_NET_ASSERT_EQ(Writer.GetPosBits(), Reader.GetPosBits());
 
 		// Verify change mask
@@ -219,4 +274,87 @@ UE_NET_TEST_FIXTURE(FCompressedChangeMaskTest, WriteAndReadChangeMasks)
 	ForAllEntries(TestWriteAndReadFunc, VeryLargeChangeMaskTestData);
 }
 
+UE_NET_TEST_FIXTURE(FCompressedChangeMaskTest, WriteAndReadSparseBitArrayDelta)
+{
+	WriteSparseBitArrayDeltaFunc WriteFunc = WriteSparseBitArrayDelta;
+	ReadSparseBitArrayDeltaFunc ReadFunc = ReadSparseBitArrayDelta;
+
+	// Test some basic scenarios for delta
+
+	const FCompressedChangeMaskTest::FWordType BaseTestData[] =
+	{
+		// Set one bit in each word
+		/* BitCount */ 33U, /* Invert */ 0U, /* Data */ 0x0U, 0x0U,
+		// Set highest bit
+		/* BitCount */ 32U, /* Invert */ 0U, /* Data */ 0x1U,
+		// Set many bit
+		/* BitCount */ 33U, /* Invert */ 0U, /* Data */ 0x0U, 0x0U,
+		// Clear single bit
+		/* BitCount */ 32U, /* Invert */ 0U, /* Data */ 0xFFFFFFFFU,
+		// Test inverted clear single bit (invert does not make much sense as we delta using xor
+		/* BitCount */ 32U, /* Invert */ 1U, /* Data */ 0xFFFFFFFFU,
+		// Test no change
+		/* BitCount */ 33U, /* Invert */ 0U, /* Data */ 0x0U, 0x0U,
+		// Test no change
+		/* BitCount */ 33U, /* Invert */ 0U, /* Data */ 0xFFFFFFFFU, 0x1U, 
+		// Test all changed
+		/* BitCount */ 33U, /* Invert */ 0U, /* Data */ 0x0U, 0x0U,
+		// Test all changed
+		/* BitCount */ 33U, /* Invert */ 0U, /* Data */ 0xFFFFFFFFU, 0x1U, 
+	};
+
+	const FCompressedChangeMaskTest::FWordType DeltaTestData[] =
+	{
+		// Set one bit in each word
+		/* BitCount */ 33U, /* Invert */ 0U, /* Data */ 0x1U, 0x1U,
+		// Set highest bit
+		/* BitCount */ 32U, /* Invert */ 0U, /* Data */ 0x1U | 0x8U,
+		// Set many bit
+		/* BitCount */ 33U, /* Invert */ 0U, /* Data */ 0xFFFFFFFFU, 0x0U,
+		// Clear single bit
+		/* BitCount */ 32U, /* Invert */ 0U, /* Data */ 0xFFFFFFFEU,
+		// Test inverted clear single bit (invert does not make much sense as we delta using xor
+		/* BitCount */ 32U, /* Invert */ 1U, /* Data */ 0xFFFFFFFFU & 0xFFFFFFFEU,
+		// Test no change
+		/* BitCount */ 33U, /* Invert */ 0U, /* Data */ 0x0U, 0x0U,
+		// Test no change
+		/* BitCount */ 33U, /* Invert */ 0U, /* Data */ 0xFFFFFFFFU, 0x1U, 
+		// Test all changed
+		/* BitCount */ 33U, /* Invert */ 0U, /* Data */ 0xFFFFFFFFU, 0x1U,
+		// Test all changed
+		/* BitCount */ 33U, /* Invert */ 0U, /* Data */ 0x0U, 0x0U,
+	};
+
+
+	UE_NET_ASSERT_TRUE(ValidateDeltaTestData(DeltaTestData, BaseTestData));	
+
+	auto&& TestWriteAndReadFunc = [this, WriteFunc, ReadFunc](FWordType BitCount, FWordType Invert, const FWordType* Data, const FWordType* BaseData)
+	{		
+		Writer.InitBytes(WriteBuffer.GetBuffer(), WriteBuffer.GetBufferCapacity());
+		WriteFunc(&Writer, Data, BaseData, BitCount);
+		Writer.CommitWrites();
+
+		UE_NET_ASSERT_FALSE(Writer.IsOverflown());
+
+		Reader.InitBits(WriteBuffer.GetBuffer(), Writer.GetPosBits());
+
+		// Init changemask
+		FNetBitArrayView RcvdChangeMask((FWordType*)TempChangeMaskData.GetBuffer(), BitCount, FNetBitArrayView::ResetOnInit);
+
+		// Read changemask
+		ReadFunc(&Reader, TempChangeMaskData.GetBuffer(), BaseData, BitCount);
+
+		UE_NET_ASSERT_FALSE(Reader.IsOverflown());
+		UE_NET_ASSERT_EQ(Writer.GetPosBits(), Reader.GetPosBits());
+
+		// Verify change mask
+		const FNetBitArrayView OriginalChangeMask = MakeNetBitArrayView(Data, BitCount);
+
+		UE_NET_ASSERT_TRUE(OriginalChangeMask == RcvdChangeMask);
+	};
+
+	ForAllEntriesDelta(TestWriteAndReadFunc, DeltaTestData, BaseTestData);
+}
+
+
 }
diff --git a/Engine/Source/Runtime/Experimental/Iris/Core/Private/Iris/Serialization/NetBitStreamUtil.cpp b/Engine/Source/Runtime/Experimental/Iris/Core/Private/Iris/Serialization/NetBitStreamUtil.cpp
--- a/Engine/Source/Runtime/Experimental/Iris/Core/Private/Iris/Serialization/NetBitStreamUtil.cpp
+++ b/Engine/Source/Runtime/Experimental/Iris/Core/Private/Iris/Serialization/NetBitStreamUtil.cpp
@@ -590,6 +590,97 @@ void ReadSparseBitArray(FNetBitStreamReader* Reader, uint32* OutData, uint32 Bit
 	}
 }
 
+template<typename WriteSparseUint32Func>
+void WriteSparseBitArrayDelta(FNetBitStreamWriter* Writer, const uint32* Data, const uint32* OldData, uint32 BitCount, WriteSparseUint32Func&& WriteSparseUint32Function)
+{
+	checkSlow(BitCount <= SerializeSparseArrayMaxBitCount);
+
+	using StorageWordType = FNetBitArrayBase::StorageWordType;
+	const uint32 WordBitCount = FNetBitArrayBase::WordBitCount;
+	const uint32 WordCount = FNetBitArrayView::CalculateRequiredWordCount(BitCount);
+	const StorageWordType LastWordMask = ~StorageWordType(0) >> (uint32(-int32(BitCount)) & (WordBitCount - 1));
+	const uint32 LastWordIt = WordCount - 1;
+
+	StorageWordType DeltaData[Private::SerializeSparseArrayMaxWordCount];
+
+	// Build mask
+	uint32 DirtyWordMask = uint32(0);
+	{
+		for (uint32 WordIt = 0, CurrentBitMask = 1U; WordIt < WordCount; ++WordIt, CurrentBitMask <<= 1)
+		{
+			const StorageWordType CurrentWord = (Data[WordIt] ^ OldData[WordIt]) & ((WordIt == LastWordIt) ? LastWordMask : ~0U);
+			DirtyWordMask |= CurrentWord ? CurrentBitMask : 0U;
+			DeltaData[WordIt] = CurrentWord;
+		}
+	}
+
+	// Write Mask
+	Writer->WriteBits(DirtyWordMask, WordCount);
+
+	// Encode dirty words
+	{
+		// Full words
+		uint32 CurrentMaskBit = 1U;
+		uint32 WordIt = 0U;
+		uint32 RemainingBits = BitCount;
+
+		while (RemainingBits >= 32U)
+		{
+			if (DirtyWordMask & CurrentMaskBit)
+			{
+				WriteSparseUint32Function(Writer, DeltaData[WordIt], 32U);
+			}
+			++WordIt;
+			CurrentMaskBit <<= 1;
+			RemainingBits -= 32U;
+		}
+		// Last word
+		if (RemainingBits && (DirtyWordMask & CurrentMaskBit))
+		{
+			const StorageWordType CurrentWord = DeltaData[WordIt] & LastWordMask;
+			WriteSparseUint32Function(Writer, CurrentWord, RemainingBits);
+		}
+	}
+}
+
+template<typename ReadSparseUint32Func>
+void ReadSparseBitArrayDelta(FNetBitStreamReader* Reader, uint32* OutData, const uint32* OldData, uint32 BitCount, ReadSparseUint32Func&& ReadSparseUint32Function)
+{
+	checkSlow(BitCount <= SerializeSparseArrayMaxBitCount);
+
+	using StorageWordType = FNetBitArrayBase::StorageWordType;
+
+	const uint32 WordBitCount = FNetBitArrayBase::WordBitCount;
+	const uint32 WordCount = FNetBitArrayView::CalculateRequiredWordCount(BitCount);
+	const StorageWordType LastWordMask = ~StorageWordType(0) >> (uint32(-int32(BitCount)) & (WordBitCount - 1U));
+	const uint32 LastWordIt = WordCount - 1U;
+
+	// Read mask
+	uint32 DirtyWordMask = Reader->ReadBits(WordCount);
+
+	// Read and decode dirty words
+	uint32 CurrentMaskBit = 1U;
+	uint32 WordIt = 0U;
+	uint32 RemainingBits = BitCount;
+
+	while (RemainingBits >= 32U)
+	{
+		const StorageWordType ReadValue = DirtyWordMask & CurrentMaskBit ? ReadSparseUint32Function(Reader, 32U) : 0U;
+		OutData[WordIt] = OldData[WordIt] ^ ReadValue;
+
+		CurrentMaskBit += CurrentMaskBit;
+		RemainingBits -= 32U;
+		++WordIt;
+	}
+
+	// Last word, make sure we do not overwrite existing data
+	if (RemainingBits > 0U)
+	{
+		const StorageWordType ReadValue = DirtyWordMask & CurrentMaskBit ? ReadSparseUint32Function(Reader, RemainingBits) : 0U;
+		OutData[WordIt] = (OutData[WordIt] & ~LastWordMask) | ((OldData[WordIt] ^ ReadValue) & LastWordMask);
+	}
+}
+
 }
 
 void WriteSparseBitArray(FNetBitStreamWriter* Writer, const uint32* Data, uint32 BitCount, ESparseBitArraySerializationHint Hint)
@@ -664,6 +755,40 @@ void ReadSparseBitArray(FNetBitStreamReader* Reader, uint32* OutData, uint32 Bit
 	}
 }
 
+void WriteSparseBitArrayDelta(FNetBitStreamWriter* Writer, const uint32* Data, const uint32* OldData, uint32 BitCount)
+{
+	// Support large bit arrays
+	uint32 RemainingBits = BitCount;
+	while (RemainingBits >= Private::SerializeSparseArrayMaxBitCount)
+	{
+		Private::WriteSparseBitArrayDelta(Writer, Data, OldData, Private::SerializeSparseArrayMaxBitCount, Private::WriteSparseUint32UsingIndices);
+		RemainingBits -= Private::SerializeSparseArrayMaxBitCount;
+		Data += Private::SerializeSparseArrayMaxWordCount;
+		OldData += Private::SerializeSparseArrayMaxWordCount;
+	}
+	if (RemainingBits)
+	{
+		Private::WriteSparseBitArrayDelta(Writer, Data, OldData, RemainingBits, Private::WriteSparseUint32UsingIndices);
+	}
+}
+
+void ReadSparseBitArrayDelta(FNetBitStreamReader* Reader, uint32* OutData, const uint32* OldData, uint32 BitCount)
+{
+	// Support large bit arrays
+	uint32 RemainingBits = BitCount;
+	while (RemainingBits >= Private::SerializeSparseArrayMaxBitCount)
+	{
+		Private::ReadSparseBitArrayDelta(Reader, OutData, OldData, Private::SerializeSparseArrayMaxBitCount, Private::ReadSparseUint32UsingIndices);
+		RemainingBits -= Private::SerializeSparseArrayMaxBitCount;
+		OutData += Private::SerializeSparseArrayMaxWordCount;
+		OldData += Private::SerializeSparseArrayMaxWordCount;
+	}
+	if (RemainingBits)
+	{
+		Private::ReadSparseBitArrayDelta(Reader, OutData, OldData, RemainingBits, Private::ReadSparseUint32UsingIndices);
+	}
+}
+
 const uint32 NetBitStreamSentinelValue = 0xBAADDEADU;
 void WriteSentinelBits(FNetBitStreamWriter* Writer, uint32 BitCount)
 {
@@ -678,5 +803,4 @@ bool ReadAndVerifySentinelBits(FNetBitStreamReader* Reader, const TCHAR* ErrorSt
 	return ensureAlwaysMsgf(!Reader->IsOverflown() &&  ExpectedValue == ReadValue, TEXT("ReadAndVerifySentinelBits %s failed OverFlow %u Got 0x%u != 0x%u"), ErrorString, Reader->IsOverflown() ? 1U : 0U, ReadValue, ExpectedValue);
 }
 
-
 }
diff --git a/Engine/Source/Runtime/Experimental/Iris/Core/Public/Iris/Serialization/NetBitStreamUtil.h b/Engine/Source/Runtime/Experimental/Iris/Core/Public/Iris/Serialization/NetBitStreamUtil.h
--- a/Engine/Source/Runtime/Experimental/Iris/Core/Public/Iris/Serialization/NetBitStreamUtil.h
+++ b/Engine/Source/Runtime/Experimental/Iris/Core/Public/Iris/Serialization/NetBitStreamUtil.h
@@ -180,6 +180,17 @@ IRISCORE_API void WriteSparseBitArray(FNetBitStreamWriter* Writer, const uint32*
  */
 IRISCORE_API void ReadSparseBitArray(FNetBitStreamReader* Reader,  uint32* OutData, uint32 BitCount, ESparseBitArraySerializationHint Hint = ESparseBitArraySerializationHint::None);
 
+/** 
+ * Write sparse BitArray which is expected to contain mostly 0`s, if the array contains mostly 1`s passing the hint ESparseBitArraySerializationHint::ContainsMostlyOnes can be provided to flip the data before writing 
+ * Note: BitCount and hint is not Written so the write must be matched by a corresponding Read
+ */
+IRISCORE_API void WriteSparseBitArrayDelta(FNetBitStreamWriter* Writer, const uint32* Data, const uint32* OldData, uint32 BitCount);
+
+/** 
+ * Read sparse BitArray which is expected to contain mostly 0`s, if the array contains mostly 1`s passing the hint ESparseBitArraySerializationHint::ContainsMostlyOnes can be provided to flip the read data
+ */
+IRISCORE_API void ReadSparseBitArrayDelta(FNetBitStreamReader* Reader,  uint32* OutData, const uint32* OldData, uint32 BitCount);
+
 /**
  * Write sentinelBits from predefined sentinel value
  */
