Iris
- Fixed bug when Tearing off objects with no replicated state data.
- Fixed bug where subobject owner would not be marked dirty when we dropped creation info for SubObject which then would not replicate until the next time owner was marked dirty.
- Fixed bug where client could get a bitstream error if we got a replicated destroy for a not yet created SubObject
- Updated tests
- Fixed issue with NetTrace not working propertly with ReplicationSystemTest program

diff --git a/Engine/Plugins/Runtime/ReplicationSystemTestPlugin/Source/Private/Tests/ReplicationSystem/ReplicationSystemServerClientTestFixture.cpp b/Engine/Plugins/Runtime/ReplicationSystemTestPlugin/Source/Private/Tests/ReplicationSystem/ReplicationSystemServerClientTestFixture.cpp
--- a/Engine/Plugins/Runtime/ReplicationSystemTestPlugin/Source/Private/Tests/ReplicationSystem/ReplicationSystemServerClientTestFixture.cpp
+++ b/Engine/Plugins/Runtime/ReplicationSystemTestPlugin/Source/Private/Tests/ReplicationSystem/ReplicationSystemServerClientTestFixture.cpp
@@ -253,9 +253,13 @@ bool FReplicationSystemTestNode::SendUpdate(uint32 ConnectionId)
 
 	Connection.DataStreamManager->EndWrite();
 
-	UE_NET_TRACE_FLUSH_COLLECTOR(Context.GetTraceCollector(), GetNetTraceId(), Connection.ConnectionId, ENetTracePacketType::Outgoing);
+	if (bResult)
+	{
+		UE_NET_TRACE_FLUSH_COLLECTOR(Context.GetTraceCollector(), GetNetTraceId(), Connection.ConnectionId, ENetTracePacketType::Outgoing);
+		UE_NET_TRACE_PACKET_SEND(GetNetTraceId(), Connection.ConnectionId, Packet.PacketId, Packet.BitCount);
+	}
+
 	UE_NET_TRACE_DESTROY_COLLECTOR(Context.GetTraceCollector());
-	UE_NET_TRACE_PACKET_SEND(GetNetTraceId(), Connection.ConnectionId, Packet.PacketId, Packet.BitCount);
 
 	return bResult;
 }
diff --git a/Engine/Plugins/Runtime/ReplicationSystemTestPlugin/Source/Private/Tests/ReplicationSystem/TestReplicationSystem.cpp b/Engine/Plugins/Runtime/ReplicationSystemTestPlugin/Source/Private/Tests/ReplicationSystem/TestReplicationSystem.cpp
--- a/Engine/Plugins/Runtime/ReplicationSystemTestPlugin/Source/Private/Tests/ReplicationSystem/TestReplicationSystem.cpp
+++ b/Engine/Plugins/Runtime/ReplicationSystemTestPlugin/Source/Private/Tests/ReplicationSystem/TestReplicationSystem.cpp
@@ -1350,6 +1350,174 @@ UE_NET_TEST_FIXTURE(FReplicationSystemServerClientTestFixture, TestImmediateTear
 	UE_NET_ASSERT_TRUE(Cast<UTestReplicatedIrisObject>(Client->GetReplicationBridge()->GetReplicatedObject(ServerSubObject->NetRefHandle)) == nullptr);
 }
 
+// Test dropped creation of subobject dirties owner
+UE_NET_TEST_FIXTURE(FReplicationSystemServerClientTestFixture, TestDroppedCreationForSubobjectDirtiesOwner)
+{
+	UReplicationSystem* ReplicationSystem = Server->ReplicationSystem;
+
+	// Add a client
+	FReplicationSystemTestClient* Client = CreateClient();
+
+	// Spawn object on server
+	UTestReplicatedIrisObject* ServerObject = Server->CreateObject(0,0);
+
+	// Send and deliver packet
+	Server->PreSendUpdate();
+	Server->SendAndDeliverTo(Client, true);
+	Server->PostSendUpdate();
+
+	// Store Pointer to objects
+	UTestReplicatedIrisObject* ClientObject = Cast<UTestReplicatedIrisObject>(Client->GetReplicationBridge()->GetReplicatedObject(ServerObject->NetRefHandle));
+	UE_NET_ASSERT_TRUE(ClientObject != nullptr);
+
+	// Spawn second object on server as a subobject
+	UTestReplicatedIrisObject* ServerSubObject = Server->CreateSubObject(ServerObject->NetRefHandle, 0, 0);
+
+	// Send and do not deliver packet
+	Server->PreSendUpdate();
+	Server->SendAndDeliverTo(Client, false);
+	Server->PostSendUpdate();
+
+	// Send and deliver packet
+	Server->PreSendUpdate();
+	Server->SendAndDeliverTo(Client, true);
+	Server->PostSendUpdate();
+
+	// Verify that ClientObject now is created as expected
+	UTestReplicatedIrisObject* ClientSubObject = Cast<UTestReplicatedIrisObject>(Client->GetReplicationBridge()->GetReplicatedObject(ServerSubObject->NetRefHandle));
+	UE_NET_ASSERT_TRUE(ClientSubObject != nullptr);
+}
+
+// Test replicated destroy for not created object
+UE_NET_TEST_FIXTURE(FReplicationSystemServerClientTestFixture, TestReplicatedDestroyForNotCreatedObject)
+{
+	UReplicationSystem* ReplicationSystem = Server->ReplicationSystem;
+
+	// Add a client
+	FReplicationSystemTestClient* Client = CreateClient();
+
+	// Spawn object on server
+	UTestReplicatedIrisObject* ServerObject = Server->CreateObject(0,0);
+
+	// Update and delay delivery
+	Server->PreSendUpdate();
+	Server->SendTo(Client);
+	Server->PostSendUpdate();
+
+	// Destroy object
+	Server->ReplicationBridge->EndReplication(ServerObject);
+
+	// Update and delay delivery
+	Server->PreSendUpdate();
+	Server->SendTo(Client);
+	Server->PostSendUpdate();
+
+	// Drop first packet containing creation info for object
+	Server->DeliverTo(Client, false);
+
+	// Deliver second packet that should contain destroy
+	Server->DeliverTo(Client, true);
+
+	// Verify that the object does not exist on client
+	UE_NET_ASSERT_TRUE(Cast<UTestReplicatedIrisObject>(Client->GetReplicationBridge()->GetReplicatedObject(ServerObject->NetRefHandle)) == nullptr);
+}
+
+
+// Test replicated SubObjectDestroy for not created subobject
+UE_NET_TEST_FIXTURE(FReplicationSystemServerClientTestFixture, TestReplicatedSubObjectDestroyForNotCreatedObject)
+{
+	UReplicationSystem* ReplicationSystem = Server->ReplicationSystem;
+
+	// Add a client
+	FReplicationSystemTestClient* Client = CreateClient();
+
+	// Spawn object on server
+	UTestReplicatedIrisObject* ServerObject = Server->CreateObject(0,0);
+
+	// Replicate object
+	Server->PreSendUpdate();
+	Server->SendAndDeliverTo(Client, true);
+	Server->PostSendUpdate();
+
+	// Spawn second object on server as a subobject
+	UTestReplicatedIrisObjectWithNoReplicatedMembers* ServerSubObject = Server->CreateSubObject<UTestReplicatedIrisObjectWithNoReplicatedMembers>(ServerObject->NetRefHandle);
+
+	// Update and delay delivery
+	Server->PreSendUpdate();
+	Server->SendTo(Client);
+	Server->PostSendUpdate();
+
+	// Destroy subobject
+	Server->ReplicationBridge->EndReplication(ServerSubObject);
+
+	// Update and delay delivery
+	Server->PreSendUpdate();
+	Server->SendTo(Client);
+	Server->PostSendUpdate();
+
+	// Drop first packet containing creation info for subobject
+	Server->DeliverTo(Client, false);
+
+	// Deliver second packet that should contain replicated subobject destroy
+	Server->DeliverTo(Client, true);
+
+	// Verify that the object still exists on client
+	UE_NET_ASSERT_TRUE(Cast<UTestReplicatedIrisObject>(Client->GetReplicationBridge()->GetReplicatedObject(ServerObject->NetRefHandle)) != nullptr);
+
+	// Verify that the subobject does not exist on client
+	UE_NET_ASSERT_TRUE(Cast<UTestReplicatedIrisObjectWithNoReplicatedMembers>(Client->GetReplicationBridge()->GetReplicatedObject(ServerSubObject->NetRefHandle)) == nullptr);
+}
+
+// Test tear-off object in PendingCreate state to ensure that tear-off logic works as expected
+UE_NET_TEST_FIXTURE(FReplicationSystemServerClientTestFixture, TestTearOffObjectWithNoFragmentsDoesNotTriggerCheckIfPendingCreateWhenDestroyed)
+{
+	UReplicationSystem* ReplicationSystem = Server->ReplicationSystem;
+
+	// Add a client
+	FReplicationSystemTestClient* Client = CreateClient();
+
+	// Spawn object on server
+	UTestReplicatedIrisObjectWithNoReplicatedMembers* ServerObject = Server->CreateObject<UTestReplicatedIrisObjectWithNoReplicatedMembers>();
+
+	Server->PreSendUpdate();
+	Server->SendAndDeliverTo(Client, false);
+	Server->PostSendUpdate();
+
+	// Tear-off using immediate tear-off
+	Server->ReplicationBridge->EndReplication(ServerObject, EEndReplicationFlags::TearOff);
+
+	// Trigger the next update but avoid sending any data so that we keep the object in the PendingCreation state while we flush the Handles PendingTearOff Array which occurs in PostSendUpdate
+	Server->PreSendUpdate();
+	Server->PostSendUpdate();
+}
+
+// Test tear-off subobject in PendingCreate state to ensure that tear-off logic works as expected
+UE_NET_TEST_FIXTURE(FReplicationSystemServerClientTestFixture, TestTearOffSubObjectWithNoFragmentsDoesNotTriggerCheckIfPendingCreateWhenDestroyed)
+{
+	UReplicationSystem* ReplicationSystem = Server->ReplicationSystem;
+
+	// Add a client
+	FReplicationSystemTestClient* Client = CreateClient();
+
+	// Spawn object on server
+	UTestReplicatedIrisObject* ServerObject = Server->CreateObject(0,0);
+
+	// Spawn second object on server as a subobject
+	UTestReplicatedIrisObjectWithNoReplicatedMembers* ServerSubObject = Server->CreateSubObject<UTestReplicatedIrisObjectWithNoReplicatedMembers>(ServerObject->NetRefHandle);
+
+	// Update and drop
+	Server->PreSendUpdate();
+	Server->SendAndDeliverTo(Client, false);
+	Server->PostSendUpdate();
+
+	// Tear-off using immediate tear-off
+	Server->ReplicationBridge->EndReplication(ServerSubObject, EEndReplicationFlags::TearOff);
+
+	// Trigger the next update but avoid sending any data so that we keep the sub-object in the PendingCreation state while we flush the Handles PendingTearOff Array which occurs in PostSendUpdate
+	Server->PreSendUpdate();
+	Server->PostSendUpdate();
+}
+
 // Test TearOff and SubObjects, SubObjects must apply state?
 UE_NET_TEST_FIXTURE(FReplicationSystemServerClientTestFixture, TestTearOffNextUpdateExistingObjectWithSubObject)
 {
diff --git a/Engine/Source/Runtime/Experimental/Iris/Core/Private/Iris/ReplicationSystem/ChangeMaskCache.h b/Engine/Source/Runtime/Experimental/Iris/Core/Private/Iris/ReplicationSystem/ChangeMaskCache.h
--- a/Engine/Source/Runtime/Experimental/Iris/Core/Private/Iris/ReplicationSystem/ChangeMaskCache.h
+++ b/Engine/Source/Runtime/Experimental/Iris/Core/Private/Iris/ReplicationSystem/ChangeMaskCache.h
@@ -35,12 +35,15 @@ struct FChangeMaskCache
 	/** Empty cached data and free memory */
 	inline void EmptyCache();
 
-	/** The pointer is only valid until the next Add call. */
+	/** The ref is only valid until the next Add call. */
 	inline FCachedInfo& AddChangeMaskForObject(uint32 InternalIndex, uint32 BitCount);
 
-	/** The pointer is only valid until the next Add call. */
+	/** The ref is only valid until the next Add call. */
 	inline FCachedInfo& AddSubObjectOwnerDirty(uint32 InternalIndex);
 
+	/** The ref is only valid until the next Add call. */
+	inline FCachedInfo& AddEmptyChangeMaskForObject(uint32 InternalIndex);
+
 	/** The pointer is only valid until the next Add call. */
 	inline uint32* GetChangeMaskStorage(const FCachedInfo& Info);
 
@@ -82,6 +85,17 @@ FChangeMaskCache::FCachedInfo& FChangeMaskCache::AddChangeMaskForObject(uint32 I
 	return Indices.Add_GetRef(Info);
 }
 
+FChangeMaskCache::FCachedInfo& FChangeMaskCache::AddEmptyChangeMaskForObject(uint32 InternalIndex)
+{
+	FCachedInfo Info;
+	Info.InternalIndex = InternalIndex;
+	Info.StorageOffset = 0U;
+	Info.bMarkSubObjectOwnerDirty = 0U;
+	Info.bHasDirtyChangeMask = 0U;
+
+	return Indices.Add_GetRef(Info);
+}
+
 uint32* FChangeMaskCache::GetChangeMaskStorage(const FCachedInfo& Info)
 {
 	return &Storage.GetData()[Info.StorageOffset];
diff --git a/Engine/Source/Runtime/Experimental/Iris/Core/Private/Iris/ReplicationSystem/ReplicationBridge.cpp b/Engine/Source/Runtime/Experimental/Iris/Core/Private/Iris/ReplicationSystem/ReplicationBridge.cpp
--- a/Engine/Source/Runtime/Experimental/Iris/Core/Private/Iris/ReplicationSystem/ReplicationBridge.cpp
+++ b/Engine/Source/Runtime/Experimental/Iris/Core/Private/Iris/ReplicationSystem/ReplicationBridge.cpp
@@ -758,22 +758,35 @@ void UReplicationBridge::InternalTearOff(FNetRefHandle Handle)
 			CallPreSendUpdateSingleHandle(Handle);
 		} 
 
-		if (FReplicationInstanceOperationsInternal::CopyObjectStateData(ChangeMaskWriter, ChangeMaskCache, *NetRefHandleManager, SerializationContext, InternalObjectIndex))
+		if (ObjectData.InstanceProtocol && ObjectData.Protocol->InternalTotalSize > 0U)
 		{
-			// Propagate changes to all connections that we currently have in scope
-			FReplicationConnections& Connections = ReplicationSystem->GetReplicationSystemInternal()->GetConnections();
-
-			// Iterate over connections and propagate dirty changemasks to all connections already scoping this object
-			auto UpdateDirtyChangeMasks = [&Connections, &ChangeMaskCache](uint32 ConnectionId)
+			FReplicationInstanceOperationsInternal::CopyObjectStateData(ChangeMaskWriter, ChangeMaskCache, *NetRefHandleManager, SerializationContext, InternalObjectIndex);
+		}
+		else
+		{
+			// Nothing to copy, but we must still propagate the tear-off state.
+			FChangeMaskCache::FCachedInfo& Info = ChangeMaskCache.AddEmptyChangeMaskForObject(InternalObjectIndex);
+			// If we are a subobject we must also mark owner as dirty.
+			const uint32 SubObjectOwnerIndex = ObjectData.SubObjectRootIndex;
+			if (SubObjectOwnerIndex != FNetRefHandleManager::InvalidInternalIndex) 
 			{
-				FReplicationConnection* Conn = Connections.GetConnection(ConnectionId);
-				const bool bMarkForTearOff = true;
-				Conn->ReplicationWriter->ForceUpdateDirtyChangeMasks(ChangeMaskCache, FReplicationWriter::FlushFlags_None, bMarkForTearOff);
-			};
-			const FNetBitArray& ValidConnections = Connections.GetValidConnections();
-			ValidConnections.ForAllSetBits(UpdateDirtyChangeMasks);		
+				ChangeMaskCache.AddSubObjectOwnerDirty(SubObjectOwnerIndex);
+			}			
 		}
 
+		// Propagate changes to all connections that we currently have in scope
+		FReplicationConnections& Connections = ReplicationSystem->GetReplicationSystemInternal()->GetConnections();
+
+		// Iterate over connections and propagate dirty changemasks to all connections already scoping this object
+		auto UpdateDirtyChangeMasks = [&Connections, &ChangeMaskCache](uint32 ConnectionId)
+		{
+			FReplicationConnection* Conn = Connections.GetConnection(ConnectionId);
+			const bool bMarkForTearOff = true;
+			Conn->ReplicationWriter->ForceUpdateDirtyChangeMasks(ChangeMaskCache, FReplicationWriter::FlushFlags_None, bMarkForTearOff);
+		};
+		const FNetBitArray& ValidConnections = Connections.GetValidConnections();
+		ValidConnections.ForAllSetBits(UpdateDirtyChangeMasks);		
+
 		// Mark object as being torn-off and that we should no longer propagate state changes
 		ObjectData.bTearOff = 1U;
 		ObjectData.bShouldPropagateChangedStates = 0U;
diff --git a/Engine/Source/Runtime/Experimental/Iris/Core/Private/Iris/ReplicationSystem/ReplicationReader.cpp b/Engine/Source/Runtime/Experimental/Iris/Core/Private/Iris/ReplicationSystem/ReplicationReader.cpp
--- a/Engine/Source/Runtime/Experimental/Iris/Core/Private/Iris/ReplicationSystem/ReplicationReader.cpp
+++ b/Engine/Source/Runtime/Experimental/Iris/Core/Private/Iris/ReplicationSystem/ReplicationReader.cpp
@@ -485,7 +485,7 @@ void FReplicationReader::ReadObject(FNetSerializationContext& Context)
 			InternalIndex = NetRefHandleManager->GetInternalIndex(IncompleteHandle);
 
 			// If this is a subobject that is being destroyed this was no error as we send destroy info for unconfirmed objects
-			if (!!(ReplicatedDestroyHeaderFlags & ReplicatedDestroyHeaderFlags_EndReplication))
+			if ((ReplicatedDestroyHeaderFlags & ReplicatedDestroyHeaderFlags_EndReplication) == 0U)
 			{
 				bHasErrors = InternalIndex == FNetRefHandleManager::InvalidInternalIndex;
 			}
diff --git a/Engine/Source/Runtime/Experimental/Iris/Core/Private/Iris/ReplicationSystem/ReplicationWriter.cpp b/Engine/Source/Runtime/Experimental/Iris/Core/Private/Iris/ReplicationSystem/ReplicationWriter.cpp
--- a/Engine/Source/Runtime/Experimental/Iris/Core/Private/Iris/ReplicationSystem/ReplicationWriter.cpp
+++ b/Engine/Source/Runtime/Experimental/Iris/Core/Private/Iris/ReplicationSystem/ReplicationWriter.cpp
@@ -744,7 +744,7 @@ void FReplicationWriter::InternalUpdateDirtyChangeMasks(const FChangeMaskCache&
 		if (Entry.bMarkSubObjectOwnerDirty == 0U)
 		{		
 			// Mark object for TearOff, that is that we will stop replication as soon as the tear-off is acknowledged
-			Info.TearOff = MarkForTearOff;
+			Info.TearOff |= MarkForTearOff;
 
 			// Update flush flags
 			Info.FlushFlags = Info.FlushFlags | ExtraFlushFlags;
@@ -1140,8 +1140,25 @@ void FReplicationWriter::HandleDroppedRecord<FReplicationWriter::EReplicatedObje
 		// Mark attachments as dirty
 		Info.HasAttachments |= RecordInfo.HasAttachments;
 
-		// Bump prio
-		SchedulingPriorities[InternalIndex] += FReplicationWriter::CreatePriority;
+		if (Info.IsSubObject)
+		{
+			// Mark owner dirty as well as subobjects only are scheduled together with owner
+			const FNetRefHandleManager::FReplicatedObjectData& ObjectData = NetRefHandleManager->GetReplicatedObjectData(InternalIndex);
+			uint32 SubObjectOwnerInternalIndex = ObjectData.SubObjectRootIndex;
+
+			FReplicationInfo& SubObjectOwnerReplicationInfo = GetReplicationInfo(SubObjectOwnerInternalIndex);
+			if (ensure(SubObjectOwnerReplicationInfo.GetState() < EReplicatedObjectState::PendingDestroy))
+			{
+				// Mark owner as dirty
+				ObjectsWithDirtyChanges.SetBit(SubObjectOwnerInternalIndex);
+
+				// Indicate that we have dirty subobjects
+				SubObjectOwnerReplicationInfo.HasDirtySubObjects = 1U;
+
+				// Give slight priority bump to owner
+				SchedulingPriorities[SubObjectOwnerInternalIndex] += FReplicationWriter::LostStatePriorityBump;
+			}
+		}
 	}
 	else if (CurrentState == EReplicatedObjectState::SubObjectPendingDestroy)
 	{
