Make the tests in ReplicationSystemTestPlugin compatible with the LowLevelTests framework.
Original ReplicationSystemTest program using the automation framework still works too.

diff --git a/Engine/Plugins/Runtime/ReplicationSystemTestPlugin/Source/Private/NetworkAutomationTest.cpp b/Engine/Plugins/Runtime/ReplicationSystemTestPlugin/Source/Private/NetworkAutomationTest.cpp
--- a/Engine/Plugins/Runtime/ReplicationSystemTestPlugin/Source/Private/NetworkAutomationTest.cpp
+++ b/Engine/Plugins/Runtime/ReplicationSystemTestPlugin/Source/Private/NetworkAutomationTest.cpp
@@ -10,12 +10,6 @@ namespace UE::Net
 //
 bool FNetworkAutomationTestConfig::bVerboseLogging = false;
 
-struct FNetworkAutomationTestStats
-{
-	SIZE_T FailureCount;
-	SIZE_T WarningCount;
-};
-
 class FNetworkAutomationTestStatsSummary
 {
 public:
@@ -36,8 +30,7 @@ static FNetworkAutomationTestStatsSummary NetworkAutomationTestStatsSummary;
 
 // 
 FNetworkAutomationTestSuiteFixture::FNetworkAutomationTestSuiteFixture()
-: Stats(nullptr)
-, bSuppressWarningsFromSummary(false)
+: bSuppressWarningsFromSummary(false)
 {
 }
 
@@ -53,22 +46,29 @@ void FNetworkAutomationTestSuiteFixture::TearDown()
 {
 }
 
-void FNetworkAutomationTestSuiteFixture::RunTest()
+void FNetworkAutomationTestSuiteFixture::PreSetUp()
 {
 	if (FNetworkAutomationTestConfig::GetVerboseLogging())
 	{
 		UE_LOG(LogNetworkAutomationTest, Display, TEXT("Running TestCase %ls"), GetName());
 	}
 
-	FNetworkAutomationTestStats TempStats = {};
-	Stats = &TempStats;
+	Stats = FNetworkAutomationTestStats();
+}
 
+void FNetworkAutomationTestSuiteFixture::RunTest()
+{
+	PreSetUp();
 	SetUp();
 	RunTestImpl();
 	TearDown();
+	PostTearDown();
+}
 
-	NetworkAutomationTestStatsSummary.AddTestResult(GetName(), *Stats);
-	Stats = nullptr;
+void FNetworkAutomationTestSuiteFixture::PostTearDown()
+{
+	NetworkAutomationTestStatsSummary.AddTestResult(GetName(), Stats);
+	Stats = FNetworkAutomationTestStats();
 }
 
 void FNetworkAutomationTestSuiteFixture::SetSuppressWarningsFromSummary(bool bSuppress)
@@ -78,12 +78,12 @@ void FNetworkAutomationTestSuiteFixture::SetSuppressWarningsFromSummary(bool bSu
 
 void FNetworkAutomationTestSuiteFixture::AddTestFailure()
 {
-	++Stats->FailureCount;
+	++Stats.FailureCount;
 }
 
 void FNetworkAutomationTestSuiteFixture::AddTestWarning()
 {
-	Stats->WarningCount += !bSuppressWarningsFromSummary;
+	Stats.WarningCount += !bSuppressWarningsFromSummary;
 }
 
 #if WITH_AUTOMATION_WORKER
diff --git a/Engine/Plugins/Runtime/ReplicationSystemTestPlugin/Source/Private/NetworkAutomationTest.h b/Engine/Plugins/Runtime/ReplicationSystemTestPlugin/Source/Private/NetworkAutomationTest.h
--- a/Engine/Plugins/Runtime/ReplicationSystemTestPlugin/Source/Private/NetworkAutomationTest.h
+++ b/Engine/Plugins/Runtime/ReplicationSystemTestPlugin/Source/Private/NetworkAutomationTest.h
@@ -5,9 +5,19 @@
 #include "CoreTypes.h"
 #include "Misc/AutomationTest.h"
 
+#if EXPLICIT_TESTS_TARGET
+#include "TestHarness.h"
+#endif
+
 namespace UE::Net
 {
 
+struct FNetworkAutomationTestStats
+{
+	SIZE_T FailureCount;
+	SIZE_T WarningCount;
+};
+
 struct FNetworkAutomationTestConfig
 {
 	static void SetVerboseLogging(bool bVerbose) { bVerboseLogging = bVerbose; }
@@ -44,7 +54,13 @@ protected:
 	void AddTestWarning();
 
 private:
-	struct FNetworkAutomationTestStats* Stats;
+	void PreSetUp();
+	void PostTearDown();
+
+	template<class Fixture, const TCHAR* Name>
+	friend class TCatchFixtureWrapper;
+
+	FNetworkAutomationTestStats Stats = {};
 	bool bSuppressWarningsFromSummary;
 };
 
@@ -72,7 +88,44 @@ private:
 
 }
 
-#if WITH_AUTOMATION_WORKER
+#if EXPLICIT_TESTS_TARGET
+
+namespace UE::Net
+{
+
+template<class Fixture, const TCHAR* Name>
+class TCatchFixtureWrapper : public Fixture
+{
+public:
+	TCatchFixtureWrapper() : Fixture()
+	{
+		this->PreSetUp();
+		Fixture::SetUp();
+	}
+
+	~TCatchFixtureWrapper()
+	{
+		Fixture::TearDown();
+		this->PostTearDown();
+	}
+
+	virtual const TCHAR* GetName() const override { return Name; }
+	virtual void RunTestImpl() override {} // Unused in Catch-based tests
+};
+
+}
+
+#define UE_NET_TEST_CATCH_INTERNAL(TestSuite, TestCase, BaseClass) \
+	constexpr TCHAR TestSuite ## _ ## TestCase ## Name[] = TEXT("Net." #TestSuite "." #TestCase); \
+	TEST_CASE_METHOD((TCatchFixtureWrapper<BaseClass, TestSuite ## _ ## TestCase ## Name>), #TestSuite ## "_" ## #TestCase)
+
+#define UE_NET_TEST_FIXTURE(TestFixture, TestCase) UE_NET_TEST_CATCH_INTERNAL(TestFixture, TestCase, TestFixture)
+#define UE_NET_TEST(TestSuite, TestCase) UE_NET_TEST_CATCH_INTERNAL(TestSuite, TestCase, UE::Net::FNetworkAutomationTestSuiteFixture)
+
+
+
+
+#elif WITH_AUTOMATION_WORKER
 
 #define UE_NET_TEST_INTERNAL(TestSuite, TestCase, BaseClass) \
 class TestSuite ## _ ## TestCase : public BaseClass \
diff --git a/Engine/Plugins/Runtime/ReplicationSystemTestPlugin/Source/Private/NetworkAutomationTestMacros.cpp b/Engine/Plugins/Runtime/ReplicationSystemTestPlugin/Source/Private/NetworkAutomationTestMacros.cpp
--- a/Engine/Plugins/Runtime/ReplicationSystemTestPlugin/Source/Private/NetworkAutomationTestMacros.cpp
+++ b/Engine/Plugins/Runtime/ReplicationSystemTestPlugin/Source/Private/NetworkAutomationTestMacros.cpp
@@ -1,6 +1,7 @@
 // Copyright Epic Games, Inc. All Rights Reserved.
 
 #include "NetworkAutomationTestMacros.h"
+#include "ReplicationSystemTestPlugin/NetworkAutomationTest.h"
 #include "NetworkAutomationTest.h"
 
 DEFINE_LOG_CATEGORY(LogNetworkAutomationTest);
diff --git a/Engine/Plugins/Runtime/ReplicationSystemTestPlugin/Source/Private/NetworkAutomationTestMacros.h b/Engine/Plugins/Runtime/ReplicationSystemTestPlugin/Source/Private/NetworkAutomationTestMacros.h
--- a/Engine/Plugins/Runtime/ReplicationSystemTestPlugin/Source/Private/NetworkAutomationTestMacros.h
+++ b/Engine/Plugins/Runtime/ReplicationSystemTestPlugin/Source/Private/NetworkAutomationTestMacros.h
@@ -3,9 +3,14 @@
 #pragma once
 
 #include "CoreTypes.h"
+#include "HAL/PreprocessorHelpers.h"
 #include "TestMessage.h"
 #include "Logging/LogVerbosity.h"
 
+#if EXPLICIT_TESTS_TARGET
+#include <catch2/catch_test_macros.hpp>
+#endif
+
 namespace UE::Net
 {
 	class FNetworkAutomationTestSuiteFixture;
@@ -42,6 +47,15 @@ class FTestMessageLog
 {
 public:
 	FTestMessageLog(FNetworkAutomationTestSuiteFixture& Test, ELogVerbosity::Type LogVerbosity);
+
+	template<typename T>
+	FTestMessageLog(FNetworkAutomationTestSuiteFixture& InTest, ELogVerbosity::Type InLogVerbosity, const T& InMessage)
+		: Test(InTest)
+		, LogVerbosity(InLogVerbosity)
+	{
+		Message << InMessage;
+	}
+
 	~FTestMessageLog();
 
 	void operator=(const FTestMessage&);
@@ -98,6 +112,51 @@ Name(const T& Value, U ExpectedValue, const char* ValueText, const char* Expecte
 	} \
 }
 
+#if EXPLICIT_TESTS_TARGET
+
+#define UE_NET_TEST_MSG_INTERNAL(Message) \
+	UE::Net::FTestMessage PREPROCESSOR_JOIN(TestMessage, __LINE__); \
+	PREPROCESSOR_JOIN(TestMessage, __LINE__) << Message; \
+	UNSCOPED_INFO(PREPROCESSOR_JOIN(TestMessage, __LINE__).C_Str())
+
+#define UE_NET_ASSERT_EQ(V1, V2) REQUIRE(V1 == V2)
+#define UE_NET_ASSERT_NE(V1, V2) REQUIRE(V1 != V2)
+#define UE_NET_ASSERT_LT(V1, V2) REQUIRE(V1 < V2)
+#define UE_NET_ASSERT_LE(V1, V2) REQUIRE(V1 <= V2)
+#define UE_NET_ASSERT_GT(V1, V2) REQUIRE(V1 > V2)
+#define UE_NET_ASSERT_GE(V1, V2) REQUIRE(V1 >= V2)
+#define UE_NET_ASSERT_TRUE(V) REQUIRE(V)
+#define UE_NET_ASSERT_FALSE(V) REQUIRE_FALSE(V)
+
+#define UE_NET_ASSERT_EQ_MSG(V1, V2, Message) UE_NET_TEST_MSG_INTERNAL(Message); REQUIRE(V1 == V2)
+#define UE_NET_ASSERT_NE_MSG(V1, V2, Message) UE_NET_TEST_MSG_INTERNAL(Message); REQUIRE(V1 != V2)
+#define UE_NET_ASSERT_LT_MSG(V1, V2, Message) UE_NET_TEST_MSG_INTERNAL(Message); REQUIRE(V1 < V2)
+#define UE_NET_ASSERT_LE_MSG(V1, V2, Message) UE_NET_TEST_MSG_INTERNAL(Message); REQUIRE(V1 <= V2)
+#define UE_NET_ASSERT_GT_MSG(V1, V2, Message) UE_NET_TEST_MSG_INTERNAL(Message); REQUIRE(V1 > V2)
+#define UE_NET_ASSERT_GE_MSG(V1, V2, Message) UE_NET_TEST_MSG_INTERNAL(Message); REQUIRE(V1 >= V2)
+#define UE_NET_ASSERT_TRUE_MSG(V, Message) UE_NET_TEST_MSG_INTERNAL(Message); REQUIRE(V)
+#define UE_NET_ASSERT_FALSE_MSG(V, Message) UE_NET_TEST_MSG_INTERNAL(Message); REQUIRE_FALSE(V)
+
+#define UE_NET_EXPECT_EQ(V1, V2) CHECK(V1 == V2)
+#define UE_NET_EXPECT_NE(V1, V2) CHECK(V1 != V2)
+#define UE_NET_EXPECT_LT(V1, V2) CHECK(V1 < V2)
+#define UE_NET_EXPECT_LE(V1, V2) CHECK(V1 <= V2)
+#define UE_NET_EXPECT_GT(V1, V2) CHECK(V1 > V2)
+#define UE_NET_EXPECT_GE(V1, V2) CHECK(V1 >= V2)
+#define UE_NET_EXPECT_TRUE(V) CHECK(V)
+#define UE_NET_EXPECT_FALSE(V) CHECK_FALSE(V)
+
+#define UE_NET_EXPECT_EQ_MSG(V1, V2, Message) UE_NET_TEST_MSG_INTERNAL(Message); CHECK(V1 == V2)
+#define UE_NET_EXPECT_NE_MSG(V1, V2, Message) UE_NET_TEST_MSG_INTERNAL(Message); CHECK(V1 != V2)
+#define UE_NET_EXPECT_LT_MSG(V1, V2, Message) UE_NET_TEST_MSG_INTERNAL(Message); CHECK(V1 < V2)
+#define UE_NET_EXPECT_LE_MSG(V1, V2, Message) UE_NET_TEST_MSG_INTERNAL(Message); CHECK(V1 <= V2)
+#define UE_NET_EXPECT_GT_MSG(V1, V2, Message) UE_NET_TEST_MSG_INTERNAL(Message); CHECK(V1 > V2)
+#define UE_NET_EXPECT_GE_MSG(V1, V2, Message) UE_NET_TEST_MSG_INTERNAL(Message); CHECK(V1 >= V2)
+#define UE_NET_EXPECT_TRUE_MSG(V, Message) UE_NET_TEST_MSG_INTERNAL(Message); CHECK(V)
+#define UE_NET_EXPECT_FALSE_MSG(V, Message) UE_NET_TEST_MSG_INTERNAL(Message); CHECK_FALSE(V)
+
+#else
+
 UE_NET_COMPARE_RETURN_TESTRESULT(TCmpEqual, ==)
 UE_NET_COMPARE_RETURN_TESTRESULT(TCmpNotEqual, !=)
 UE_NET_COMPARE_RETURN_TESTRESULT(TCmpLess, <)
@@ -105,14 +164,23 @@ UE_NET_COMPARE_RETURN_TESTRESULT(TCmpLessOrEqual, <=)
 UE_NET_COMPARE_RETURN_TESTRESULT(TCmpGreater, >)
 UE_NET_COMPARE_RETURN_TESTRESULT(TCmpGreaterOrEqual, >=)
 
-#define UE_NET_FAIL(Message) return this->AddTestFailure(), UE::Net::FTestMessageLog(*this, ELogVerbosity::Error) = UE::Net::FTestMessage() << Message
-#define UE_NET_WARN(Message) this->AddTestWarning(), UE::Net::FTestMessageLog(*this, ELogVerbosity::Warning) = UE::Net::FTestMessage() << Message
+#define UE_NET_FAIL_MSG(Message) return this->AddTestFailure(), UE::Net::FTestMessageLog(*this, ELogVerbosity::Error) = UE::Net::FTestMessage() << Message
+#define UE_NET_WARN_MSG(Message) this->AddTestWarning(), UE::Net::FTestMessageLog(*this, ELogVerbosity::Warning) = UE::Net::FTestMessage() << Message
+
+#define UE_NET_FAIL(Message) return this->AddTestFailure(), void(UE::Net::FTestMessageLog(*this, ELogVerbosity::Error, Message))
+#define UE_NET_WARN(Message) this->AddTestWarning(), void(UE::Net::FTestMessageLog(*this, ELogVerbosity::Warning, Message))
 
 #define UE_NET_TEST_(CompareWithTestResult, V1, V2, FailFunction) if (UE::Net::FTestResult Result_ = CompareWithTestResult(V1, V2, #V1, #V2)) {} else FailFunction(Result_.GetMessage())
+#define UE_NET_ASSERT_INTERNAL_MSG(CompareWithTestResult, V1, V2) UE_NET_TEST_(CompareWithTestResult, V1, V2, UE_NET_FAIL_MSG)
+#define UE_NET_EXPECT_INTERNAL_MSG(CompareWithTestResult, V1, V2) UE_NET_TEST_(CompareWithTestResult, V1, V2, UE_NET_WARN_MSG)
+
 #define UE_NET_ASSERT_INTERNAL(CompareWithTestResult, V1, V2) UE_NET_TEST_(CompareWithTestResult, V1, V2, UE_NET_FAIL)
 #define UE_NET_EXPECT_INTERNAL(CompareWithTestResult, V1, V2) UE_NET_TEST_(CompareWithTestResult, V1, V2, UE_NET_WARN)
 
 #define UE_NET_TEST_BOOL_(V1, V2, FailFunction) if (UE::Net::FTestResult Result_ = UE::Net::BooleanTestResult(V1, V2, #V1)) {} else FailFunction(Result_.GetMessage())
+#define UE_NET_ASSERT_BOOL_INTERNAL_MSG(V1, V2) UE_NET_TEST_BOOL_(V1, V2, UE_NET_FAIL_MSG)
+#define UE_NET_EXPECT_BOOL_INTERNAL_MSG(V1, V2) UE_NET_TEST_BOOL_(V1, V2, UE_NET_WARN_MSG)
+
 #define UE_NET_ASSERT_BOOL_INTERNAL(V1, V2) UE_NET_TEST_BOOL_(V1, V2, UE_NET_FAIL)
 #define UE_NET_EXPECT_BOOL_INTERNAL(V1, V2) UE_NET_TEST_BOOL_(V1, V2, UE_NET_WARN)
 
@@ -125,6 +193,15 @@ UE_NET_COMPARE_RETURN_TESTRESULT(TCmpGreaterOrEqual, >=)
 #define UE_NET_ASSERT_TRUE(V) UE_NET_ASSERT_BOOL_INTERNAL(V, true)
 #define UE_NET_ASSERT_FALSE(V) UE_NET_ASSERT_BOOL_INTERNAL(V, false)
 
+#define UE_NET_ASSERT_EQ_MSG(V1, V2, Message) UE_NET_ASSERT_INTERNAL_MSG(TCmpEqual, V1, V2) << Message
+#define UE_NET_ASSERT_NE_MSG(V1, V2, Message) UE_NET_ASSERT_INTERNAL_MSG(TCmpNotEqual, V1, V2) << Message
+#define UE_NET_ASSERT_LT_MSG(V1, V2, Message) UE_NET_ASSERT_INTERNAL_MSG(TCmpLess, V1, V2) << Message
+#define UE_NET_ASSERT_LE_MSG(V1, V2, Message) UE_NET_ASSERT_INTERNAL_MSG(TCmpLessOrEqual, V1, V2) << Message
+#define UE_NET_ASSERT_GT_MSG(V1, V2, Message) UE_NET_ASSERT_INTERNAL_MSG(TCmpGreater, V1, V2) << Message
+#define UE_NET_ASSERT_GE_MSG(V1, V2, Message) UE_NET_ASSERT_INTERNAL_MSG(TCmpGreaterOrEqual, V1, V2) << Message
+#define UE_NET_ASSERT_TRUE_MSG(V, Message) UE_NET_ASSERT_BOOL_INTERNAL_MSG(V, true) << Message
+#define UE_NET_ASSERT_FALSE_MSG(V, Message) UE_NET_ASSERT_BOOL_INTERNAL_MSG(V, false) << Message
+
 #define UE_NET_EXPECT_EQ(V1, V2) UE_NET_EXPECT_INTERNAL(TCmpEqual, V1, V2)
 #define UE_NET_EXPECT_NE(V1, V2) UE_NET_EXPECT_INTERNAL(TCmpNotEqual, V1, V2)
 #define UE_NET_EXPECT_LT(V1, V2) UE_NET_EXPECT_INTERNAL(TCmpLess, V1, V2)
@@ -134,6 +211,17 @@ UE_NET_COMPARE_RETURN_TESTRESULT(TCmpGreaterOrEqual, >=)
 #define UE_NET_EXPECT_TRUE(V) UE_NET_EXPECT_BOOL_INTERNAL(V, true)
 #define UE_NET_EXPECT_FALSE(V) UE_NET_EXPECT_BOOL_INTERNAL(V, false)
 
+#define UE_NET_EXPECT_EQ_MSG(V1, V2, Message) UE_NET_EXPECT_INTERNAL_MSG(TCmpEqual, V1, V2) << Message
+#define UE_NET_EXPECT_NE_MSG(V1, V2, Message) UE_NET_EXPECT_INTERNAL_MSG(TCmpNotEqual, V1, V2) << Message
+#define UE_NET_EXPECT_LT_MSG(V1, V2, Message) UE_NET_EXPECT_INTERNAL_MSG(TCmpLess, V1, V2) << Message
+#define UE_NET_EXPECT_LE_MSG(V1, V2, Message) UE_NET_EXPECT_INTERNAL_MSG(TCmpLessOrEqual, V1, V2) << Message
+#define UE_NET_EXPECT_GT_MSG(V1, V2, Message) UE_NET_EXPECT_INTERNAL_MSG(TCmpGreater, V1, V2) << Message
+#define UE_NET_EXPECT_GE_MSG(V1, V2, Message) UE_NET_EXPECT_INTERNAL_MSG(TCmpGreaterOrEqual, V1, V2) << Message
+#define UE_NET_EXPECT_TRUE_MSG(V, Message) UE_NET_EXPECT_BOOL_INTERNAL_MSG(V, true) << Message
+#define UE_NET_EXPECT_FALSE_MSG(V, Message) UE_NET_EXPECT_BOOL_INTERNAL_MSG(V, false) << Message
+
+#endif
+
 #define UE_NET_LOG(Message) FTestMessageLog(*this, ELogVerbosity::Display) = FTestMessage() << Message
 
 #define UE_NET_STRASSERT_EQ
diff --git a/Engine/Plugins/Runtime/ReplicationSystemTestPlugin/Source/Private/Tests/DataStream/TestDataStream.cpp b/Engine/Plugins/Runtime/ReplicationSystemTestPlugin/Source/Private/Tests/DataStream/TestDataStream.cpp
--- a/Engine/Plugins/Runtime/ReplicationSystemTestPlugin/Source/Private/Tests/DataStream/TestDataStream.cpp
+++ b/Engine/Plugins/Runtime/ReplicationSystemTestPlugin/Source/Private/Tests/DataStream/TestDataStream.cpp
@@ -184,7 +184,8 @@ UE_NET_TEST_FIXTURE(FTestDataStream, DataStreamGetsReadDataCall)
 
 // FTestDataStream implementation
 FTestDataStream::FTestDataStream()
-: DataStreamManager(nullptr)
+: FNetworkAutomationTestSuiteFixture()
+, DataStreamManager(nullptr)
 , DataStreamDefinitions(nullptr)
 , CurrentDataStreamDefinitions(nullptr)
 , DataStreamContext(&BitStreamReader, &BitStreamWriter)
diff --git a/Engine/Plugins/Runtime/ReplicationSystemTestPlugin/Source/Private/Tests/ReplicationState/TestPropertyReplicationState.cpp b/Engine/Plugins/Runtime/ReplicationSystemTestPlugin/Source/Private/Tests/ReplicationState/TestPropertyReplicationState.cpp
--- a/Engine/Plugins/Runtime/ReplicationSystemTestPlugin/Source/Private/Tests/ReplicationState/TestPropertyReplicationState.cpp
+++ b/Engine/Plugins/Runtime/ReplicationSystemTestPlugin/Source/Private/Tests/ReplicationState/TestPropertyReplicationState.cpp
@@ -363,7 +363,7 @@ UE_NET_TEST_FIXTURE(FTestPropertyReplicationStateContext, PropertyReplicationSta
 	for (SIZE_T MemberIt = 0, MemberEndIt = Descriptor->MemberCount; MemberIt != MemberEndIt; ++MemberIt)
 	{
 		// Init state doesn't have changemasks
-		UE_NET_ASSERT_TRUE(ReplicationStateB.IsDirty(MemberIt)) << "Member '" << Descriptor->MemberProperties[MemberIt]->GetName() << "' index " << MemberIt;
+		UE_NET_ASSERT_TRUE_MSG(ReplicationStateB.IsDirty(MemberIt), "Member '" << Descriptor->MemberProperties[MemberIt]->GetName() << "' index " << MemberIt);
 	}
 }
 
@@ -395,11 +395,11 @@ UE_NET_TEST_FIXTURE(FTestPropertyReplicationStateContext, PropertyReplicationSta
 	{
 		if (MemberIt == FullyReplicatedArrayElementIndex || (MemberIt == NotFullyReplicatedArrayElementIndex))
 		{
-			UE_NET_ASSERT_TRUE(ReplicationStateB.IsDirty(MemberIt)) << "Member '" << Descriptor->MemberProperties[MemberIt]->GetName() << "' index " << MemberIt;
+			UE_NET_ASSERT_TRUE_MSG(ReplicationStateB.IsDirty(MemberIt), "Member '" << Descriptor->MemberProperties[MemberIt]->GetName() << "' index " << MemberIt);
 		}
 		else
 		{
-			UE_NET_ASSERT_FALSE(ReplicationStateB.IsDirty(MemberIt)) << "Member '" << Descriptor->MemberProperties[MemberIt]->GetName() << "' index " << MemberIt;
+			UE_NET_ASSERT_FALSE_MSG(ReplicationStateB.IsDirty(MemberIt), "Member '" << Descriptor->MemberProperties[MemberIt]->GetName() << "' index " << MemberIt);
 		}
 	}
 }
@@ -430,7 +430,7 @@ UE_NET_TEST_FIXTURE(FTestPropertyReplicationStateContext, PropertyReplicationSta
 	ReplicationStateB.Set(ReplicationState);
 	for (SIZE_T MemberIt = 0, MemberEndIt = Descriptor->MemberCount; MemberIt != MemberEndIt; ++MemberIt)
 	{
-		UE_NET_ASSERT_FALSE(ReplicationStateB.IsDirty(MemberIt)) << "Member '" << Descriptor->MemberProperties[MemberIt]->GetName() << "' index " << MemberIt;
+		UE_NET_ASSERT_FALSE_MSG(ReplicationStateB.IsDirty(MemberIt), "Member '" << Descriptor->MemberProperties[MemberIt]->GetName() << "' index " << MemberIt);
 	}
 }
 
@@ -468,11 +468,11 @@ UE_NET_TEST_FIXTURE(FTestPropertyReplicationStateContext, PropertyReplicationSta
 	{
 		if (MemberIt == FullyReplicatedArrayElementIndex || (MemberIt == NotFullyReplicatedArrayElementIndex))
 		{
-			UE_NET_ASSERT_TRUE(ReplicationStateB.IsDirty(MemberIt)) << "Member '" << Descriptor->MemberProperties[MemberIt]->GetName() << "' index " << MemberIt;
+			UE_NET_ASSERT_TRUE_MSG(ReplicationStateB.IsDirty(MemberIt), "Member '" << Descriptor->MemberProperties[MemberIt]->GetName() << "' index " << MemberIt);
 		}
 		else
 		{
-			UE_NET_ASSERT_FALSE(ReplicationStateB.IsDirty(MemberIt)) << "Member '" << Descriptor->MemberProperties[MemberIt]->GetName() << "' index " << MemberIt;
+			UE_NET_ASSERT_FALSE_MSG(ReplicationStateB.IsDirty(MemberIt), "Member '" << Descriptor->MemberProperties[MemberIt]->GetName() << "' index " << MemberIt);
 		}
 	}
 }
diff --git a/Engine/Plugins/Runtime/ReplicationSystemTestPlugin/Source/Private/Tests/ReplicationState/TestReplicationStateDescriptorBuilder.cpp b/Engine/Plugins/Runtime/ReplicationSystemTestPlugin/Source/Private/Tests/ReplicationState/TestReplicationStateDescriptorBuilder.cpp
--- a/Engine/Plugins/Runtime/ReplicationSystemTestPlugin/Source/Private/Tests/ReplicationState/TestReplicationStateDescriptorBuilder.cpp
+++ b/Engine/Plugins/Runtime/ReplicationSystemTestPlugin/Source/Private/Tests/ReplicationState/TestReplicationStateDescriptorBuilder.cpp
@@ -774,8 +774,8 @@ UE_NET_TEST_FIXTURE(FTestReplicationStateDescriptorBuilderFixture, UTestReplicat
 	UE_NET_ASSERT_NE(SecondRoleConfig->RelativeInternalOffsetToOtherRole, 0);
 	UE_NET_ASSERT_NE(SecondRoleConfig->RelativeExternalOffsetToOtherRole, 0);
 
-	UE_NET_ASSERT_EQ(FirstRoleConfig->RelativeInternalOffsetToOtherRole, -SecondRoleConfig->RelativeInternalOffsetToOtherRole) << "Roles aren't referring to each other";
-	UE_NET_ASSERT_EQ(FirstRoleConfig->RelativeExternalOffsetToOtherRole, -SecondRoleConfig->RelativeExternalOffsetToOtherRole) << "Roles aren't referring to each other";
+	UE_NET_ASSERT_EQ_MSG(FirstRoleConfig->RelativeInternalOffsetToOtherRole, -SecondRoleConfig->RelativeInternalOffsetToOtherRole, "Roles aren't referring to each other");
+	UE_NET_ASSERT_EQ_MSG(FirstRoleConfig->RelativeExternalOffsetToOtherRole, -SecondRoleConfig->RelativeExternalOffsetToOtherRole, "Roles aren't referring to each other");
 }
 
 UE_NET_TEST_FIXTURE(FTestReplicationStateDescriptorBuilderFixture, UTestReplicationStateDescriptor_TestClassWithManyRoles)
@@ -1066,7 +1066,7 @@ UE_NET_TEST_FIXTURE(FTestReplicationStateDescriptorBuilderFixture, UTestReplicat
 	for (const FReplicationStateMemberLifetimeConditionDescriptor& ConditionDesc : MakeArrayView(LifetimeConditionalsDesc->MemberLifetimeConditionDescriptors, LifetimeConditionalsDesc->MemberCount))
 	{
 		const SIZE_T MemberIndex = &ConditionDesc - LifetimeConditionalsDesc->MemberLifetimeConditionDescriptors;
-		UE_NET_ASSERT_NE(ConditionDesc.Condition, int8(COND_None)) << "Unexpected lifetime condition for property " << LifetimeConditionalsDesc->MemberProperties[MemberIndex]->GetName();
+		UE_NET_ASSERT_NE_MSG(ConditionDesc.Condition, int8(COND_None), "Unexpected lifetime condition for property " << LifetimeConditionalsDesc->MemberProperties[MemberIndex]->GetName());
 	}
 }
 
diff --git a/Engine/Plugins/Runtime/ReplicationSystemTestPlugin/Source/Private/Tests/ReplicationSystem/Filtering/TestFiltering.cpp b/Engine/Plugins/Runtime/ReplicationSystemTestPlugin/Source/Private/Tests/ReplicationSystem/Filtering/TestFiltering.cpp
--- a/Engine/Plugins/Runtime/ReplicationSystemTestPlugin/Source/Private/Tests/ReplicationSystem/Filtering/TestFiltering.cpp
+++ b/Engine/Plugins/Runtime/ReplicationSystemTestPlugin/Source/Private/Tests/ReplicationSystem/Filtering/TestFiltering.cpp
@@ -14,6 +14,7 @@ namespace UE::Net::Private
 
 class FTestFilteringFixture : public FReplicationSystemServerClientTestFixture
 {
+protected:
 	virtual void SetUp() override
 	{
 		InitNetObjectFilterDefinitions();
@@ -27,6 +28,7 @@ class FTestFilteringFixture : public FReplicationSystemServerClientTestFixture
 		RestoreNetObjectPrioritizerDefinitions();
 	}
 
+private:
 	void InitNetObjectFilterDefinitions()
 	{
 		const UClass* NetObjectFilterDefinitionsClass = UNetObjectFilterDefinitions::StaticClass();
diff --git a/Engine/Plugins/Runtime/ReplicationSystemTestPlugin/Source/Private/Tests/ReplicationSystem/Prioritization/TestNetObjectPrioritizer.cpp b/Engine/Plugins/Runtime/ReplicationSystemTestPlugin/Source/Private/Tests/ReplicationSystem/Prioritization/TestNetObjectPrioritizer.cpp
--- a/Engine/Plugins/Runtime/ReplicationSystemTestPlugin/Source/Private/Tests/ReplicationSystem/Prioritization/TestNetObjectPrioritizer.cpp
+++ b/Engine/Plugins/Runtime/ReplicationSystemTestPlugin/Source/Private/Tests/ReplicationSystem/Prioritization/TestNetObjectPrioritizer.cpp
@@ -18,7 +18,7 @@ class FTestNetPrioritizerFixture : public FReplicationSystemTestFixture
 public:
 	FTestNetPrioritizerFixture() : MockNetObjectPrioritizer(nullptr) {}
 
-private:
+protected:
 	virtual void SetUp() override
 	{
 		InitNetObjectPrioritizerDefinitions();
@@ -36,6 +36,7 @@ private:
 		NetRefHandleManager = nullptr;
 	}
 
+private:
 	void InitNetObjectPrioritizerDefinitions()
 	{
 		const UClass* NetObjectPrioritizerDefinitionsClass = UNetObjectPrioritizerDefinitions::StaticClass();
diff --git a/Engine/Plugins/Runtime/ReplicationSystemTestPlugin/Source/Private/Tests/ReplicationSystem/TestDirtyNetObjectTracker.cpp b/Engine/Plugins/Runtime/ReplicationSystemTestPlugin/Source/Private/Tests/ReplicationSystem/TestDirtyNetObjectTracker.cpp
--- a/Engine/Plugins/Runtime/ReplicationSystemTestPlugin/Source/Private/Tests/ReplicationSystem/TestDirtyNetObjectTracker.cpp
+++ b/Engine/Plugins/Runtime/ReplicationSystemTestPlugin/Source/Private/Tests/ReplicationSystem/TestDirtyNetObjectTracker.cpp
@@ -42,7 +42,7 @@ UE_NET_TEST_FIXTURE(FDirtyNetObjectTrackerTestFixture, TestNoObjectIsDirtyFromSt
 	const FNetBitArrayView& DirtyObjects = DirtyNetObjectTracker->GetDirtyNetObjects();
 	UE_NET_ASSERT_EQ(DirtyObjects.GetNumBits(), NetObjectIndexRangeEnd + 1U);
 
-	UE_NET_ASSERT_FALSE(DirtyObjects.IsAnyBitSet()) << "Objects are marked as dirty before marking any object state as dirty";
+	UE_NET_ASSERT_FALSE_MSG(DirtyObjects.IsAnyBitSet(), "Objects are marked as dirty before marking any object state as dirty");
 }
 
 UE_NET_TEST_FIXTURE(FDirtyNetObjectTrackerTestFixture, CannotMarkInvalidObjectAsDirty)
diff --git a/Engine/Plugins/Runtime/ReplicationSystemTestPlugin/Source/Private/Tests/ReplicationSystem/TestObjectReferences.cpp b/Engine/Plugins/Runtime/ReplicationSystemTestPlugin/Source/Private/Tests/ReplicationSystem/TestObjectReferences.cpp
--- a/Engine/Plugins/Runtime/ReplicationSystemTestPlugin/Source/Private/Tests/ReplicationSystem/TestObjectReferences.cpp
+++ b/Engine/Plugins/Runtime/ReplicationSystemTestPlugin/Source/Private/Tests/ReplicationSystem/TestObjectReferences.cpp
@@ -762,7 +762,7 @@ UE_NET_TEST_FIXTURE(FReplicationSystemServerClientTestFixture, TestInlinedSubObj
 
 	// Verify subobject
 	UE_NET_ASSERT_TRUE(Object->CreatedSubObjectRef != nullptr);
-	UE_NET_ASSERT_TRUE(Object->CreatedSubObjectRef->HasAnyFlags(RF_DefaultSubObject) || Object->IsDefaultSubobject());
+	UE_NET_ASSERT_TRUE((Object->CreatedSubObjectRef->HasAnyFlags(RF_DefaultSubObject) || Object->IsDefaultSubobject()));
 
 	// Replicate object
 	Server->PreSendUpdate();
@@ -774,7 +774,7 @@ UE_NET_TEST_FIXTURE(FReplicationSystemServerClientTestFixture, TestInlinedSubObj
 
 	// Verify that default subobject is created
 	UE_NET_ASSERT_TRUE(ClientObject->CreatedSubObjectRef != nullptr);
-	UE_NET_ASSERT_TRUE(ClientObject->CreatedSubObjectRef->HasAnyFlags(RF_DefaultSubObject) || ClientObject->IsDefaultSubobject());
+	UE_NET_ASSERT_TRUE((ClientObject->CreatedSubObjectRef->HasAnyFlags(RF_DefaultSubObject) || ClientObject->IsDefaultSubobject()));
 	
 	// Set a reference to subobject
 	Object->ObjectRef = Object->CreatedSubObjectRef;
@@ -800,7 +800,7 @@ UE_NET_TEST_FIXTURE(FReplicationSystemServerClientTestFixture, TestExternalSubOb
 
 	// Verify internal subobject
 	UE_NET_ASSERT_TRUE(Object->CreatedSubObjectRef != nullptr);
-	UE_NET_ASSERT_TRUE(Object->CreatedSubObjectRef->HasAnyFlags(RF_DefaultSubObject) || Object->IsDefaultSubobject());
+	UE_NET_ASSERT_TRUE((Object->CreatedSubObjectRef->HasAnyFlags(RF_DefaultSubObject) || Object->IsDefaultSubobject()));
 
 	// Create external subobject
 	UTestObjectReferences_TestClassWithDefaultSubObject* ExternalSubObject = Server->CreateSubObject<UTestObjectReferences_TestClassWithDefaultSubObject>(Object->NetRefHandle);
@@ -817,7 +817,7 @@ UE_NET_TEST_FIXTURE(FReplicationSystemServerClientTestFixture, TestExternalSubOb
 	// Verify that default subobject is created
 	UE_NET_ASSERT_TRUE(ClientObject->CreatedSubObjectRef != nullptr);
 	UE_NET_ASSERT_TRUE(ClientExternalSubObject != nullptr);
-	UE_NET_ASSERT_TRUE(ClientObject->CreatedSubObjectRef->HasAnyFlags(RF_DefaultSubObject) || ClientObject->IsDefaultSubobject());
+	UE_NET_ASSERT_TRUE((ClientObject->CreatedSubObjectRef->HasAnyFlags(RF_DefaultSubObject) || ClientObject->IsDefaultSubobject()));
 		
 	// Set a reference to external subobject
 	Object->ObjectRef = ExternalSubObject;
diff --git a/Engine/Plugins/Runtime/ReplicationSystemTestPlugin/Source/Private/Tests/ReplicationSystem/TestObjectSplitting.cpp b/Engine/Plugins/Runtime/ReplicationSystemTestPlugin/Source/Private/Tests/ReplicationSystem/TestObjectSplitting.cpp
--- a/Engine/Plugins/Runtime/ReplicationSystemTestPlugin/Source/Private/Tests/ReplicationSystem/TestObjectSplitting.cpp
+++ b/Engine/Plugins/Runtime/ReplicationSystemTestPlugin/Source/Private/Tests/ReplicationSystem/TestObjectSplitting.cpp
@@ -685,7 +685,7 @@ UE_NET_TEST_FIXTURE(FSplitObjectTestFixture, TearOffWhileHugeObjectStateIsSendin
 
 	// Verify we have the previous state
 	const UTestReplicatedIrisObject* ClientObjectThatWillBeTornOff = Cast<UTestReplicatedIrisObject>(Client->GetReplicationBridge()->GetReplicatedObject(ServerObject->NetRefHandle));
-	UE_NET_ASSERT_EQ(ClientObjectThatWillBeTornOff->IntA ^ 1, ServerObject->IntA);
+	UE_NET_ASSERT_EQ((ClientObjectThatWillBeTornOff->IntA ^ 1), ServerObject->IntA);
 
 	Server->PreSendUpdate();
 	Server->SendAndDeliverTo(Client, true);
diff --git a/Engine/Plugins/Runtime/ReplicationSystemTestPlugin/Source/Private/Tests/ReplicationSystem/TestRepTag.cpp b/Engine/Plugins/Runtime/ReplicationSystemTestPlugin/Source/Private/Tests/ReplicationSystem/TestRepTag.cpp
--- a/Engine/Plugins/Runtime/ReplicationSystemTestPlugin/Source/Private/Tests/ReplicationSystem/TestRepTag.cpp
+++ b/Engine/Plugins/Runtime/ReplicationSystemTestPlugin/Source/Private/Tests/ReplicationSystem/TestRepTag.cpp
@@ -15,7 +15,7 @@ class FTestRepTagFixture : public FReplicationSystemTestFixture
 public:
 	FTestRepTagFixture() : ReplicationProtocol(nullptr), StateDescriptorWithTag(~0U) {}
 
-private:
+protected:
 	virtual void SetUp() override
 	{
 		FReplicationSystemTestFixture::SetUp();
diff --git a/Engine/Plugins/Runtime/ReplicationSystemTestPlugin/Source/Private/Tests/ReplicationSystem/TestReplicationBridge.cpp b/Engine/Plugins/Runtime/ReplicationSystemTestPlugin/Source/Private/Tests/ReplicationSystem/TestReplicationBridge.cpp
--- a/Engine/Plugins/Runtime/ReplicationSystemTestPlugin/Source/Private/Tests/ReplicationSystem/TestReplicationBridge.cpp
+++ b/Engine/Plugins/Runtime/ReplicationSystemTestPlugin/Source/Private/Tests/ReplicationSystem/TestReplicationBridge.cpp
@@ -117,7 +117,7 @@ UE_NET_TEST_FIXTURE(FReplicationBridgeTestFixture, CreateObjectWithDynamicStateC
 	UE_NET_ASSERT_NE(Protocol, nullptr);
 	UE_NET_ASSERT_NE(Protocol->ProtocolIdentifier, FReplicationProtocolIdentifier(0));
 	UE_NET_ASSERT_EQ(Protocol->ReplicationStateCount, 1U + DynamicStateComponentCount);
-	UE_NET_ASSERT_EQ(Protocol->ProtocolTraits, EReplicationProtocolTraits::HasDynamicState | EReplicationProtocolTraits::SupportsDeltaCompression);
+	UE_NET_ASSERT_EQ(Protocol->ProtocolTraits, (EReplicationProtocolTraits::HasDynamicState | EReplicationProtocolTraits::SupportsDeltaCompression));
 
 	// Destroy the handle
 	ReplicationBridge->EndReplication(TestObject);
diff --git a/Engine/Plugins/Runtime/ReplicationSystemTestPlugin/Source/Private/Tests/ReplicationSystem/TestReplicationOperations.cpp b/Engine/Plugins/Runtime/ReplicationSystemTestPlugin/Source/Private/Tests/ReplicationSystem/TestReplicationOperations.cpp
--- a/Engine/Plugins/Runtime/ReplicationSystemTestPlugin/Source/Private/Tests/ReplicationSystem/TestReplicationOperations.cpp
+++ b/Engine/Plugins/Runtime/ReplicationSystemTestPlugin/Source/Private/Tests/ReplicationSystem/TestReplicationOperations.cpp
@@ -194,7 +194,7 @@ UE_NET_TEST_FIXTURE(FTestReplicationOperationsFixture, CanQuantizeAndDequantize)
 	// check if we did get the expected values
 	UE_NET_ASSERT_EQ(1, TestObjectB->IntA);
 	UE_NET_ASSERT_EQ((int8)2, TestObjectB->IntC);
-	UE_NET_ASSERT_EQ(0, TestObjectB->StructD.NotReplicatedIntA) << "Property that isn't replicated has been overwritten";
+	UE_NET_ASSERT_EQ_MSG(0, TestObjectB->StructD.NotReplicatedIntA, "Property that isn't replicated has been overwritten");
 	UE_NET_ASSERT_EQ(int32('B'), TestObjectB->StructD.IntB);
 	UE_NET_ASSERT_EQ(3, TestObjectB->Components[0]->IntA);
 	UE_NET_ASSERT_EQ(4, TestObjectB->Components[11]->IntA);
diff --git a/Engine/Plugins/Runtime/ReplicationSystemTestPlugin/Source/Private/Tests/Serialization/MockNetSerializer.h b/Engine/Plugins/Runtime/ReplicationSystemTestPlugin/Source/Private/Tests/Serialization/MockNetSerializer.h
--- a/Engine/Plugins/Runtime/ReplicationSystemTestPlugin/Source/Private/Tests/Serialization/MockNetSerializer.h
+++ b/Engine/Plugins/Runtime/ReplicationSystemTestPlugin/Source/Private/Tests/Serialization/MockNetSerializer.h
@@ -6,17 +6,17 @@
 
 struct FMockNetSerializerCallCounter
 {
-	uint32 Serialize;
-	uint32 Deserialize;
-	uint32 SerializeDelta;
-	uint32 DeserializeDelta;
-	uint32 Quantize;
-	uint32 Dequantize;
-	uint32 IsEqual;
-	uint32 Validate;
-	uint32 CloneDynamicState;
-	uint32 FreeDynamicState;
-	uint32 CollectNetReferences;
+	uint32 Serialize{};
+	uint32 Deserialize{};
+	uint32 SerializeDelta{};
+	uint32 DeserializeDelta{};
+	uint32 Quantize{};
+	uint32 Dequantize{};
+	uint32 IsEqual{};
+	uint32 Validate{};
+	uint32 CloneDynamicState{};
+	uint32 FreeDynamicState{};
+	uint32 CollectNetReferences{};
 
 	void Reset() { *this = FMockNetSerializerCallCounter() = {}; }
 };
diff --git a/Engine/Plugins/Runtime/ReplicationSystemTestPlugin/Source/Private/Tests/Serialization/TestArrayPropertyNetSerializer.cpp b/Engine/Plugins/Runtime/ReplicationSystemTestPlugin/Source/Private/Tests/Serialization/TestArrayPropertyNetSerializer.cpp
--- a/Engine/Plugins/Runtime/ReplicationSystemTestPlugin/Source/Private/Tests/Serialization/TestArrayPropertyNetSerializer.cpp
+++ b/Engine/Plugins/Runtime/ReplicationSystemTestPlugin/Source/Private/Tests/Serialization/TestArrayPropertyNetSerializer.cpp
@@ -290,10 +290,10 @@ void FTestSimpleArrayPropertyNetSerializer::SetUp()
 	FTestArrayPropertyNetSerializerBase::SetUp();
 
 	ReplicationStateDescriptor = FReplicationStateDescriptorBuilder::CreateDescriptorForStruct(StaticStruct<FStructWithDynamicArrayOfPrimitiveTypeForArrayPropertyNetSerializerTest>());
-	UE_NET_ASSERT_EQ(ReplicationStateDescriptor->MemberCount, uint16(1)) << "Expected FStructWithDynamicArrayOfPrimitiveTypeForArrayPropertyNetSerializerTest to only contain a single array";
+	UE_NET_ASSERT_EQ_MSG(ReplicationStateDescriptor->MemberCount, uint16(1), "Expected FStructWithDynamicArrayOfPrimitiveTypeForArrayPropertyNetSerializerTest to only contain a single array");
 
 	ArrayPropertyNetSerializerConfig = const_cast<FArrayPropertyNetSerializerConfig*>(static_cast<const FArrayPropertyNetSerializerConfig*>(ReplicationStateDescriptor->MemberSerializerDescriptors[0].SerializerConfig));
-	UE_NET_ASSERT_EQ(ArrayPropertyNetSerializerConfig->StateDescriptor->MemberCount, uint16(1)) << "Expected array element descriptor to only contain a single member";
+	UE_NET_ASSERT_EQ_MSG(ArrayPropertyNetSerializerConfig->StateDescriptor->MemberCount, uint16(1), "Expected array element descriptor to only contain a single member");
 
 	OriginalMemberSerializerDescriptor = ArrayPropertyNetSerializerConfig->StateDescriptor->MemberSerializerDescriptors[0];
 
@@ -880,10 +880,10 @@ void FTestComplexArrayPropertyNetSerializer::SetUp()
 	FTestArrayPropertyNetSerializerBase::SetUp();
 
 	ReplicationStateDescriptor = FReplicationStateDescriptorBuilder::CreateDescriptorForStruct(StaticStruct<FStructWithDynamicArrayOfComplexTypeForArrayPropertyNetSerializerTest>());
-	UE_NET_ASSERT_EQ(ReplicationStateDescriptor->MemberCount, uint16(1)) << "Expected FStructWithDynamicArrayOfComplexTypeForArrayPropertyNetSerializerTest to only contain a single array";
+	UE_NET_ASSERT_EQ_MSG(ReplicationStateDescriptor->MemberCount, uint16(1), "Expected FStructWithDynamicArrayOfComplexTypeForArrayPropertyNetSerializerTest to only contain a single array");
 
 	ArrayPropertyNetSerializerConfig = const_cast<FArrayPropertyNetSerializerConfig*>(static_cast<const FArrayPropertyNetSerializerConfig*>(ReplicationStateDescriptor->MemberSerializerDescriptors[0].SerializerConfig));
-	UE_NET_ASSERT_EQ(ArrayPropertyNetSerializerConfig->StateDescriptor->MemberCount, uint16(1)) << "Expected array element descriptor to only contain a single member";
+	UE_NET_ASSERT_EQ_MSG(ArrayPropertyNetSerializerConfig->StateDescriptor->MemberCount, uint16(1), "Expected array element descriptor to only contain a single member");
 	UE_NET_ASSERT_TRUE(IsUsingStructNetSerializer(ArrayPropertyNetSerializerConfig->StateDescriptor->MemberSerializerDescriptors[0]));
 	UE_NET_ASSERT_TRUE(IsUsingArrayPropertyNetSerializer(static_cast<const FStructNetSerializerConfig*>(ArrayPropertyNetSerializerConfig->StateDescriptor->MemberSerializerDescriptors[0].SerializerConfig)->StateDescriptor->MemberSerializerDescriptors[0]));
 
diff --git a/Engine/Plugins/Runtime/ReplicationSystemTestPlugin/Source/Private/Tests/Serialization/TestBitPacking.cpp b/Engine/Plugins/Runtime/ReplicationSystemTestPlugin/Source/Private/Tests/Serialization/TestBitPacking.cpp
--- a/Engine/Plugins/Runtime/ReplicationSystemTestPlugin/Source/Private/Tests/Serialization/TestBitPacking.cpp
+++ b/Engine/Plugins/Runtime/ReplicationSystemTestPlugin/Source/Private/Tests/Serialization/TestBitPacking.cpp
@@ -144,14 +144,14 @@ void TTestBitPacking<SourceType, BitCount>::TestSerialize()
 
 			InitWriter();
 			SerializeDelta(Value, PrevValue);
-			UE_NET_ASSERT_FALSE(Writer.IsOverflown()) << "Serialize delta between " << Value << " and " << PrevValue << " caused bit stream overflow";
+			UE_NET_ASSERT_FALSE_MSG(Writer.IsOverflown(), "Serialize delta between " << Value << " and " << PrevValue << " caused bit stream overflow");
 
 			InitReaderFromWriter();
 			SourceType DeserializedValue;
 			DeserializeDelta(DeserializedValue, PrevValue);
-			UE_NET_ASSERT_FALSE(Reader.IsOverflown()) << "Deserialize delta between " << Value << " and " << PrevValue << " caused bit stream overflow";
+			UE_NET_ASSERT_FALSE_MSG(Reader.IsOverflown(), "Deserialize delta between " << Value << " and " << PrevValue << " caused bit stream overflow");
 
-			UE_NET_ASSERT_EQ(DeserializedValue, Value) << "Failed serializing value " << Value << " with previous value " << PrevValue << " using " << (TIsSigned<SourceType>::Value ? "signed " : "unsigned ") << BitCount << "-bit integer";
+			UE_NET_ASSERT_EQ_MSG(DeserializedValue, Value, "Failed serializing value " << Value << " with previous value " << PrevValue << " using " << (TIsSigned<SourceType>::Value ? "signed " : "unsigned ") << BitCount << "-bit integer");
 		}
 	}
 }
diff --git a/Engine/Plugins/Runtime/ReplicationSystemTestPlugin/Source/Private/Tests/Serialization/TestBitfieldNetSerializer.cpp b/Engine/Plugins/Runtime/ReplicationSystemTestPlugin/Source/Private/Tests/Serialization/TestBitfieldNetSerializer.cpp
--- a/Engine/Plugins/Runtime/ReplicationSystemTestPlugin/Source/Private/Tests/Serialization/TestBitfieldNetSerializer.cpp
+++ b/Engine/Plugins/Runtime/ReplicationSystemTestPlugin/Source/Private/Tests/Serialization/TestBitfieldNetSerializer.cpp
@@ -87,7 +87,7 @@ void FTestBitfieldNetSerializer<StructType, SourceType>::SetUp()
 
 		FBitfieldNetSerializerConfig Config;
 		const bool bIsValidBitfield = InitBitfieldNetSerializerConfigFromProperty(Config, BoolProperty);
-		UE_NET_ASSERT_TRUE(bIsValidBitfield) << "Unable to initialize bitfield from " << Struct->GetFName() << "::" << BoolProperty->GetNameCPP();
+		UE_NET_ASSERT_TRUE_MSG(bIsValidBitfield, "Unable to initialize bitfield from " << Struct->GetFName() << "::" << BoolProperty->GetNameCPP());
 		Configs.Add(Config);
 	}
 
@@ -106,19 +106,19 @@ void FTestBitfieldNetSerializer<StructType, SourceType>::TestIsEqual()
 
 		bIsQuantized = false;
 		const bool bIsDequantizedEqual = FTestNetSerializerFixture::TestIsEqual(Config, NetSerializerValuePointer(&Bit), NetSerializerValuePointer(&Bit), true, bIsQuantized);
-		UE_NET_ASSERT_TRUE(bIsDequantizedEqual) << "Set bit dequantized equality test failed with config " << Config;
+		UE_NET_ASSERT_TRUE_MSG(bIsDequantizedEqual, "Set bit dequantized equality test failed with config " << Config);
 
 		bIsQuantized = true;
 		const bool bIsQuantizedEqual = FTestNetSerializerFixture::TestIsEqual(Config, NetSerializerValuePointer(&Bit), NetSerializerValuePointer(&Bit), true, bIsQuantized);
-		UE_NET_ASSERT_TRUE(bIsQuantizedEqual) << "Set bit quantized equality test failed with config " << Config;
+		UE_NET_ASSERT_TRUE_MSG(bIsQuantizedEqual, "Set bit quantized equality test failed with config " << Config);
 
 		bIsQuantized = false;
 		const bool bZeroIsNotDequantizedEqualToOne = FTestNetSerializerFixture::TestIsEqual(Config, NetSerializerValuePointer(&Bit), NetSerializerValuePointer(&Zero), false, bIsQuantized);
-		UE_NET_ASSERT_TRUE(bZeroIsNotDequantizedEqualToOne) << "Set bit was dequantized equal to zero with config " << Config;
+		UE_NET_ASSERT_TRUE_MSG(bZeroIsNotDequantizedEqualToOne, "Set bit was dequantized equal to zero with config " << Config);
 
 		bIsQuantized = true;
 		const bool bZeroIsNotQuantizedEqualToOne = FTestNetSerializerFixture::TestIsEqual(Config, NetSerializerValuePointer(&Bit), NetSerializerValuePointer(&Zero), false, bIsQuantized);
-		UE_NET_ASSERT_TRUE(bZeroIsNotQuantizedEqualToOne) << "Set bit was quantized equal to zero with config " << Config;
+		UE_NET_ASSERT_TRUE_MSG(bZeroIsNotQuantizedEqualToOne, "Set bit was quantized equal to zero with config " << Config);
 	}
 }
 
@@ -133,10 +133,10 @@ void FTestBitfieldNetSerializer<StructType, SourceType>::TestValidate()
 	{
 		bool bIsValid;
 		bIsValid = FTestNetSerializerFixture::TestValidate(Config, NetSerializerValuePointer(&Zero), bExpectedResult);
-		UE_NET_ASSERT_TRUE(bIsValid) << "Unset bit was determined invalid with config " << Config;
+		UE_NET_ASSERT_TRUE_MSG(bIsValid, "Unset bit was determined invalid with config " << Config);
 
 		bIsValid = FTestNetSerializerFixture::TestValidate(Config, NetSerializerValuePointer(&One), bExpectedResult);
-		UE_NET_ASSERT_TRUE(bIsValid) << "Set bit was determined invalid with config " << Config;
+		UE_NET_ASSERT_TRUE_MSG(bIsValid, "Set bit was determined invalid with config " << Config);
 	}
 }
 
@@ -150,10 +150,10 @@ void FTestBitfieldNetSerializer<StructType, SourceType>::TestQuantize()
 	for (const FBitfieldNetSerializerConfig& Config : Configs)
 	{
 		bool bQuantize0Works = FTestNetSerializerFixture::TestQuantize(Config, NetSerializerValuePointer(&Zero));
-		UE_NET_ASSERT_TRUE(bQuantize0Works) << "0 bit could not be quantized with config " << Config;
+		UE_NET_ASSERT_TRUE_MSG(bQuantize0Works, "0 bit could not be quantized with config " << Config);
 
 		bool bQuantize1Works = FTestNetSerializerFixture::TestQuantize(Config, NetSerializerValuePointer(&One));
-		UE_NET_ASSERT_TRUE(bQuantize1Works) << " 1 bit could not be quantized with config " << Config;
+		UE_NET_ASSERT_TRUE_MSG(bQuantize1Works, " 1 bit could not be quantized with config " << Config);
 	}
 }
 
@@ -171,7 +171,7 @@ void FTestBitfieldNetSerializer<StructType, SourceType>::TestSerialize()
 		{
 			const SourceType Value = Values[ValueIt];
 			const bool bSerializeWorks = FTestNetSerializerFixture::TestSerialize(Config, NetSerializerValuePointer(&Value), NetSerializerValuePointer(&Value), bQuantizedCompare);
-			UE_NET_ASSERT_TRUE(bSerializeWorks) << Value << " could not be serialized with config " << Config;
+			UE_NET_ASSERT_TRUE_MSG(bSerializeWorks, Value << " could not be serialized with config " << Config);
 		}
 	}
 }
@@ -193,7 +193,7 @@ void FTestBitfieldNetSerializer<StructType, SourceType>::TestDequantize()
 		QuantizeArgs.Source = NetSerializerValuePointer(&Value);
 		QuantizeArgs.Target = NetSerializerValuePointer(QuantizedBuffer[0]);
 		Serializer.Quantize(Context, QuantizeArgs);
-		UE_NET_ASSERT_FALSE(Context.HasError()) << "Quantize() of value reported an error. Config " << Config;
+		UE_NET_ASSERT_FALSE_MSG(Context.HasError(), "Quantize() of value reported an error. Config " << Config);
 
 		Writer.InitBytes(BitStreamBuffer, sizeof(BitStreamBuffer));
 		FNetSerializeArgs SerializeArgs;
@@ -202,8 +202,8 @@ void FTestBitfieldNetSerializer<StructType, SourceType>::TestDequantize()
 		SerializeArgs.Source = QuantizeArgs.Target;
 		Serializer.Serialize(Context, SerializeArgs);
 		Writer.CommitWrites();
-		UE_NET_ASSERT_FALSE(Writer.IsOverflown()) << "FNetBitStreamWriter overflowed. Config " << Config;
-		UE_NET_ASSERT_FALSE(Context.HasError()) << "Serialize() reported an error. Config " << Config;
+		UE_NET_ASSERT_FALSE_MSG(Writer.IsOverflown(), "FNetBitStreamWriter overflowed. Config " << Config);
+		UE_NET_ASSERT_FALSE_MSG(Context.HasError(), "Serialize() reported an error. Config " << Config);
 
 		Reader.InitBits(BitStreamBuffer, Writer.GetPosBits());
 		FNetDeserializeArgs DeserializeArgs;
@@ -211,8 +211,8 @@ void FTestBitfieldNetSerializer<StructType, SourceType>::TestDequantize()
 		DeserializeArgs.NetSerializerConfig = &Config;
 		DeserializeArgs.Target = NetSerializerValuePointer(QuantizedBuffer[1]);
 		Serializer.Deserialize(Context, DeserializeArgs);
-		UE_NET_ASSERT_FALSE(Reader.IsOverflown()) << "FNetBitStreamReader overflowed. Config " << Config;
-		UE_NET_ASSERT_FALSE(Context.HasError()) << "Deserialize() reported an error. Config " << Config;
+		UE_NET_ASSERT_FALSE_MSG(Reader.IsOverflown(), "FNetBitStreamReader overflowed. Config " << Config);
+		UE_NET_ASSERT_FALSE_MSG(Context.HasError(), "Deserialize() reported an error. Config " << Config);
 
 		// Need to dequantize the deserialized value
 		FNetDequantizeArgs DequantizeArgs;
@@ -221,9 +221,9 @@ void FTestBitfieldNetSerializer<StructType, SourceType>::TestDequantize()
 		DequantizeArgs.Source = DeserializeArgs.Target;
 		DequantizeArgs.Target = NetSerializerValuePointer(&Target);
 		Serializer.Dequantize(Context, DequantizeArgs);
-		UE_NET_ASSERT_FALSE(Context.HasError()) << "Dequantize() of deserialized value reported an error. Config " << Config;
+		UE_NET_ASSERT_FALSE_MSG(Context.HasError(), "Dequantize() of deserialized value reported an error. Config " << Config);
 
-		UE_NET_ASSERT_EQ(OriginalTarget, *reinterpret_cast<const SourceType*>(DequantizeArgs.Target)) << "Dequantize touched bits that should not have been modified. Config " << Config;
+		UE_NET_ASSERT_EQ_MSG(OriginalTarget, *reinterpret_cast<const SourceType*>(DequantizeArgs.Target), "Dequantize touched bits that should not have been modified. Config " << Config);
 	}
 }
 
diff --git a/Engine/Plugins/Runtime/ReplicationSystemTestPlugin/Source/Private/Tests/Serialization/TestBoolNetSerializer.cpp b/Engine/Plugins/Runtime/ReplicationSystemTestPlugin/Source/Private/Tests/Serialization/TestBoolNetSerializer.cpp
--- a/Engine/Plugins/Runtime/ReplicationSystemTestPlugin/Source/Private/Tests/Serialization/TestBoolNetSerializer.cpp
+++ b/Engine/Plugins/Runtime/ReplicationSystemTestPlugin/Source/Private/Tests/Serialization/TestBoolNetSerializer.cpp
@@ -30,7 +30,7 @@ protected:
 
 UE_NET_TEST_FIXTURE(FTestBoolNetSerializer, HasTestValues)
 {
-	UE_NET_ASSERT_GT(ValueCount, SIZE_T(0)) << "No test values found";
+	UE_NET_ASSERT_GT_MSG(ValueCount, SIZE_T(0), "No test values found");
 }
 
 UE_NET_TEST_FIXTURE(FTestBoolNetSerializer, TestValuesAreAsExpected)
diff --git a/Engine/Plugins/Runtime/ReplicationSystemTestPlugin/Source/Private/Tests/Serialization/TestEnumNetSerializer.cpp b/Engine/Plugins/Runtime/ReplicationSystemTestPlugin/Source/Private/Tests/Serialization/TestEnumNetSerializer.cpp
--- a/Engine/Plugins/Runtime/ReplicationSystemTestPlugin/Source/Private/Tests/Serialization/TestEnumNetSerializer.cpp
+++ b/Engine/Plugins/Runtime/ReplicationSystemTestPlugin/Source/Private/Tests/Serialization/TestEnumNetSerializer.cpp
@@ -35,10 +35,9 @@ public:
 	void TestSerialize();
 	void TestSerializeDelta();
 
-private:
+protected:
 	virtual void SetUp() override;
 
-protected:
 	static TArray<SourceType> Values;
 	static TArray<SourceType> InvalidValues;
 	static SerializerConfig Config;
@@ -239,7 +238,7 @@ void FTestEnumIntNetSerializer<SerializerConfig, SourceType, EnumType>::TestVali
 
 	// Check invalid values
 	{
-		UE_NET_EXPECT_GT(InvalidValues.Num(), 0) << "Unable to test EnumIntSerializer Validate with invalid values.";
+		UE_NET_EXPECT_GT_MSG(InvalidValues.Num(), 0, "Unable to test EnumIntSerializer Validate with invalid values.");
 
 		TArray<bool> ExpectedResults;
 		ExpectedResults.SetNumZeroed(InvalidValues.Num());
diff --git a/Engine/Plugins/Runtime/ReplicationSystemTestPlugin/Source/Private/Tests/Serialization/TestIntNetSerializer.cpp b/Engine/Plugins/Runtime/ReplicationSystemTestPlugin/Source/Private/Tests/Serialization/TestIntNetSerializer.cpp
--- a/Engine/Plugins/Runtime/ReplicationSystemTestPlugin/Source/Private/Tests/Serialization/TestIntNetSerializer.cpp
+++ b/Engine/Plugins/Runtime/ReplicationSystemTestPlugin/Source/Private/Tests/Serialization/TestIntNetSerializer.cpp
@@ -40,7 +40,7 @@ public: \
 \
 UE_NET_TEST_FIXTURE(TestClassName, HasTestValues) \
 { \
-	UE_NET_ASSERT_GT(ValueCount, SIZE_T(0)) << "No test values found"; \
+	UE_NET_ASSERT_GT_MSG(ValueCount, SIZE_T(0), "No test values found"); \
 } \
 \
 UE_NET_TEST_FIXTURE(TestClassName, TestIsEqual) \
@@ -137,7 +137,7 @@ void FTestIntNetSerializer<SourceType>::TestValidate()
 			ExpectedResults[ValueIt] = (FMath::Clamp(Value, MinValueForBitCount, MaxValueForBitCount) == Value);
 			if (BitCount == sizeof(SourceType)*8)
 			{
-				UE_NET_ASSERT_TRUE(ExpectedResults[ValueIt]) << "Clamping of value with full bit precision resulted in value being clamped. This is unexpected. Make sure no undefined behavior is used in code.";
+				UE_NET_ASSERT_TRUE_MSG(ExpectedResults[ValueIt], "Clamping of value with full bit precision resulted in value being clamped. This is unexpected. Make sure no undefined behavior is used in code.");
 			}
 		}
 
diff --git a/Engine/Plugins/Runtime/ReplicationSystemTestPlugin/Source/Private/Tests/Serialization/TestIntRangeNetSerializer.cpp b/Engine/Plugins/Runtime/ReplicationSystemTestPlugin/Source/Private/Tests/Serialization/TestIntRangeNetSerializer.cpp
--- a/Engine/Plugins/Runtime/ReplicationSystemTestPlugin/Source/Private/Tests/Serialization/TestIntRangeNetSerializer.cpp
+++ b/Engine/Plugins/Runtime/ReplicationSystemTestPlugin/Source/Private/Tests/Serialization/TestIntRangeNetSerializer.cpp
@@ -50,12 +50,12 @@ public:
 
 UE_NET_TEST_FIXTURE(FTestInt8RangeNetSerializer, HasTestValues)
 {
-	UE_NET_ASSERT_GT(ValueCount, SIZE_T(0)) << "No test values found";
+	UE_NET_ASSERT_GT_MSG(ValueCount, SIZE_T(0), "No test values found");
 }
 
 UE_NET_TEST_FIXTURE(FTestInt8RangeNetSerializer, HasTestConfigs)
 {
-	UE_NET_ASSERT_GT(ConfigCount, SIZE_T(0)) << "No configs found";
+	UE_NET_ASSERT_GT_MSG(ConfigCount, SIZE_T(0), "No configs found");
 }
 
 UE_NET_TEST_FIXTURE(FTestInt8RangeNetSerializer, TestIsEqual)
@@ -92,12 +92,12 @@ public:
 
 UE_NET_TEST_FIXTURE(FTestInt16RangeNetSerializer, HasTestValues)
 {
-	UE_NET_ASSERT_GT(ValueCount, SIZE_T(0)) << "No test values found";
+	UE_NET_ASSERT_GT_MSG(ValueCount, SIZE_T(0), "No test values found");
 }
 
 UE_NET_TEST_FIXTURE(FTestInt16RangeNetSerializer, HasTestConfigs)
 {
-	UE_NET_ASSERT_GT(ConfigCount, SIZE_T(0)) << "No configs found";
+	UE_NET_ASSERT_GT_MSG(ConfigCount, SIZE_T(0), "No configs found");
 }
 
 UE_NET_TEST_FIXTURE(FTestInt16RangeNetSerializer, TestIsEqual)
@@ -134,12 +134,12 @@ public:
 
 UE_NET_TEST_FIXTURE(FTestInt32RangeNetSerializer, HasTestValues)
 {
-	UE_NET_ASSERT_GT(ValueCount, SIZE_T(0)) << "No test values found";
+	UE_NET_ASSERT_GT_MSG(ValueCount, SIZE_T(0), "No test values found");
 }
 
 UE_NET_TEST_FIXTURE(FTestInt32RangeNetSerializer, HasTestConfigs)
 {
-	UE_NET_ASSERT_GT(ConfigCount, SIZE_T(0)) << "No configs found";
+	UE_NET_ASSERT_GT_MSG(ConfigCount, SIZE_T(0), "No configs found");
 }
 
 UE_NET_TEST_FIXTURE(FTestInt32RangeNetSerializer, TestIsEqual)
@@ -176,12 +176,12 @@ public:
 
 UE_NET_TEST_FIXTURE(FTestInt64RangeNetSerializer, HasTestValues)
 {
-	UE_NET_ASSERT_GT(ValueCount, SIZE_T(0)) << "No test values found";
+	UE_NET_ASSERT_GT_MSG(ValueCount, SIZE_T(0), "No test values found");
 }
 
 UE_NET_TEST_FIXTURE(FTestInt64RangeNetSerializer, HasTestConfigs)
 {
-	UE_NET_ASSERT_GT(ConfigCount, SIZE_T(0)) << "No configs found";
+	UE_NET_ASSERT_GT_MSG(ConfigCount, SIZE_T(0), "No configs found");
 }
 
 UE_NET_TEST_FIXTURE(FTestInt64RangeNetSerializer, TestIsEqual)
@@ -300,7 +300,7 @@ void FTestIntRangeNetSerializer<SerializerConfig, SourceType>::TestValidate()
 			ExpectedResults[ValueIt] = (FMath::Clamp(Value, Config.LowerBound, Config.UpperBound) == Value);
 			if (ConfigIt == FullRangeConfigIndex)
 			{
-				UE_NET_ASSERT_TRUE(ExpectedResults[ValueIt]) << "Clamping of value with full bit precision resulted in value being clamped. This is unexpected. Make sure no undefined behavior is used in code.";
+				UE_NET_ASSERT_TRUE_MSG(ExpectedResults[ValueIt], "Clamping of value with full bit precision resulted in value being clamped. This is unexpected. Make sure no undefined behavior is used in code.");
 			}
 		}
 
diff --git a/Engine/Plugins/Runtime/ReplicationSystemTestPlugin/Source/Private/Tests/Serialization/TestLastResortPropertyNetSerializer.cpp b/Engine/Plugins/Runtime/ReplicationSystemTestPlugin/Source/Private/Tests/Serialization/TestLastResortPropertyNetSerializer.cpp
--- a/Engine/Plugins/Runtime/ReplicationSystemTestPlugin/Source/Private/Tests/Serialization/TestLastResortPropertyNetSerializer.cpp
+++ b/Engine/Plugins/Runtime/ReplicationSystemTestPlugin/Source/Private/Tests/Serialization/TestLastResortPropertyNetSerializer.cpp
@@ -94,7 +94,7 @@ void FTestLastResortPropertyNetSerializer::SetUp()
 	Super::SetUp();
 
 	ReplicationStateDescriptor = FReplicationStateDescriptorBuilder::CreateDescriptorForStruct(StaticStruct<FStructForLastResortPropertyNetSerializerTest>());
-	UE_NET_ASSERT_EQ(ReplicationStateDescriptor->MemberCount, uint16(1)) << "Expected FStructForLastResortPropertyNetSerializerTest to contain exactly one member";
+	UE_NET_ASSERT_EQ_MSG(ReplicationStateDescriptor->MemberCount, uint16(1), "Expected FStructForLastResortPropertyNetSerializerTest to contain exactly one member");
 
 
 	SerializerConfig = const_cast<FLastResortPropertyNetSerializerConfig*>(static_cast<const FLastResortPropertyNetSerializerConfig*>(ReplicationStateDescriptor->MemberSerializerDescriptors[0].SerializerConfig));
diff --git a/Engine/Plugins/Runtime/ReplicationSystemTestPlugin/Source/Private/Tests/Serialization/TestNetBitStreams.cpp b/Engine/Plugins/Runtime/ReplicationSystemTestPlugin/Source/Private/Tests/Serialization/TestNetBitStreams.cpp
--- a/Engine/Plugins/Runtime/ReplicationSystemTestPlugin/Source/Private/Tests/Serialization/TestNetBitStreams.cpp
+++ b/Engine/Plugins/Runtime/ReplicationSystemTestPlugin/Source/Private/Tests/Serialization/TestNetBitStreams.cpp
@@ -188,8 +188,8 @@ UE_NET_TEST_FIXTURE(FNetBitStreamReaderWriterTest, WriteBitsAtOffset0)
 
 		const uint32 ExpectedValue = uint32((uint64(1) << BitCount) - uint64(1));
 
-		UE_NET_ASSERT_EQ(ReadValue, ExpectedValue) << FString::Printf(TEXT("Failed testing with %u bits"), BitCount);
-		UE_NET_ASSERT_EQ(ReadSentinel, Sentinel) << FString::Printf(TEXT("Failed testing with %u bits"), BitCount);
+		UE_NET_ASSERT_EQ_MSG(ReadValue, ExpectedValue, FString::Printf(TEXT("Failed testing with %u bits"), BitCount));
+		UE_NET_ASSERT_EQ_MSG(ReadSentinel, Sentinel, FString::Printf(TEXT("Failed testing with %u bits"), BitCount));
 	}
 }
 
@@ -214,8 +214,8 @@ UE_NET_TEST_FIXTURE(FNetBitStreamReaderWriterTest, WriteBitsAtArbitraryOffsets)
 
 			const uint32 ExpectedValue = uint32((uint64(1) << BitCount) - uint64(1));
 
-			UE_NET_ASSERT_EQ(ReadValue, ExpectedValue) << FString::Printf(TEXT("Failed testing with %u bits at offset %u"), BitCount, BitOffset);
-			UE_NET_ASSERT_EQ(ReadSentinel, Sentinel) << FString::Printf(TEXT("Failed testing with %u bits at offset %u"), BitCount, BitOffset);
+			UE_NET_ASSERT_EQ_MSG(ReadValue, ExpectedValue, FString::Printf(TEXT("Failed testing with %u bits at offset %u"), BitCount, BitOffset));
+			UE_NET_ASSERT_EQ_MSG(ReadSentinel, Sentinel, FString::Printf(TEXT("Failed testing with %u bits at offset %u"), BitCount, BitOffset));
 		}
 	}
 }
@@ -266,10 +266,10 @@ UE_NET_TEST_FIXTURE(FNetBitStreamReaderWriterTest, WriteStreamWithBitsWrittenAtA
 				const uint32 BitCount = ValuesAndBitCount[ValueIt][1];
 
 				const uint32 ReadValue = Reader.ReadBits(BitCount);
-				UE_NET_ASSERT_EQ(ReadValue, ExpectedValue) << FString::Printf(TEXT("Write stream with bits written at offset %u to stream at offset %u"), BitOffset0, BitOffset1);
+				UE_NET_ASSERT_EQ_MSG(ReadValue, ExpectedValue, FString::Printf(TEXT("Write stream with bits written at offset %u to stream at offset %u"), BitOffset0, BitOffset1));
 			}
 			const uint32 ReadSentinel = Reader.ReadBits(sizeof(Sentinel)*8U);
-			UE_NET_ASSERT_EQ(ReadSentinel, Sentinel) << FString::Printf(TEXT("Write stream with bits written at offset %u to stream at offset %u"), BitOffset0, BitOffset1);
+			UE_NET_ASSERT_EQ_MSG(ReadSentinel, Sentinel, FString::Printf(TEXT("Write stream with bits written at offset %u to stream at offset %u"), BitOffset0, BitOffset1));
 		}
 	}
 }
@@ -325,7 +325,7 @@ UE_NET_TEST_FIXTURE(FNetBitStreamReaderWriterTest, ReadStreamWithBitsWrittenAtAr
 			Reader.ReadBitStream(ResultBuffer, TotalBitCount);
 
 			const uint32 ReadSentinel = Reader.ReadBits(sizeof(Sentinel)*8U);
-			UE_NET_ASSERT_EQ(ReadSentinel, Sentinel) << FString::Printf(TEXT("Write stream with bits written at offset %u to stream at offset %u"), BitOffset0, BitOffset1);
+			UE_NET_ASSERT_EQ_MSG(ReadSentinel, Sentinel, FString::Printf(TEXT("Write stream with bits written at offset %u to stream at offset %u"), BitOffset0, BitOffset1));
 
 			Reader.InitBits(ResultBuffer, TotalBitCount);
 			for (size_t ValueIt = 0, ValueEndIt = sizeof(ValuesAndBitCount)/sizeof(ValuesAndBitCount[0]); ValueIt != ValueEndIt; ++ValueIt)
@@ -334,7 +334,7 @@ UE_NET_TEST_FIXTURE(FNetBitStreamReaderWriterTest, ReadStreamWithBitsWrittenAtAr
 				const uint32 BitCount = ValuesAndBitCount[ValueIt][1];
 
 				const uint32 ReadValue = Reader.ReadBits(BitCount);
-				UE_NET_ASSERT_EQ(ReadValue, ExpectedValue) << FString::Printf(TEXT("Write stream with bits written at offset %u to stream at offset %u"), BitOffset0, BitOffset1);
+				UE_NET_ASSERT_EQ_MSG(ReadValue, ExpectedValue, FString::Printf(TEXT("Write stream with bits written at offset %u to stream at offset %u"), BitOffset0, BitOffset1));
 			}
 		}
 	}
@@ -502,7 +502,7 @@ UE_NET_TEST_FIXTURE(FNetBitStreamSubstreamTest, CanReadDataCommittedFromSubstrea
 		Reader.Seek(SeekCount1);
 		const uint32 ReadWord = Reader.ReadBits(WriteCount2);
 		const uint32 WordMask = ~0U >> (32 - WriteCount2);
-		UE_NET_ASSERT_EQ(ReadWord & WordMask, WriteWord & WordMask);
+		UE_NET_ASSERT_EQ((ReadWord & WordMask), (WriteWord & WordMask));
 	}
 }
 
diff --git a/Engine/Plugins/Runtime/ReplicationSystemTestPlugin/Source/Private/Tests/Serialization/TestNetRoleNetSerializer.cpp b/Engine/Plugins/Runtime/ReplicationSystemTestPlugin/Source/Private/Tests/Serialization/TestNetRoleNetSerializer.cpp
--- a/Engine/Plugins/Runtime/ReplicationSystemTestPlugin/Source/Private/Tests/Serialization/TestNetRoleNetSerializer.cpp
+++ b/Engine/Plugins/Runtime/ReplicationSystemTestPlugin/Source/Private/Tests/Serialization/TestNetRoleNetSerializer.cpp
@@ -44,10 +44,9 @@ public:
 	void TestDequantizeSwapsRoles();
 	void TestDowngradeRole();
 
-private:
+protected:
 	virtual void SetUp() override;
 
-protected:
 	static TArray<uint8> Values;
 	static TArray<uint8> InvalidValues;
 	static FNetRoleNetSerializerConfig Config;
@@ -205,7 +204,7 @@ void FTestNetRoleNetSerializer::TestValidate()
 
 	// Check invalid values
 	{
-		UE_NET_EXPECT_GT(InvalidValues.Num(), 0) << "No invalid values found to test EnumIntSerializer::Validate.";
+		UE_NET_EXPECT_GT_MSG(InvalidValues.Num(), 0, "No invalid values found to test EnumIntSerializer::Validate.");
 
 		TArray<bool> ExpectedResults;
 		ExpectedResults.SetNumZeroed(InvalidValues.Num());
@@ -279,8 +278,8 @@ void FTestNetRoleNetSerializer::TestDequantizeSwapsRoles()
 		constexpr bool bPushAll = true;
 		TargetReplicationState.PushPropertyReplicationState(Target.Get(), bPushAll);
 		
-		UE_NET_ASSERT_EQ(Target->RemoteRole, Source->Role) << "Roles were not swapped";
-		UE_NET_ASSERT_EQ(Target->Role, Source->RemoteRole) << "Roles were not swapped";
+		UE_NET_ASSERT_EQ_MSG(Target->RemoteRole, Source->Role, "Roles were not swapped");
+		UE_NET_ASSERT_EQ_MSG(Target->Role, Source->RemoteRole, "Roles were not swapped");
 	}
 }
 
@@ -327,8 +326,8 @@ void FTestNetRoleNetSerializer::TestDowngradeRole()
 		constexpr bool bPushAll = true;
 		TargetReplicationState.PushPropertyReplicationState(Target.Get(), bPushAll);
 		
-		UE_NET_ASSERT_EQ(Target->RemoteRole, Source->Role) << "Roles were not swapped";
-		UE_NET_ASSERT_EQ(Target->Role, TEnumAsByte<ENetRole>(ENetRole::ROLE_SimulatedProxy)) << "Role was not downgraded";
+		UE_NET_ASSERT_EQ_MSG(Target->RemoteRole, Source->Role, "Roles were not swapped");
+		UE_NET_ASSERT_EQ_MSG(Target->Role, TEnumAsByte<ENetRole>(ENetRole::ROLE_SimulatedProxy), "Role was not downgraded");
 	}
 }
 
diff --git a/Engine/Plugins/Runtime/ReplicationSystemTestPlugin/Source/Private/Tests/Serialization/TestNetSerializerFixture.cpp b/Engine/Plugins/Runtime/ReplicationSystemTestPlugin/Source/Private/Tests/Serialization/TestNetSerializerFixture.cpp
--- a/Engine/Plugins/Runtime/ReplicationSystemTestPlugin/Source/Private/Tests/Serialization/TestNetSerializerFixture.cpp
+++ b/Engine/Plugins/Runtime/ReplicationSystemTestPlugin/Source/Private/Tests/Serialization/TestNetSerializerFixture.cpp
@@ -104,7 +104,7 @@ bool FTestNetSerializerFixture::TestQuantize(const FNetSerializerConfig& Config,
 	FreeScope.AddBuffer(QuantizeArgs.Target);
 	if (Context.HasError())
 	{
-		UE_NET_EXPECT_FALSE(Context.HasError()) << "Quantize() of value reported an error.";
+		UE_NET_EXPECT_FALSE_MSG(Context.HasError(), "Quantize() of value reported an error.");
 		return false;
 	}
 
@@ -117,7 +117,7 @@ bool FTestNetSerializerFixture::TestQuantize(const FNetSerializerConfig& Config,
 	const bool bIsEqual = Serializer.IsEqual(Context, QuantizedIsEqualArgs);
 	if (!bIsEqual)
 	{
-		UE_NET_EXPECT_TRUE(bIsEqual) << "IsEqual() on quantized value compared to itself returned false.";
+		UE_NET_EXPECT_TRUE_MSG(bIsEqual, "IsEqual() on quantized value compared to itself returned false.");
 		return false;
 	}
 
@@ -129,7 +129,7 @@ bool FTestNetSerializerFixture::TestQuantize(const FNetSerializerConfig& Config,
 	Serializer.Dequantize(Context, DequantizeArgs);
 	if (Context.HasError())
 	{
-		UE_NET_EXPECT_FALSE(Context.HasError()) << "Dequantize() of quantized value reported an error.";
+		UE_NET_EXPECT_FALSE_MSG(Context.HasError(), "Dequantize() of quantized value reported an error.");
 		return false;
 	}
 
@@ -142,7 +142,7 @@ bool FTestNetSerializerFixture::TestQuantize(const FNetSerializerConfig& Config,
 	const bool bIsDequantizedEqual = Serializer.IsEqual(Context, DequantizedIsEqualArgs);
 	if (!bIsDequantizedEqual)
 	{
-		UE_NET_EXPECT_TRUE(bIsDequantizedEqual) << "IsEqual() on dequantized value compared to source value returned false.";
+		UE_NET_EXPECT_TRUE_MSG(bIsDequantizedEqual, "IsEqual() on dequantized value compared to source value returned false.");
 		return false;
 	}
 
@@ -161,7 +161,7 @@ bool FTestNetSerializerFixture::TestQuantize(const FNetSerializerConfig& Config,
 		const bool bIsQuantizedEqual = !FPlatformMemory::Memcmp(QuantizedBuffer[0], QuantizedBuffer[1], Serializer.QuantizedTypeSize);
 		if (!bIsQuantizedEqual)
 		{
-			UE_NET_EXPECT_TRUE(bIsQuantizedEqual) << "Memcmp() on quantized values with tainted buffers failed.";
+			UE_NET_EXPECT_TRUE_MSG(bIsQuantizedEqual, "Memcmp() on quantized values with tainted buffers failed.");
 			return false;
 		}
 	}
@@ -191,7 +191,7 @@ bool FTestNetSerializerFixture::TestSerialize(const FNetSerializerConfig& Config
 	FreeScope.AddBuffer(QuantizeArgs.Target);
 	if (Context.HasError())
 	{
-		UE_NET_EXPECT_FALSE(Context.HasError()) << "Quantize() of value reported an error.";
+		UE_NET_EXPECT_FALSE_MSG(Context.HasError(), "Quantize() of value reported an error.");
 		return false;
 	}
 
@@ -204,8 +204,8 @@ bool FTestNetSerializerFixture::TestSerialize(const FNetSerializerConfig& Config
 	Writer.CommitWrites();
 	if (Writer.IsOverflown() || Context.HasError())
 	{
-		UE_NET_EXPECT_FALSE(Writer.IsOverflown()) << "FNetBitStreamWriter overflowed.";
-		UE_NET_EXPECT_FALSE(Context.HasError()) << "Serialize() reported an error.";
+		UE_NET_EXPECT_FALSE_MSG(Writer.IsOverflown(), "FNetBitStreamWriter overflowed.");
+		UE_NET_EXPECT_FALSE_MSG(Context.HasError(), "Serialize() reported an error.");
 		return false;
 	}
 
@@ -218,8 +218,8 @@ bool FTestNetSerializerFixture::TestSerialize(const FNetSerializerConfig& Config
 	FreeScope.AddBuffer(DeserializeArgs.Target);
 	if (Reader.IsOverflown() || Context.HasError())
 	{
-		UE_NET_EXPECT_FALSE(Reader.IsOverflown()) << "FNetBitStreamReader overflowed.";
-		UE_NET_EXPECT_FALSE(Context.HasError()) << "Deserialize() reported an error.";
+		UE_NET_EXPECT_FALSE_MSG(Reader.IsOverflown(), "FNetBitStreamReader overflowed.");
+		UE_NET_EXPECT_FALSE_MSG(Context.HasError(), "Deserialize() reported an error.");
 		return false;
 	}
 
@@ -234,7 +234,7 @@ bool FTestNetSerializerFixture::TestSerialize(const FNetSerializerConfig& Config
 		Serializer.Dequantize(Context, DequantizeArgs);
 		if (Context.HasError())
 		{
-			UE_NET_EXPECT_FALSE(Context.HasError()) << "Dequantize() of deserialized value reported an error.";
+			UE_NET_EXPECT_FALSE_MSG(Context.HasError(), "Dequantize() of deserialized value reported an error.");
 			return false;
 		}
 
@@ -250,7 +250,7 @@ bool FTestNetSerializerFixture::TestSerialize(const FNetSerializerConfig& Config
 			const bool bIsEqual = (*CompareFunc)(IsEqualArgs.Source0, IsEqualArgs.Source1);
 			if (!bIsEqual)
 			{
-				UE_NET_EXPECT_TRUE(bIsEqual) << "Custom compare on dequantized values returned false";
+				UE_NET_EXPECT_TRUE_MSG(bIsEqual, "Custom compare on dequantized values returned false");
 				return false;
 			}
 		}
@@ -259,7 +259,7 @@ bool FTestNetSerializerFixture::TestSerialize(const FNetSerializerConfig& Config
 			const bool bIsEqual = Serializer.IsEqual(Context, IsEqualArgs);
 			if (!bIsEqual)
 			{
-				UE_NET_EXPECT_TRUE(bIsEqual) << "IsEqual() on dequantized values returned false";
+				UE_NET_EXPECT_TRUE_MSG(bIsEqual, "IsEqual() on dequantized values returned false");
 				return false;
 			}
 		}
@@ -275,7 +275,7 @@ bool FTestNetSerializerFixture::TestSerialize(const FNetSerializerConfig& Config
 			Serializer.Quantize(Context, QuantizeArgs);
 			if (Context.HasError())
 			{
-				UE_NET_EXPECT_FALSE(Context.HasError()) << "Quantize() of expected value reported an error.";
+				UE_NET_EXPECT_FALSE_MSG(Context.HasError(), "Quantize() of expected value reported an error.");
 				return false;
 			}
 		}
@@ -314,7 +314,7 @@ bool FTestNetSerializerFixture::TestSerializeDelta(const FNetSerializerConfig& C
 			FreeScope.AddBuffer(QuantizeValueArgs.Target);
 			if (Context.HasError())
 			{
-				UE_NET_EXPECT_FALSE(Context.HasError()) << "Quantize() of value reported an error.";
+				UE_NET_EXPECT_FALSE_MSG(Context.HasError(), "Quantize() of value reported an error.");
 				return false;
 			}
 		}
@@ -328,7 +328,7 @@ bool FTestNetSerializerFixture::TestSerializeDelta(const FNetSerializerConfig& C
 			FreeScope.AddBuffer(QuantizePrevValueArgs.Target);
 			if (Context.HasError())
 			{
-				UE_NET_EXPECT_FALSE(Context.HasError()) << "Quantize() of value reported an error.";
+				UE_NET_EXPECT_FALSE_MSG(Context.HasError(), "Quantize() of value reported an error.");
 				return false;
 			}
 		}
@@ -346,8 +346,8 @@ bool FTestNetSerializerFixture::TestSerializeDelta(const FNetSerializerConfig& C
 		Writer.CommitWrites();
 		if (Writer.IsOverflown() || Context.HasError())
 		{
-			UE_NET_EXPECT_FALSE(Writer.IsOverflown()) << "FNetBitStreamWriter overflowed.";
-			UE_NET_EXPECT_FALSE(Context.HasError()) << "SerializeDelta() reported an error.";
+			UE_NET_EXPECT_FALSE_MSG(Writer.IsOverflown(), "FNetBitStreamWriter overflowed.");
+			UE_NET_EXPECT_FALSE_MSG(Context.HasError(), "SerializeDelta() reported an error.");
 			return false;
 		}
 	}
@@ -364,8 +364,8 @@ bool FTestNetSerializerFixture::TestSerializeDelta(const FNetSerializerConfig& C
 		FreeScope.AddBuffer(DeserializeDeltaArgs.Target);
 		if (Reader.IsOverflown() || Context.HasError())
 		{
-			UE_NET_EXPECT_FALSE(Reader.IsOverflown()) << "FNetBitStreamReader overflowed.";
-			UE_NET_EXPECT_FALSE(Context.HasError()) << "DeserializeDelta() reported an error.";
+			UE_NET_EXPECT_FALSE_MSG(Reader.IsOverflown(), "FNetBitStreamReader overflowed.");
+			UE_NET_EXPECT_FALSE_MSG(Context.HasError(), "DeserializeDelta() reported an error.");
 			return false;
 		}
 	}
@@ -426,7 +426,7 @@ bool FTestNetSerializerFixture::TestIsEqual(const FNetSerializerConfig& Config,
 	const bool bResult2 = Serializer.IsEqual(Context, IsEqualArgs);
 	if (bResult2 != bExpectedResult)
 	{
-		UE_NET_EXPECT_EQ(bResult1, bResult2) << "IsEqual did not return the same result when source values were swapped.";
+		UE_NET_EXPECT_EQ_MSG(bResult1, bResult2, "IsEqual did not return the same result when source values were swapped.");
 		return false;
 	}
 
@@ -464,7 +464,7 @@ bool FTestNetSerializerFixture::TestCloneDynamicState(const FNetSerializerConfig
 	FreeScope.AddBuffer(QuantizeArgs.Target);
 	if (Context.HasError())
 	{
-		UE_NET_EXPECT_FALSE(Context.HasError()) << "Quantize() of value reported an error.";
+		UE_NET_EXPECT_FALSE_MSG(Context.HasError(), "Quantize() of value reported an error.");
 		return false;
 	}
 
@@ -477,7 +477,7 @@ bool FTestNetSerializerFixture::TestCloneDynamicState(const FNetSerializerConfig
 	Serializer.CloneDynamicState(Context, CloneArgs);
 	if (Context.HasError())
 	{
-		UE_NET_EXPECT_FALSE(Context.HasError()) << "CloneDynamicState() reported an error.";
+		UE_NET_EXPECT_FALSE_MSG(Context.HasError(), "CloneDynamicState() reported an error.");
 		return false;
 	}
 
@@ -490,7 +490,7 @@ bool FTestNetSerializerFixture::TestCloneDynamicState(const FNetSerializerConfig
 		Serializer.FreeDynamicState(Context, FreeArgs);
 		if (Context.HasError())
 		{
-			UE_NET_EXPECT_FALSE(Context.HasError()) << "FreeDynamicState() reported an error.";
+			UE_NET_EXPECT_FALSE_MSG(Context.HasError(), "FreeDynamicState() reported an error.");
 			return false;
 		}
 	}
@@ -505,7 +505,7 @@ bool FTestNetSerializerFixture::TestCloneDynamicState(const FNetSerializerConfig
 		Serializer.Dequantize(Context, DequantizeArgs);
 		if (Context.HasError())
 		{
-			UE_NET_EXPECT_FALSE(Context.HasError()) << "Dequantize() of quantized value reported an error.";
+			UE_NET_EXPECT_FALSE_MSG(Context.HasError(), "Dequantize() of quantized value reported an error.");
 			return false;
 		}
 
@@ -518,7 +518,7 @@ bool FTestNetSerializerFixture::TestCloneDynamicState(const FNetSerializerConfig
 		const bool bIsDequantizedEqual = Serializer.IsEqual(Context, DequantizedIsEqualArgs);
 		if (!bIsDequantizedEqual)
 		{
-			UE_NET_EXPECT_TRUE(bIsDequantizedEqual) << "IsEqual() on dequantized value compared to source value returned false.";
+			UE_NET_EXPECT_TRUE_MSG(bIsDequantizedEqual, "IsEqual() on dequantized value compared to source value returned false.");
 			return false;
 		}
 	}
@@ -542,7 +542,7 @@ bool FTestNetSerializerFixture::Serialize(const FNetSerializerConfig& Config, Ne
 	FreeScope.AddBuffer(QuantizeArgs.Target);
 	if (Context.HasError())
 	{
-		UE_NET_EXPECT_FALSE(Context.HasError()) << "Quantize() of value reported an error.";
+		UE_NET_EXPECT_FALSE_MSG(Context.HasError(), "Quantize() of value reported an error.");
 		return false;
 	}
 
@@ -555,8 +555,8 @@ bool FTestNetSerializerFixture::Serialize(const FNetSerializerConfig& Config, Ne
 	Writer.CommitWrites();
 	if (Writer.IsOverflown() || Context.HasError())
 	{
-		UE_NET_EXPECT_FALSE(Writer.IsOverflown()) << "FNetBitStreamWriter overflowed.";
-		UE_NET_EXPECT_FALSE(Context.HasError()) << "Serialize() reported an error.";
+		UE_NET_EXPECT_FALSE_MSG(Writer.IsOverflown(), "FNetBitStreamWriter overflowed.");
+		UE_NET_EXPECT_FALSE_MSG(Context.HasError(), "Serialize() reported an error.");
 		return false;
 	}
 
diff --git a/Engine/Plugins/Runtime/ReplicationSystemTestPlugin/Source/Private/Tests/Serialization/TestNetSerializerFixture.h b/Engine/Plugins/Runtime/ReplicationSystemTestPlugin/Source/Private/Tests/Serialization/TestNetSerializerFixture.h
--- a/Engine/Plugins/Runtime/ReplicationSystemTestPlugin/Source/Private/Tests/Serialization/TestNetSerializerFixture.h
+++ b/Engine/Plugins/Runtime/ReplicationSystemTestPlugin/Source/Private/Tests/Serialization/TestNetSerializerFixture.h
@@ -61,9 +61,9 @@ protected:
 		BufferSize = 1024,
 	};
 
-	alignas(16) uint8 BitStreamBuffer[BitStreamBufferSize];
-	alignas(16) uint8 QuantizedBuffer[2][BufferSize];
-	alignas(16) uint8 SourceBuffer[2][BufferSize];
+	alignas(16) uint8 BitStreamBuffer[BitStreamBufferSize] = {};
+	alignas(16) uint8 QuantizedBuffer[2][BufferSize] = {};
+	alignas(16) uint8 SourceBuffer[2][BufferSize] = {};
 
 	bool bTaintBuffersBeforeTest = false;
 };
@@ -113,7 +113,7 @@ bool TTestNetSerializerFixture<ConfigPrinter, SourceType>::TestQuantize(const So
 		{
 			FTestMessage ConfigMessage;
 			bOutTestCaseSuccess = FTestNetSerializerFixture::TestQuantize(Config, NetSerializerValuePointer(Value));
-			UE_NET_ASSERT_TRUE(bOutTestCaseSuccess) << "Quantization failed with Value: '" << *Value << "' " << ConfigPrinter(ConfigMessage, Config);
+			UE_NET_ASSERT_TRUE_MSG(bOutTestCaseSuccess, "Quantization failed with Value: '" << *Value << "' " << ConfigPrinter(ConfigMessage, Config));
 		};
 
 		bool bOutTestCaseSuccess = false;
@@ -140,7 +140,7 @@ bool TTestNetSerializerFixture<ConfigPrinter, SourceType>::TestSerialize(const S
 		{
 			FTestMessage ConfigMessage;
 			bOutTestCaseSuccess = FTestNetSerializerFixture::TestSerialize(Config, NetSerializerValuePointer(Value), NetSerializerValuePointer(ExpectedValue), bQuantizedCompare, CompareFunc);
-			UE_NET_ASSERT_TRUE(bOutTestCaseSuccess) << "Serialization failed with Value: '" << *Value << "' and ExpectedValue: '" << *ExpectedValue << "' " << ConfigPrinter(ConfigMessage, Config);
+			UE_NET_ASSERT_TRUE_MSG(bOutTestCaseSuccess, "Serialization failed with Value: '" << *Value << "' and ExpectedValue: '" << *ExpectedValue << "' " << ConfigPrinter(ConfigMessage, Config));
 		};
 
 		bool bOutTestCaseSuccess = false;
@@ -170,7 +170,7 @@ bool TTestNetSerializerFixture<ConfigPrinter, SourceType>::TestSerializeDelta(co
 			{
 				FTestMessage ConfigMessage;
 				bOutTestCaseSuccess = FTestNetSerializerFixture::TestSerializeDelta(Config, NetSerializerValuePointer(Value), NetSerializerValuePointer(PrevValue));
-				UE_NET_ASSERT_TRUE(bOutTestCaseSuccess) << "Delta serialization failed with Value: '" << *Value << "' and PrevValue: '" << *PrevValue << "' " << ConfigPrinter(ConfigMessage, Config);
+				UE_NET_ASSERT_TRUE_MSG(bOutTestCaseSuccess, "Delta serialization failed with Value: '" << *Value << "' and PrevValue: '" << *PrevValue << "' " << ConfigPrinter(ConfigMessage, Config));
 			};
 
 			bool bOutTestCaseSuccess = false;
@@ -195,7 +195,7 @@ bool TTestNetSerializerFixture<ConfigPrinter, SourceType>::TestIsEqual(const Sou
 		{
 			FTestMessage ConfigMessage;
 			bOutTestCaseSuccess = FTestNetSerializerFixture::TestIsEqual(Config, NetSerializerValuePointer(Value0), NetSerializerValuePointer(Value1), bExpectedResult, bQuantizedCompare);
-			UE_NET_ASSERT_TRUE(bOutTestCaseSuccess) << "IsEqual failed. Value0: '" << *Value0 << "' Value1: '" << *Value1 << "' " << ConfigPrinter(ConfigMessage, Config) << ". Expected IsEqual to return " << bExpectedResult;
+			UE_NET_ASSERT_TRUE_MSG(bOutTestCaseSuccess, "IsEqual failed. Value0: '" << *Value0 << "' Value1: '" << *Value1 << "' " << ConfigPrinter(ConfigMessage, Config) << ". Expected IsEqual to return " << bExpectedResult);
 		};
 
 		const SourceType* Value0 = Values0 + ValueIt;
@@ -222,7 +222,7 @@ bool TTestNetSerializerFixture<ConfigPrinter, SourceType>::TestValidate(const So
 		{
 			FTestMessage ConfigMessage;
 			bOutTestCaseSuccess = FTestNetSerializerFixture::TestValidate(Config, NetSerializerValuePointer(Value), bExpectedResult);
-			UE_NET_ASSERT_TRUE(bOutTestCaseSuccess) << "Validate failed. Value: '" << *Value << "' " << ConfigPrinter(ConfigMessage, Config) << ". Expected Validate to return " << bExpectedResult;
+			UE_NET_ASSERT_TRUE_MSG(bOutTestCaseSuccess, "Validate failed. Value: '" << *Value << "' " << ConfigPrinter(ConfigMessage, Config) << ". Expected Validate to return " << bExpectedResult);
 		};
 
 		const SourceType* Value = Values + ValueIt;
@@ -250,7 +250,7 @@ bool TTestNetSerializerFixture<ConfigPrinter, SourceType>::TestCloneDynamicState
 		{
 			FTestMessage ConfigMessage;
 			bOutTestCaseSuccess = FTestNetSerializerFixture::TestCloneDynamicState(Config, NetSerializerValuePointer(Value));
-			UE_NET_ASSERT_TRUE(bOutTestCaseSuccess) << "CloneDynamicState failed with Value: '" << *Value << "' " << ConfigPrinter(ConfigMessage, Config);
+			UE_NET_ASSERT_TRUE_MSG(bOutTestCaseSuccess, "CloneDynamicState failed with Value: '" << *Value << "' " << ConfigPrinter(ConfigMessage, Config));
 		};
 
 		bool bOutTestCaseSuccess = false;
diff --git a/Engine/Plugins/Runtime/ReplicationSystemTestPlugin/Source/Private/Tests/Serialization/TestPackedIntNetSerializer.cpp b/Engine/Plugins/Runtime/ReplicationSystemTestPlugin/Source/Private/Tests/Serialization/TestPackedIntNetSerializer.cpp
--- a/Engine/Plugins/Runtime/ReplicationSystemTestPlugin/Source/Private/Tests/Serialization/TestPackedIntNetSerializer.cpp
+++ b/Engine/Plugins/Runtime/ReplicationSystemTestPlugin/Source/Private/Tests/Serialization/TestPackedIntNetSerializer.cpp
@@ -38,7 +38,7 @@ public:
 
 UE_NET_TEST_FIXTURE(FTestPackedInt32NetSerializer, HasTestValues)
 {
-	UE_NET_ASSERT_GT(ValueCount, SIZE_T(0)) << "No test values found";
+	UE_NET_ASSERT_GT_MSG(ValueCount, SIZE_T(0), "No test values found");
 }
 
 UE_NET_TEST_FIXTURE(FTestPackedInt32NetSerializer, TestIsEqual)
@@ -130,14 +130,14 @@ void TTestPackedIntNetSerializer<SourceType>::TestSmallValuesArePacked()
 	constexpr SourceType LargeValue = TNumericLimits<SourceType>::Max();
 
 	const bool bSerializedSmallValue = this->Serialize(Config, NetSerializerValuePointer(&SmallValue));
-	UE_NET_ASSERT_TRUE(bSerializedSmallValue) << "Failed serializing " << SmallValue;
+	UE_NET_ASSERT_TRUE_MSG(bSerializedSmallValue, "Failed serializing " << SmallValue);
 	const uint32 SmallValueBits = this->Writer.GetPosBits();
 
 	const bool bSerializedLargeValue = this->Serialize(Config, NetSerializerValuePointer(&LargeValue));
-	UE_NET_ASSERT_TRUE(bSerializedLargeValue) << "Failed serializing " << LargeValue;
+	UE_NET_ASSERT_TRUE_MSG(bSerializedLargeValue, "Failed serializing " << LargeValue);
 	const uint32 LargeValueBits = this->Writer.GetPosBits();
 
-	UE_NET_ASSERT_LT(SmallValueBits, LargeValueBits) << "Bits serialized for small value (" << SmallValue << ") should've been less than for large value (" << LargeValue << ")";
+	UE_NET_ASSERT_LT_MSG(SmallValueBits, LargeValueBits, "Bits serialized for small value (" << SmallValue << ") should've been less than for large value (" << LargeValue << ")");
 }
 
 }
\ No newline at end of file
diff --git a/Engine/Plugins/Runtime/ReplicationSystemTestPlugin/Source/Private/Tests/Serialization/TestPolymorphicStructNetSerializer.cpp b/Engine/Plugins/Runtime/ReplicationSystemTestPlugin/Source/Private/Tests/Serialization/TestPolymorphicStructNetSerializer.cpp
--- a/Engine/Plugins/Runtime/ReplicationSystemTestPlugin/Source/Private/Tests/Serialization/TestPolymorphicStructNetSerializer.cpp
+++ b/Engine/Plugins/Runtime/ReplicationSystemTestPlugin/Source/Private/Tests/Serialization/TestPolymorphicStructNetSerializer.cpp
@@ -542,7 +542,8 @@ void FTestPolymorphicArrayStructNetSerializerFixture::Deserialize()
 void FTestPolymorphicArrayStructNetSerializerFixture::SerializeDelta()
 {
 	// Check pre-conditions
-	UE_NET_ASSERT_TRUE(bHasQuantizedState && QuantizedStateCount == 2U);
+	UE_NET_ASSERT_TRUE(bHasQuantizedState);
+	UE_NET_ASSERT_EQ(QuantizedStateCount, 2U);
 	
 	// Serialize data
 	{
@@ -673,9 +674,17 @@ bool FTestPolymorphicArrayStructNetSerializerFixture::IsEqual(bool bQuantized)
 {
 	if (bQuantized)
 	{
-		// Trick to return false if assertion fails.
-		UE_NET_ASSERT_TRUE(bHasQuantizedState), false;
-		UE_NET_ASSERT_TRUE(bHasClonedQuantizedState), false;
+		UE_NET_EXPECT_TRUE(bHasQuantizedState);
+		if (!bHasQuantizedState)
+		{
+			return false;
+		}
+
+		UE_NET_EXPECT_TRUE(bHasClonedQuantizedState);
+		if (!bHasClonedQuantizedState)
+		{
+			return false;
+		}
 	}
 
 	FNetSerializationContext Context;
diff --git a/Engine/Plugins/Runtime/ReplicationSystemTestPlugin/Source/Private/Tests/Serialization/TestStringNetSerializer.cpp b/Engine/Plugins/Runtime/ReplicationSystemTestPlugin/Source/Private/Tests/Serialization/TestStringNetSerializer.cpp
--- a/Engine/Plugins/Runtime/ReplicationSystemTestPlugin/Source/Private/Tests/Serialization/TestStringNetSerializer.cpp
+++ b/Engine/Plugins/Runtime/ReplicationSystemTestPlugin/Source/Private/Tests/Serialization/TestStringNetSerializer.cpp
@@ -278,7 +278,7 @@ void FTestStringNetSerializer::TestIllFormedStringSerialize()
 			QuantizeArgs.Source = NetSerializerValuePointer(&SourceString);
 			QuantizeArgs.Target = NetSerializerValuePointer(QuantizedStateBuffer[0]);
 			Serializer.Quantize(Context, QuantizeArgs);
-			UE_NET_ASSERT_FALSE(Context.HasError()) << "Quantize failed.";
+			UE_NET_ASSERT_FALSE_MSG(Context.HasError(), "Quantize failed.");
 		}
 
 		// Serialize
@@ -290,8 +290,8 @@ void FTestStringNetSerializer::TestIllFormedStringSerialize()
 			SerializeArgs.Source = NetSerializerValuePointer(QuantizedStateBuffer[0]);
 			Serializer.Serialize(Context, SerializeArgs);
 			Writer.CommitWrites();
-			UE_NET_ASSERT_FALSE(Context.HasError()) << "Serialize() reported an error.";
-			UE_NET_ASSERT_FALSE(Writer.IsOverflown()) << "FNetBitStreamWriter overflowed.";
+			UE_NET_ASSERT_FALSE_MSG(Context.HasError(), "Serialize() reported an error.");
+			UE_NET_ASSERT_FALSE_MSG(Writer.IsOverflown(), "FNetBitStreamWriter overflowed.");
 		}
 
 		// Deserialize
@@ -303,8 +303,8 @@ void FTestStringNetSerializer::TestIllFormedStringSerialize()
 			DeserializeArgs.Target = NetSerializerValuePointer(QuantizedStateBuffer[1]);
 			Serializer.Deserialize(Context, DeserializeArgs);
 			Writer.CommitWrites();
-			UE_NET_ASSERT_FALSE(Context.HasError()) << "Deserialize() reported an error.";
-			UE_NET_ASSERT_FALSE(Reader.IsOverflown()) << "FNetBitStreamRead overflowed.";
+			UE_NET_ASSERT_FALSE_MSG(Context.HasError(), "Deserialize() reported an error.");
+			UE_NET_ASSERT_FALSE_MSG(Reader.IsOverflown(), "FNetBitStreamRead overflowed.");
 		}
 
 		// Dequantize
@@ -316,7 +316,7 @@ void FTestStringNetSerializer::TestIllFormedStringSerialize()
 			DequantizeArgs.Target = NetSerializerValuePointer(&TargetString);
 			Serializer.Dequantize(Context, DequantizeArgs);
 			// Depending on the size of wchar_t an error may be generated or not on certain ill-formed strings.
-			//UE_NET_ASSERT_FALSE(Context.HasError()) << "Dequantize() reported an error for string '" << SourceString << '\'';
+			//UE_NET_ASSERT_FALSE_MSG(Context.HasError(), "Dequantize() reported an error for string '" << SourceString << '\'');
 		}
 
 		// Free memory
diff --git a/Engine/Plugins/Runtime/ReplicationSystemTestPlugin/Source/Private/Tests/Serialization/TestStringNetSerializerUtils.cpp b/Engine/Plugins/Runtime/ReplicationSystemTestPlugin/Source/Private/Tests/Serialization/TestStringNetSerializerUtils.cpp
--- a/Engine/Plugins/Runtime/ReplicationSystemTestPlugin/Source/Private/Tests/Serialization/TestStringNetSerializerUtils.cpp
+++ b/Engine/Plugins/Runtime/ReplicationSystemTestPlugin/Source/Private/Tests/Serialization/TestStringNetSerializerUtils.cpp
@@ -40,15 +40,15 @@ UE_NET_TEST_FIXTURE(FTestStringNetSerializerUtils, TestEnc32Dec32)
 	uint8* EncodingBuffer = static_cast<uint8*>(FMemory_Alloca(EncodedBufferLen));
 	uint32 OutEncodedLen = 0;
 	FStringNetSerializerUtils::TStringCodec<char32_t>::Encode(EncodingBuffer, EncodedBufferLen, TestStringUTF32, TestStringUTF32Len, OutEncodedLen);
-	UE_NET_ASSERT_GT(OutEncodedLen, 0U) << FString::Printf(TEXT("Test encoding of UTF-32 '%s'."), FUTF8ToTCHAR(TestStringUTF8).Get());
+	UE_NET_ASSERT_GT_MSG(OutEncodedLen, 0U, FString::Printf(TEXT("Test encoding of UTF-32 '%s'."), FUTF8ToTCHAR(TestStringUTF8).Get()));
 
 	// Decode
 	uint32 DecodedBufferLen = FStringNetSerializerUtils::TStringCodec<char32_t>::GetSafeDecodedBufferLength(OutEncodedLen);
 	char32_t* DecodingBuffer = static_cast<char32_t*>(FMemory_Alloca(DecodedBufferLen * sizeof(char32_t)));
 	uint32 OutDecodedLen = 0;
 	FStringNetSerializerUtils::TStringCodec<char32_t>::Decode(DecodingBuffer, DecodedBufferLen, EncodingBuffer, OutEncodedLen, OutDecodedLen);
-	UE_NET_ASSERT_EQ(OutDecodedLen, static_cast<uint32>(TestStringUTF32Len)) << FString::Printf(TEXT("Test UTF-32 decoding of '%s'."), FUTF8ToTCHAR(TestStringUTF8).Get());
-	UE_NET_ASSERT_EQ(FMemory::Memcmp(DecodingBuffer, TestStringUTF32, TestStringUTF32Len), 0) << FString::Printf(TEXT("Test UTF-32 decoding of '%s'."), FUTF8ToTCHAR(TestStringUTF8).Get());
+	UE_NET_ASSERT_EQ_MSG(OutDecodedLen, static_cast<uint32>(TestStringUTF32Len), FString::Printf(TEXT("Test UTF-32 decoding of '%s'."), FUTF8ToTCHAR(TestStringUTF8).Get()));
+	UE_NET_ASSERT_EQ_MSG(FMemory::Memcmp(DecodingBuffer, TestStringUTF32, TestStringUTF32Len), 0, FString::Printf(TEXT("Test UTF-32 decoding of '%s'."), FUTF8ToTCHAR(TestStringUTF8).Get()));
 }
 
 UE_NET_TEST_FIXTURE(FTestStringNetSerializerUtils, TestEnc16Dec16)
@@ -58,15 +58,15 @@ UE_NET_TEST_FIXTURE(FTestStringNetSerializerUtils, TestEnc16Dec16)
 	uint8* EncodingBuffer = static_cast<uint8*>(FMemory_Alloca(EncodedBufferLen));
 	uint32 OutEncodedLen = 0;
 	FStringNetSerializerUtils::TStringCodec<char16_t>::Encode(EncodingBuffer, EncodedBufferLen, TestStringUTF16, TestStringUTF16Len, OutEncodedLen);
-	UE_NET_ASSERT_GT(OutEncodedLen, 0U) << FString::Printf(TEXT("Test encoding of UTF-16 '%s'."), FUTF8ToTCHAR(TestStringUTF8).Get());
+	UE_NET_ASSERT_GT_MSG(OutEncodedLen, 0U, FString::Printf(TEXT("Test encoding of UTF-16 '%s'."), FUTF8ToTCHAR(TestStringUTF8).Get()));
 
 	// Decode
 	uint32 DecodedBufferLen = FStringNetSerializerUtils::TStringCodec<char16_t>::GetSafeDecodedBufferLength(OutEncodedLen * sizeof(char16_t));
 	char16_t* DecodingBuffer = static_cast<char16_t*>(FMemory_Alloca(DecodedBufferLen));
 	uint32 OutDecodedLen = 0;
 	FStringNetSerializerUtils::TStringCodec<char16_t>::Decode(DecodingBuffer, DecodedBufferLen, EncodingBuffer, OutEncodedLen, OutDecodedLen);
-	UE_NET_ASSERT_EQ(OutDecodedLen, static_cast<uint32>(TestStringUTF16Len)) << FString::Printf(TEXT("Test UTF-16 decoding of '%s'."), FUTF8ToTCHAR(TestStringUTF8).Get());
-	UE_NET_ASSERT_EQ(FMemory::Memcmp(DecodingBuffer, TestStringUTF16, TestStringUTF16Len), 0) << FString::Printf(TEXT("Test UTF-16 decoding of '%s'."), FUTF8ToTCHAR(TestStringUTF8).Get());
+	UE_NET_ASSERT_EQ_MSG(OutDecodedLen, static_cast<uint32>(TestStringUTF16Len), FString::Printf(TEXT("Test UTF-16 decoding of '%s'."), FUTF8ToTCHAR(TestStringUTF8).Get()));
+	UE_NET_ASSERT_EQ_MSG(FMemory::Memcmp(DecodingBuffer, TestStringUTF16, TestStringUTF16Len), 0, FString::Printf(TEXT("Test UTF-16 decoding of '%s'."), FUTF8ToTCHAR(TestStringUTF8).Get()));
 }
 
 UE_NET_TEST_FIXTURE(FTestStringNetSerializerUtils, TestEnc16Dec32)
@@ -82,8 +82,8 @@ UE_NET_TEST_FIXTURE(FTestStringNetSerializerUtils, TestEnc16Dec32)
 	char32_t* DecodingBuffer = static_cast<char32_t*>(FMemory_Alloca(DecodedBufferLen * sizeof(char32_t)));
 	uint32 OutDecodedLen = 0;
 	FStringNetSerializerUtils::TStringCodec<char32_t>::Decode(DecodingBuffer, DecodedBufferLen, EncodingBuffer, OutEncodedLen, OutDecodedLen);
-	UE_NET_ASSERT_EQ(OutDecodedLen, static_cast<uint32>(TestStringUTF32Len)) << FString::Printf(TEXT("Test UTF-32 decoding of UTF-16 '%s'."), FUTF8ToTCHAR(TestStringUTF8).Get());
-	UE_NET_ASSERT_EQ(FMemory::Memcmp(DecodingBuffer, TestStringUTF32, TestStringUTF32Len), 0) << FString::Printf(TEXT("Test UTF-32 decoding of UTF-16 '%s'."), FUTF8ToTCHAR(TestStringUTF8).Get());
+	UE_NET_ASSERT_EQ_MSG(OutDecodedLen, static_cast<uint32>(TestStringUTF32Len), FString::Printf(TEXT("Test UTF-32 decoding of UTF-16 '%s'."), FUTF8ToTCHAR(TestStringUTF8).Get()));
+	UE_NET_ASSERT_EQ_MSG(FMemory::Memcmp(DecodingBuffer, TestStringUTF32, TestStringUTF32Len), 0, FString::Printf(TEXT("Test UTF-32 decoding of UTF-16 '%s'."), FUTF8ToTCHAR(TestStringUTF8).Get()));
 }
 
 UE_NET_TEST_FIXTURE(FTestStringNetSerializerUtils, TestEnc32Dec16)
@@ -99,8 +99,8 @@ UE_NET_TEST_FIXTURE(FTestStringNetSerializerUtils, TestEnc32Dec16)
 	char16_t* DecodingBuffer = static_cast<char16_t*>(FMemory_Alloca(DecodedBufferLen * sizeof(char16_t)));
 	uint32 OutDecodedLen = 0;
 	FStringNetSerializerUtils::TStringCodec<char16_t>::Decode(DecodingBuffer, DecodedBufferLen, EncodingBuffer, OutEncodedLen, OutDecodedLen);
-	UE_NET_ASSERT_EQ(OutDecodedLen, static_cast<uint32>(TestStringUTF16Len)) << FString::Printf(TEXT("Test UTF-16 decoding of UTF-32 '%s'."), FUTF8ToTCHAR(TestStringUTF8).Get());
-	UE_NET_ASSERT_EQ(FMemory::Memcmp(DecodingBuffer, TestStringUTF16, TestStringUTF16Len), 0) << FString::Printf(TEXT("Test UTF-16 decoding of UTF-32 '%s'."), FUTF8ToTCHAR(TestStringUTF8).Get());
+	UE_NET_ASSERT_EQ_MSG(OutDecodedLen, static_cast<uint32>(TestStringUTF16Len), FString::Printf(TEXT("Test UTF-16 decoding of UTF-32 '%s'."), FUTF8ToTCHAR(TestStringUTF8).Get()));
+	UE_NET_ASSERT_EQ_MSG(FMemory::Memcmp(DecodingBuffer, TestStringUTF16, TestStringUTF16Len), 0, FString::Printf(TEXT("Test UTF-16 decoding of UTF-32 '%s'."), FUTF8ToTCHAR(TestStringUTF8).Get()));
 }
 
 }
diff --git a/Engine/Plugins/Runtime/ReplicationSystemTestPlugin/Source/Private/Tests/Serialization/TestUintNetSerializer.cpp b/Engine/Plugins/Runtime/ReplicationSystemTestPlugin/Source/Private/Tests/Serialization/TestUintNetSerializer.cpp
--- a/Engine/Plugins/Runtime/ReplicationSystemTestPlugin/Source/Private/Tests/Serialization/TestUintNetSerializer.cpp
+++ b/Engine/Plugins/Runtime/ReplicationSystemTestPlugin/Source/Private/Tests/Serialization/TestUintNetSerializer.cpp
@@ -40,7 +40,7 @@ public: \
 \
 UE_NET_TEST_FIXTURE(TestClassName, HasTestValues) \
 { \
-	UE_NET_ASSERT_GT(ValueCount, SIZE_T(0)) << "No test values found"; \
+	UE_NET_ASSERT_GT_MSG(ValueCount, SIZE_T(0), "No test values found"); \
 } \
 \
 UE_NET_TEST_FIXTURE(TestClassName, TestIsEqual) \
diff --git a/Engine/Plugins/Runtime/ReplicationSystemTestPlugin/Source/Private/Tests/Serialization/TestUintRangeNetSerializer.cpp b/Engine/Plugins/Runtime/ReplicationSystemTestPlugin/Source/Private/Tests/Serialization/TestUintRangeNetSerializer.cpp
--- a/Engine/Plugins/Runtime/ReplicationSystemTestPlugin/Source/Private/Tests/Serialization/TestUintRangeNetSerializer.cpp
+++ b/Engine/Plugins/Runtime/ReplicationSystemTestPlugin/Source/Private/Tests/Serialization/TestUintRangeNetSerializer.cpp
@@ -50,12 +50,12 @@ public: \
 \
 UE_NET_TEST_FIXTURE(TestClassName, HasTestValues) \
 { \
-	UE_NET_ASSERT_GT(ValueCount, SIZE_T(0)) << "No test values found"; \
+	UE_NET_ASSERT_GT_MSG(ValueCount, SIZE_T(0), "No test values found"); \
 } \
 \
 UE_NET_TEST_FIXTURE(TestClassName, HasTestConfigs) \
 { \
-	UE_NET_ASSERT_GT(ConfigCount, SIZE_T(0)) << "No configs found"; \
+	UE_NET_ASSERT_GT_MSG(ConfigCount, SIZE_T(0), "No configs found"); \
 } \
 \
 UE_NET_TEST_FIXTURE(TestClassName, TestIsEqual) \
@@ -181,7 +181,7 @@ void FTestUintRangeNetSerializer<SerializerConfig, SourceType>::TestValidate()
 			ExpectedResults[ValueIt] = (FMath::Clamp(Value, Config.LowerBound, Config.UpperBound) == Value);
 			if (ConfigIt == FullRangeConfigIndex)
 			{
-				UE_NET_ASSERT_TRUE(ExpectedResults[ValueIt]) << "Clamping of value with full bit precision resulted in value being clamped. This is unexpected. Make sure no undefined behavior is used in code.";
+				UE_NET_ASSERT_TRUE_MSG(ExpectedResults[ValueIt], "Clamping of value with full bit precision resulted in value being clamped. This is unexpected. Make sure no undefined behavior is used in code.");
 			}
 		}
 
diff --git a/Engine/Plugins/Runtime/ReplicationSystemTestPlugin/Source/Private/Tests/TestNetworkAutomationTestFramework.cpp b/Engine/Plugins/Runtime/ReplicationSystemTestPlugin/Source/Private/Tests/TestNetworkAutomationTestFramework.cpp
--- a/Engine/Plugins/Runtime/ReplicationSystemTestPlugin/Source/Private/Tests/TestNetworkAutomationTestFramework.cpp
+++ b/Engine/Plugins/Runtime/ReplicationSystemTestPlugin/Source/Private/Tests/TestNetworkAutomationTestFramework.cpp
@@ -97,8 +97,10 @@ UE_NET_TEST(NetworkAutomationTest, TestExpectEQ)
 	UE_NET_EXPECT_EQ(751216, 751216);
 	UE_NET_EXPECT_EQ(UINT_MAX, UINT_MAX);
 
+#if !EXPLICIT_TESTS_TARGET
 	SetSuppressWarningsFromSummary(true);
-	UE_NET_EXPECT_EQ(UINT_MAX, ~UINT_MAX) << ExpectedExpectFailureMessage;
+	UE_NET_EXPECT_EQ_MSG(UINT_MAX, ~UINT_MAX, ExpectedExpectFailureMessage);
+#endif
 }
 
 UE_NET_TEST(NetworkAutomationTest, TestExpectNE)
@@ -106,8 +108,10 @@ UE_NET_TEST(NetworkAutomationTest, TestExpectNE)
 	UE_NET_EXPECT_NE(true, false);
 	UE_NET_EXPECT_NE(false, true);
 
+#if !EXPLICIT_TESTS_TARGET
 	SetSuppressWarningsFromSummary(true);
-	UE_NET_EXPECT_NE(false, false) << ExpectedExpectFailureMessage;
+	UE_NET_EXPECT_NE_MSG(false, false, ExpectedExpectFailureMessage);
+#endif
 }
 
 UE_NET_TEST(NetworkAutomationTest, TestExpectLT)
@@ -116,8 +120,10 @@ UE_NET_TEST(NetworkAutomationTest, TestExpectLT)
 	UE_NET_EXPECT_LT(0x80000000U, ~0U);
 	UE_NET_EXPECT_LT(-1.0f, 2.0f);
 
+#if !EXPLICIT_TESTS_TARGET
 	SetSuppressWarningsFromSummary(true);
-	UE_NET_EXPECT_LT(-1.0f, -2.0f) << ExpectedExpectFailureMessage;
+	UE_NET_EXPECT_LT_MSG(-1.0f, -2.0f, ExpectedExpectFailureMessage);
+#endif
 }
 
 UE_NET_TEST(NetworkAutomationTest, TestExpectLE)
@@ -130,8 +136,10 @@ UE_NET_TEST(NetworkAutomationTest, TestExpectLE)
 	UE_NET_EXPECT_LE(~0U, ~0U);
 	UE_NET_EXPECT_LE(2.0f, 2.0f);
 
+#if !EXPLICIT_TESTS_TARGET
 	SetSuppressWarningsFromSummary(true);
-	UE_NET_EXPECT_LE(2.0f, 1.0f) << ExpectedExpectFailureMessage;
+	UE_NET_EXPECT_LE_MSG(2.0f, 1.0f, ExpectedExpectFailureMessage);
+#endif
 }
 
 UE_NET_TEST(NetworkAutomationTest, TestExpectGT)
@@ -140,8 +148,10 @@ UE_NET_TEST(NetworkAutomationTest, TestExpectGT)
 	UE_NET_EXPECT_GT(~0U, 0x80000000U);
 	UE_NET_EXPECT_GT(2.0f, -1.0f);
 
+#if !EXPLICIT_TESTS_TARGET
 	SetSuppressWarningsFromSummary(true);
-	UE_NET_EXPECT_GT(2.0f, 3.0f) << ExpectedExpectFailureMessage;
+	UE_NET_EXPECT_GT_MSG(2.0f, 3.0f, ExpectedExpectFailureMessage);
+#endif
 }
 
 UE_NET_TEST(NetworkAutomationTest, TestExpectGE)
@@ -154,24 +164,30 @@ UE_NET_TEST(NetworkAutomationTest, TestExpectGE)
 	UE_NET_EXPECT_GE(~0U, ~0U);
 	UE_NET_EXPECT_GE(2.0f, 2.0f);
 
+#if !EXPLICIT_TESTS_TARGET
 	SetSuppressWarningsFromSummary(true);
-	UE_NET_EXPECT_GE(2.0f, 3.0f) << ExpectedExpectFailureMessage;
+	UE_NET_EXPECT_GE_MSG(2.0f, 3.0f, ExpectedExpectFailureMessage);
+#endif
 }
 
 UE_NET_TEST(NetworkAutomationTest, TestExpectTrue)
 {
 	UE_NET_EXPECT_TRUE(true);
 
+#if !EXPLICIT_TESTS_TARGET
 	SetSuppressWarningsFromSummary(true);
-	UE_NET_EXPECT_TRUE(false) << ExpectedExpectFailureMessage;
+	UE_NET_EXPECT_TRUE_MSG(false, ExpectedExpectFailureMessage);
+#endif
 }
 
 UE_NET_TEST(NetworkAutomationTest, TestExpectFalse)
 {
 	UE_NET_EXPECT_FALSE(false);
 
+#if !EXPLICIT_TESTS_TARGET
 	SetSuppressWarningsFromSummary(true);
-	UE_NET_EXPECT_FALSE(true) << ExpectedExpectFailureMessage;
+	UE_NET_EXPECT_FALSE_MSG(true, ExpectedExpectFailureMessage);
+#endif
 }
 
 UE_NET_TEST(NetworkAutomationTest, TestLog)
diff --git a/Engine/Plugins/Runtime/ReplicationSystemTestPlugin/Source/ReplicationSystemTestPlugin.Build.cs b/Engine/Plugins/Runtime/ReplicationSystemTestPlugin/Source/ReplicationSystemTestPlugin.Build.cs
--- a/Engine/Plugins/Runtime/ReplicationSystemTestPlugin/Source/ReplicationSystemTestPlugin.Build.cs
+++ b/Engine/Plugins/Runtime/ReplicationSystemTestPlugin/Source/ReplicationSystemTestPlugin.Build.cs
@@ -44,12 +44,15 @@ namespace UnrealBuildTool.Rules
 				}
 				);
 
-			DynamicallyLoadedModuleNames.AddRange(
-				new string[]
-				{
-					// ... add any modules that your module loads dynamically here ...
-				}
-				);
+			if (Target.IsTestTarget)
+			{
+				PrivateDependencyModuleNames.AddRange(
+					new string[]
+					{
+						"LowLevelTestsRunner",
+					}
+					);
+			}
 
 			SetupIrisSupport(Target);
 		}
