Fix missing macro errors in ReplicationSystemTestPlugin, use a custom macro to indicate whether it's building with low level tests.

#rnx

diff --git a/Engine/Plugins/Runtime/ReplicationSystemTestPlugin/Source/Private/NetworkAutomationTest.h b/Engine/Plugins/Runtime/ReplicationSystemTestPlugin/Source/Private/NetworkAutomationTest.h
--- a/Engine/Plugins/Runtime/ReplicationSystemTestPlugin/Source/Private/NetworkAutomationTest.h
+++ b/Engine/Plugins/Runtime/ReplicationSystemTestPlugin/Source/Private/NetworkAutomationTest.h
@@ -5,7 +5,7 @@
 #include "CoreTypes.h"
 #include "Misc/AutomationTest.h"
 
-#if EXPLICIT_TESTS_TARGET
+#if UE_NET_WITH_LOW_LEVEL_TESTS
 #include "TestHarness.h"
 #endif
 
@@ -88,7 +88,7 @@ private:
 
 }
 
-#if EXPLICIT_TESTS_TARGET
+#if UE_NET_WITH_LOW_LEVEL_TESTS
 
 namespace UE::Net
 {
diff --git a/Engine/Plugins/Runtime/ReplicationSystemTestPlugin/Source/Private/NetworkAutomationTestMacros.h b/Engine/Plugins/Runtime/ReplicationSystemTestPlugin/Source/Private/NetworkAutomationTestMacros.h
--- a/Engine/Plugins/Runtime/ReplicationSystemTestPlugin/Source/Private/NetworkAutomationTestMacros.h
+++ b/Engine/Plugins/Runtime/ReplicationSystemTestPlugin/Source/Private/NetworkAutomationTestMacros.h
@@ -7,7 +7,7 @@
 #include "TestMessage.h"
 #include "Logging/LogVerbosity.h"
 
-#if EXPLICIT_TESTS_TARGET
+#if UE_NET_WITH_LOW_LEVEL_TESTS
 #include <catch2/catch_test_macros.hpp>
 #endif
 
@@ -112,7 +112,7 @@ Name(const T& Value, U ExpectedValue, const char* ValueText, const char* Expecte
 	} \
 }
 
-#if EXPLICIT_TESTS_TARGET
+#if UE_NET_WITH_LOW_LEVEL_TESTS
 
 #define UE_NET_TEST_MSG_INTERNAL(Message) \
 	UE::Net::FTestMessage PREPROCESSOR_JOIN(TestMessage, __LINE__); \
diff --git a/Engine/Plugins/Runtime/ReplicationSystemTestPlugin/Source/Private/Tests/TestNetworkAutomationTestFramework.cpp b/Engine/Plugins/Runtime/ReplicationSystemTestPlugin/Source/Private/Tests/TestNetworkAutomationTestFramework.cpp
--- a/Engine/Plugins/Runtime/ReplicationSystemTestPlugin/Source/Private/Tests/TestNetworkAutomationTestFramework.cpp
+++ b/Engine/Plugins/Runtime/ReplicationSystemTestPlugin/Source/Private/Tests/TestNetworkAutomationTestFramework.cpp
@@ -97,7 +97,7 @@ UE_NET_TEST(NetworkAutomationTest, TestExpectEQ)
 	UE_NET_EXPECT_EQ(751216, 751216);
 	UE_NET_EXPECT_EQ(UINT_MAX, UINT_MAX);
 
-#if !EXPLICIT_TESTS_TARGET
+#if !UE_NET_WITH_LOW_LEVEL_TESTS
 	SetSuppressWarningsFromSummary(true);
 	UE_NET_EXPECT_EQ_MSG(UINT_MAX, ~UINT_MAX, ExpectedExpectFailureMessage);
 #endif
@@ -108,7 +108,7 @@ UE_NET_TEST(NetworkAutomationTest, TestExpectNE)
 	UE_NET_EXPECT_NE(true, false);
 	UE_NET_EXPECT_NE(false, true);
 
-#if !EXPLICIT_TESTS_TARGET
+#if !UE_NET_WITH_LOW_LEVEL_TESTS
 	SetSuppressWarningsFromSummary(true);
 	UE_NET_EXPECT_NE_MSG(false, false, ExpectedExpectFailureMessage);
 #endif
@@ -120,7 +120,7 @@ UE_NET_TEST(NetworkAutomationTest, TestExpectLT)
 	UE_NET_EXPECT_LT(0x80000000U, ~0U);
 	UE_NET_EXPECT_LT(-1.0f, 2.0f);
 
-#if !EXPLICIT_TESTS_TARGET
+#if !UE_NET_WITH_LOW_LEVEL_TESTS
 	SetSuppressWarningsFromSummary(true);
 	UE_NET_EXPECT_LT_MSG(-1.0f, -2.0f, ExpectedExpectFailureMessage);
 #endif
@@ -136,7 +136,7 @@ UE_NET_TEST(NetworkAutomationTest, TestExpectLE)
 	UE_NET_EXPECT_LE(~0U, ~0U);
 	UE_NET_EXPECT_LE(2.0f, 2.0f);
 
-#if !EXPLICIT_TESTS_TARGET
+#if !UE_NET_WITH_LOW_LEVEL_TESTS
 	SetSuppressWarningsFromSummary(true);
 	UE_NET_EXPECT_LE_MSG(2.0f, 1.0f, ExpectedExpectFailureMessage);
 #endif
@@ -148,7 +148,7 @@ UE_NET_TEST(NetworkAutomationTest, TestExpectGT)
 	UE_NET_EXPECT_GT(~0U, 0x80000000U);
 	UE_NET_EXPECT_GT(2.0f, -1.0f);
 
-#if !EXPLICIT_TESTS_TARGET
+#if !UE_NET_WITH_LOW_LEVEL_TESTS
 	SetSuppressWarningsFromSummary(true);
 	UE_NET_EXPECT_GT_MSG(2.0f, 3.0f, ExpectedExpectFailureMessage);
 #endif
@@ -164,7 +164,7 @@ UE_NET_TEST(NetworkAutomationTest, TestExpectGE)
 	UE_NET_EXPECT_GE(~0U, ~0U);
 	UE_NET_EXPECT_GE(2.0f, 2.0f);
 
-#if !EXPLICIT_TESTS_TARGET
+#if !UE_NET_WITH_LOW_LEVEL_TESTS
 	SetSuppressWarningsFromSummary(true);
 	UE_NET_EXPECT_GE_MSG(2.0f, 3.0f, ExpectedExpectFailureMessage);
 #endif
@@ -174,7 +174,7 @@ UE_NET_TEST(NetworkAutomationTest, TestExpectTrue)
 {
 	UE_NET_EXPECT_TRUE(true);
 
-#if !EXPLICIT_TESTS_TARGET
+#if !UE_NET_WITH_LOW_LEVEL_TESTS
 	SetSuppressWarningsFromSummary(true);
 	UE_NET_EXPECT_TRUE_MSG(false, ExpectedExpectFailureMessage);
 #endif
@@ -184,7 +184,7 @@ UE_NET_TEST(NetworkAutomationTest, TestExpectFalse)
 {
 	UE_NET_EXPECT_FALSE(false);
 
-#if !EXPLICIT_TESTS_TARGET
+#if !UE_NET_WITH_LOW_LEVEL_TESTS
 	SetSuppressWarningsFromSummary(true);
 	UE_NET_EXPECT_FALSE_MSG(true, ExpectedExpectFailureMessage);
 #endif
diff --git a/Engine/Plugins/Runtime/ReplicationSystemTestPlugin/Source/ReplicationSystemTestPlugin.Build.cs b/Engine/Plugins/Runtime/ReplicationSystemTestPlugin/Source/ReplicationSystemTestPlugin.Build.cs
--- a/Engine/Plugins/Runtime/ReplicationSystemTestPlugin/Source/ReplicationSystemTestPlugin.Build.cs
+++ b/Engine/Plugins/Runtime/ReplicationSystemTestPlugin/Source/ReplicationSystemTestPlugin.Build.cs
@@ -1,5 +1,6 @@
 // Copyright Epic Games, Inc. All Rights Reserved.
 
+using System;
 using System.IO;
 
 namespace UnrealBuildTool.Rules
@@ -54,6 +55,8 @@ namespace UnrealBuildTool.Rules
 					);
 			}
 
+			PrivateDefinitions.Add(String.Format("UE_NET_WITH_LOW_LEVEL_TESTS={0}", Target.ExplicitTestsTarget ? "1" : "0"));
+
 			SetupIrisSupport(Target);
 		}
 	}
