Added ELifetimeCondition COND_Dynamic support. Added unit tests for Iris and functional tests for replication system agnostic testing.

COND_Dynamic allows for a property to decide its actual replication condition after instantiation of the actor/component rather than be the same for all instances. The condition can be changed using the macros DOREPDYNAMICCONDITION_INITCONDITION_FAST, when used in the object's GetReplicatedCustomConditionState implementation, or DOREPDYNAMICCONDITION_SETCONDITION_FAST in any other place where it's appropriate to change the condition. If the condition is never modified, thus remaining COND_Dynamic, the property will always replicate when changed. It's possible to use the DOREPCUSTOMCONDITION_ macros too on properties with this condition.

diff --git a/Engine/Plugins/Runtime/ReplicationSystemTestPlugin/Source/Private/Tests/ReplicationSystem/Conditionals/TestDynamicCondition.cpp b/Engine/Plugins/Runtime/ReplicationSystemTestPlugin/Source/Private/Tests/ReplicationSystem/Conditionals/TestDynamicCondition.cpp
new file mode 100644
--- /dev/null
+++ b/Engine/Plugins/Runtime/ReplicationSystemTestPlugin/Source/Private/Tests/ReplicationSystem/Conditionals/TestDynamicCondition.cpp
@@ -0,0 +1,506 @@
+// Copyright Epic Games, Inc. All Rights Reserved.
+
+#include "NetworkAutomationTest.h"
+#include "NetworkAutomationTestMacros.h"
+#include "Iris/ReplicationSystem/ReplicationSystem.h"
+#include "Iris/ReplicationSystem/Conditionals/ReplicationCondition.h"
+#include "Tests/ReplicationSystem/ReplicatedTestObject.h"
+#include "Tests/ReplicationSystem/ReplicationSystemServerClientTestFixture.h"
+
+namespace UE::Net::Private
+{
+
+class FTestDynamicConditionFixture : public FReplicationSystemServerClientTestFixture
+{
+protected:
+};
+
+
+UE_NET_TEST_FIXTURE(FTestDynamicConditionFixture, DynamicConditionIsReplicatedByDefault)
+{
+	// Add client
+	FReplicationSystemTestClient* Client = CreateClient();
+
+	// Spawn object on server
+	UTestReplicatedIrisObjectWithDynamicCondition* ServerObject = Server->CreateObject<UTestReplicatedIrisObjectWithDynamicCondition>();
+
+	// Set value
+	ServerObject->DynamicConditionInt ^= 13;
+
+	// Send and deliver packet
+	Server->PreSendUpdate();
+	Server->SendAndDeliverTo(Client, DeliverPacket);
+	Server->PostSendUpdate();
+	
+	// Object should have been created on the client
+	const UTestReplicatedIrisObjectWithDynamicCondition* ClientObject = Cast<UTestReplicatedIrisObjectWithDynamicCondition>(Client->GetReplicationBridge()->GetReplicatedObject(ServerObject->NetRefHandle));
+	UE_NET_ASSERT_NE(ClientObject, nullptr);
+
+	// Check that the dynamic condition member has been modified
+	UE_NET_ASSERT_EQ(ClientObject->DynamicConditionInt, ServerObject->DynamicConditionInt);
+}
+
+UE_NET_TEST_FIXTURE(FTestDynamicConditionFixture, DynamicConditionSupports_COND_Never)
+{
+	// Add client
+	FReplicationSystemTestClient* Client = CreateClient();
+
+	// Spawn object on server
+	UTestReplicatedIrisObjectWithDynamicCondition* ServerObject = Server->CreateObject<UTestReplicatedIrisObjectWithDynamicCondition>();
+
+	// Set value
+	const int32 OriginalValue = ServerObject->DynamicConditionInt;
+	ServerObject->DynamicConditionInt ^= 13;
+
+	// Set condition to COND_Never
+	ServerObject->SetDynamicCondition(COND_Never);
+
+	// Send and deliver packet
+	Server->PreSendUpdate();
+	Server->SendAndDeliverTo(Client, DeliverPacket);
+	Server->PostSendUpdate();
+	
+	// Object should have been created on the client
+	const UTestReplicatedIrisObjectWithDynamicCondition* ClientObject = Cast<UTestReplicatedIrisObjectWithDynamicCondition>(Client->GetReplicationBridge()->GetReplicatedObject(ServerObject->NetRefHandle));
+	UE_NET_ASSERT_NE(ClientObject, nullptr);
+
+	// Check that the dynamic condition member is unmodified
+	UE_NET_ASSERT_EQ(ClientObject->DynamicConditionInt, OriginalValue);
+}
+
+UE_NET_TEST_FIXTURE(FTestDynamicConditionFixture, DynamicConditionSupports_COND_AutonomousOnly)
+{
+	// Add clients
+	FReplicationSystemTestClient* SomeClients[] = { CreateClient(), CreateClient() };
+	constexpr SIZE_T ClientIndex = 0;
+	constexpr SIZE_T AutonomousClientIndex = 1;
+
+	UTestReplicatedIrisObjectWithDynamicCondition* ServerObject = Server->CreateObject<UTestReplicatedIrisObjectWithDynamicCondition>();
+
+	// Set value
+	const int32 OriginalValue = ServerObject->DynamicConditionInt;
+	ServerObject->DynamicConditionInt ^= 13;
+
+	// Set condition to COND_AutonomousOnly
+	ServerObject->SetDynamicCondition(COND_AutonomousOnly);
+
+	// Set other connection to be autonomous.
+	Server->ReplicationSystem->SetReplicationConditionConnectionFilter(ServerObject->NetRefHandle, EReplicationCondition::RoleAutonomous, SomeClients[AutonomousClientIndex]->ConnectionIdOnServer, true);
+
+	// Send and deliver packets
+	for (auto Client : SomeClients)
+	{
+		Server->PreSendUpdate();
+		Server->SendAndDeliverTo(Client, DeliverPacket);
+		Server->PostSendUpdate();
+	}
+	
+	// Check value on client
+	{
+		const UTestReplicatedIrisObjectWithDynamicCondition* ClientObject = Cast<UTestReplicatedIrisObjectWithDynamicCondition>(SomeClients[ClientIndex]->GetReplicationBridge()->GetReplicatedObject(ServerObject->NetRefHandle));
+		UE_NET_ASSERT_NE(ClientObject, nullptr);
+
+		// Check that the dynamic condition member is unmodified
+		UE_NET_ASSERT_EQ(ClientObject->DynamicConditionInt, OriginalValue);
+	}
+
+	// Check value on autonomous client
+	{
+		const UTestReplicatedIrisObjectWithDynamicCondition* ClientObject = Cast<UTestReplicatedIrisObjectWithDynamicCondition>(SomeClients[AutonomousClientIndex]->GetReplicationBridge()->GetReplicatedObject(ServerObject->NetRefHandle));
+		UE_NET_ASSERT_NE(ClientObject, nullptr);
+
+		// Check that the dynamic condition member is the same as on the server
+		UE_NET_ASSERT_EQ(ClientObject->DynamicConditionInt, ServerObject->DynamicConditionInt);
+	}
+}
+
+UE_NET_TEST_FIXTURE(FTestDynamicConditionFixture, DynamicConditionSupports_COND_SimulatedOnly)
+{
+	// Add clients
+	FReplicationSystemTestClient* SomeClients[] = { CreateClient(), CreateClient() };
+	constexpr SIZE_T SimulatedClientIndex = 0;
+	constexpr SIZE_T AutonomousClientIndex = 1;
+
+	UTestReplicatedIrisObjectWithDynamicCondition* ServerObject = Server->CreateObject<UTestReplicatedIrisObjectWithDynamicCondition>();
+
+	// Set value
+	const int32 OriginalValue = ServerObject->DynamicConditionInt;
+	ServerObject->DynamicConditionInt ^= 13;
+
+	// Set condition to COND_SimulatedOnly
+	ServerObject->SetDynamicCondition(COND_SimulatedOnly);
+
+	// Set one connection to be autonomous.
+	Server->ReplicationSystem->SetReplicationConditionConnectionFilter(ServerObject->NetRefHandle, EReplicationCondition::RoleAutonomous, SomeClients[AutonomousClientIndex]->ConnectionIdOnServer, true);
+
+	// Send and deliver packets
+	for (auto Client : SomeClients)
+	{
+		Server->PreSendUpdate();
+		Server->SendAndDeliverTo(Client, DeliverPacket);
+		Server->PostSendUpdate();
+	}
+	
+	// Check value on simulated client
+	{
+		const UTestReplicatedIrisObjectWithDynamicCondition* ClientObject = Cast<UTestReplicatedIrisObjectWithDynamicCondition>(SomeClients[SimulatedClientIndex]->GetReplicationBridge()->GetReplicatedObject(ServerObject->NetRefHandle));
+		UE_NET_ASSERT_NE(ClientObject, nullptr);
+
+		// Check that the dynamic condition member is the same as on the server
+		UE_NET_ASSERT_EQ(ClientObject->DynamicConditionInt, ServerObject->DynamicConditionInt);
+	}
+
+	// Check value on autonomous client
+	{
+		const UTestReplicatedIrisObjectWithDynamicCondition* ClientObject = Cast<UTestReplicatedIrisObjectWithDynamicCondition>(SomeClients[AutonomousClientIndex]->GetReplicationBridge()->GetReplicatedObject(ServerObject->NetRefHandle));
+		UE_NET_ASSERT_NE(ClientObject, nullptr);
+
+		// Check that the dynamic condition member is unmodified
+		UE_NET_ASSERT_EQ(ClientObject->DynamicConditionInt, OriginalValue);
+	}
+}
+
+UE_NET_TEST_FIXTURE(FTestDynamicConditionFixture, DynamicConditionSupports_COND_SkipOwner)
+{
+	// Add clients
+	FReplicationSystemTestClient* SomeClients[] = { CreateClient(), CreateClient() };
+	constexpr SIZE_T NonOwnerClientIndex = 0;
+	constexpr SIZE_T OwnerClientIndex = 1;
+
+	UTestReplicatedIrisObjectWithDynamicCondition* ServerObject = Server->CreateObject<UTestReplicatedIrisObjectWithDynamicCondition>();
+
+	// Set value
+	const int32 OriginalValue = ServerObject->DynamicConditionInt;
+	ServerObject->DynamicConditionInt ^= 13;
+
+	// Set condition to COND_SkipOwner
+	ServerObject->SetDynamicCondition(COND_SkipOwner);
+
+	// Set owning connection
+	Server->ReplicationSystem->SetOwningNetConnection(ServerObject->NetRefHandle, SomeClients[OwnerClientIndex]->ConnectionIdOnServer);
+
+	// Send and deliver packets
+	for (auto Client : SomeClients)
+	{
+		Server->PreSendUpdate();
+		Server->SendAndDeliverTo(Client, DeliverPacket);
+		Server->PostSendUpdate();
+	}
+	
+	// Check value on owning client
+	{
+		const UTestReplicatedIrisObjectWithDynamicCondition* ClientObject = Cast<UTestReplicatedIrisObjectWithDynamicCondition>(SomeClients[OwnerClientIndex]->GetReplicationBridge()->GetReplicatedObject(ServerObject->NetRefHandle));
+		UE_NET_ASSERT_NE(ClientObject, nullptr);
+
+		// Check that the dynamic condition member is unmodified
+		UE_NET_ASSERT_EQ(ClientObject->DynamicConditionInt, OriginalValue);
+	}
+
+	// Check value on non-owning client
+	{
+		const UTestReplicatedIrisObjectWithDynamicCondition* ClientObject = Cast<UTestReplicatedIrisObjectWithDynamicCondition>(SomeClients[NonOwnerClientIndex]->GetReplicationBridge()->GetReplicatedObject(ServerObject->NetRefHandle));
+		UE_NET_ASSERT_NE(ClientObject, nullptr);
+
+		// Check that the dynamic condition member is the same as on the server
+		UE_NET_ASSERT_EQ(ClientObject->DynamicConditionInt, ServerObject->DynamicConditionInt);
+	}
+}
+
+UE_NET_TEST_FIXTURE(FTestDynamicConditionFixture, DynamicConditionSupports_COND_OwnerOnly)
+{
+	// Add clients
+	FReplicationSystemTestClient* SomeClients[] = { CreateClient(), CreateClient() };
+	constexpr SIZE_T NonOwnerClientIndex = 0;
+	constexpr SIZE_T OwnerClientIndex = 1;
+
+	UTestReplicatedIrisObjectWithDynamicCondition* ServerObject = Server->CreateObject<UTestReplicatedIrisObjectWithDynamicCondition>();
+
+	// Set value
+	const int32 OriginalValue = ServerObject->DynamicConditionInt;
+	ServerObject->DynamicConditionInt ^= 13;
+
+	// Set condition to COND_OwnerOnly
+	ServerObject->SetDynamicCondition(COND_OwnerOnly);
+
+	// Set owning connection
+	Server->ReplicationSystem->SetOwningNetConnection(ServerObject->NetRefHandle, SomeClients[OwnerClientIndex]->ConnectionIdOnServer);
+
+	// Send and deliver packets
+	for (auto Client : SomeClients)
+	{
+		Server->PreSendUpdate();
+		Server->SendAndDeliverTo(Client, DeliverPacket);
+		Server->PostSendUpdate();
+	}
+	
+	// Check value on owning client
+	{
+		const UTestReplicatedIrisObjectWithDynamicCondition* ClientObject = Cast<UTestReplicatedIrisObjectWithDynamicCondition>(SomeClients[OwnerClientIndex]->GetReplicationBridge()->GetReplicatedObject(ServerObject->NetRefHandle));
+		UE_NET_ASSERT_NE(ClientObject, nullptr);
+
+		// Check that the dynamic condition member is the same as on the server
+		UE_NET_ASSERT_EQ(ClientObject->DynamicConditionInt, ServerObject->DynamicConditionInt);
+	}
+
+	// Check value on non-owning client
+	{
+		const UTestReplicatedIrisObjectWithDynamicCondition* ClientObject = Cast<UTestReplicatedIrisObjectWithDynamicCondition>(SomeClients[NonOwnerClientIndex]->GetReplicationBridge()->GetReplicatedObject(ServerObject->NetRefHandle));
+		UE_NET_ASSERT_NE(ClientObject, nullptr);
+
+		// Check that the dynamic condition member is unmodified
+		UE_NET_ASSERT_EQ(ClientObject->DynamicConditionInt, OriginalValue);
+	}
+}
+
+UE_NET_TEST_FIXTURE(FTestDynamicConditionFixture, DynamicConditionSupports_COND_InitialOnly)
+{
+	// Add clients
+	FReplicationSystemTestClient* Client = CreateClient();
+
+	UTestReplicatedIrisObjectWithDynamicCondition* ServerObject = Server->CreateObject<UTestReplicatedIrisObjectWithDynamicCondition>();
+
+	// Set value
+	ServerObject->DynamicConditionInt ^= 13;
+
+	// Set condition to COND_InitialOnly
+	ServerObject->SetDynamicCondition(COND_InitialOnly);
+
+	// Send and deliver packet
+	Server->PreSendUpdate();
+	Server->SendAndDeliverTo(Client, DeliverPacket);
+	Server->PostSendUpdate();
+	
+	const UTestReplicatedIrisObjectWithDynamicCondition* ClientObject = Cast<UTestReplicatedIrisObjectWithDynamicCondition>(Client->GetReplicationBridge()->GetReplicatedObject(ServerObject->NetRefHandle));
+	UE_NET_ASSERT_NE(ClientObject, nullptr);
+
+	// Check that the dynamic condition member is the same as on the server
+	UE_NET_ASSERT_EQ(ClientObject->DynamicConditionInt, ServerObject->DynamicConditionInt);
+
+	// Modify value
+	const int32 PrevValue = ServerObject->DynamicConditionInt;
+	ServerObject->DynamicConditionInt += 13;
+
+	// Send and deliver packet
+	Server->PreSendUpdate();
+	Server->SendAndDeliverTo(Client, DeliverPacket);
+	Server->PostSendUpdate();
+
+	// Check that the dynamic condition member has not been modified
+	UE_NET_ASSERT_EQ(ClientObject->DynamicConditionInt, PrevValue);
+}
+
+UE_NET_TEST_FIXTURE(FTestDynamicConditionFixture, DynamicConditionSupports_COND_InitialOrOwner)
+{
+	// Add clients
+	FReplicationSystemTestClient* SomeClients[] = { CreateClient(), CreateClient() };
+	constexpr SIZE_T NonOwnerClientIndex = 0;
+	constexpr SIZE_T OwnerClientIndex = 1;
+
+	UTestReplicatedIrisObjectWithDynamicCondition* ServerObject = Server->CreateObject<UTestReplicatedIrisObjectWithDynamicCondition>();
+
+	// Set value
+	ServerObject->DynamicConditionInt ^= 13;
+
+	// Set condition to COND_OwnerOnly
+	ServerObject->SetDynamicCondition(COND_InitialOrOwner);
+
+	// Set owning connection
+	Server->ReplicationSystem->SetOwningNetConnection(ServerObject->NetRefHandle, SomeClients[OwnerClientIndex]->ConnectionIdOnServer);
+
+	// Send and deliver packets
+	for (auto Client : SomeClients)
+	{
+		Server->PreSendUpdate();
+		Server->SendAndDeliverTo(Client, DeliverPacket);
+		Server->PostSendUpdate();
+	}
+	
+	// Check value on owning client
+	{
+		const UTestReplicatedIrisObjectWithDynamicCondition* ClientObject = Cast<UTestReplicatedIrisObjectWithDynamicCondition>(SomeClients[OwnerClientIndex]->GetReplicationBridge()->GetReplicatedObject(ServerObject->NetRefHandle));
+		UE_NET_ASSERT_NE(ClientObject, nullptr);
+
+		// Check that the dynamic condition member is the same as on the server
+		UE_NET_ASSERT_EQ(ClientObject->DynamicConditionInt, ServerObject->DynamicConditionInt);
+	}
+
+	// Check value on non-owning client
+	{
+		const UTestReplicatedIrisObjectWithDynamicCondition* ClientObject = Cast<UTestReplicatedIrisObjectWithDynamicCondition>(SomeClients[NonOwnerClientIndex]->GetReplicationBridge()->GetReplicatedObject(ServerObject->NetRefHandle));
+		UE_NET_ASSERT_NE(ClientObject, nullptr);
+
+		// Check that the dynamic condition member is the same as on the server
+		UE_NET_ASSERT_EQ(ClientObject->DynamicConditionInt, ServerObject->DynamicConditionInt);
+	}
+
+	// Modify value
+	const int32 PrevValue = ServerObject->DynamicConditionInt;
+	ServerObject->DynamicConditionInt += 13;
+
+	// Send and deliver packets
+	for (auto Client : SomeClients)
+	{
+		Server->PreSendUpdate();
+		Server->SendAndDeliverTo(Client, DeliverPacket);
+		Server->PostSendUpdate();
+	}
+
+	// Check value on owning client. It should've been updated.
+	{
+		const UTestReplicatedIrisObjectWithDynamicCondition* ClientObject = Cast<UTestReplicatedIrisObjectWithDynamicCondition>(SomeClients[OwnerClientIndex]->GetReplicationBridge()->GetReplicatedObject(ServerObject->NetRefHandle));
+		UE_NET_ASSERT_NE(ClientObject, nullptr);
+
+		// Check that the dynamic condition member is the same as on the server
+		UE_NET_ASSERT_EQ(ClientObject->DynamicConditionInt, ServerObject->DynamicConditionInt);
+	}
+
+	// Check value on non-owning client. It should remain the same as the previous value.
+	{
+		const UTestReplicatedIrisObjectWithDynamicCondition* ClientObject = Cast<UTestReplicatedIrisObjectWithDynamicCondition>(SomeClients[NonOwnerClientIndex]->GetReplicationBridge()->GetReplicatedObject(ServerObject->NetRefHandle));
+		UE_NET_ASSERT_NE(ClientObject, nullptr);
+
+		// Check that the dynamic condition member is unmodified.
+		UE_NET_ASSERT_EQ(ClientObject->DynamicConditionInt, PrevValue);
+	}
+}
+
+UE_NET_TEST_FIXTURE(FTestDynamicConditionFixture, DynamicConditionSupports_COND_Custom)
+{
+	// Add clients
+	FReplicationSystemTestClient* Client = CreateClient();
+
+	UTestReplicatedIrisObjectWithDynamicCondition* ServerObject = Server->CreateObject<UTestReplicatedIrisObjectWithDynamicCondition>();
+
+	// Set value
+	ServerObject->DynamicConditionInt ^= 13;
+
+	// Set condition to COND_Custom. Technically speaking this isn't required as all propertys with a real condition (not COND_Never/None) can be enabled/disabled at will.
+	ServerObject->SetDynamicCondition(COND_Custom);
+
+	// Send and deliver packet
+	Server->PreSendUpdate();
+	Server->SendAndDeliverTo(Client, DeliverPacket);
+	Server->PostSendUpdate();
+	
+	const UTestReplicatedIrisObjectWithDynamicCondition* ClientObject = Cast<UTestReplicatedIrisObjectWithDynamicCondition>(Client->GetReplicationBridge()->GetReplicatedObject(ServerObject->NetRefHandle));
+	UE_NET_ASSERT_NE(ClientObject, nullptr);
+
+	// Check that the dynamic condition member is the same as on the server
+	UE_NET_ASSERT_EQ(ClientObject->DynamicConditionInt, ServerObject->DynamicConditionInt);
+
+	// Disable the custom condition
+	ServerObject->SetDynamicConditionCustomCondition(false);
+
+	// Modify value
+	const int32 PrevValue = ServerObject->DynamicConditionInt;
+	ServerObject->DynamicConditionInt += 13;
+
+	// Send and deliver packet
+	Server->PreSendUpdate();
+	Server->SendAndDeliverTo(Client, DeliverPacket);
+	Server->PostSendUpdate();
+
+	// Check that the dynamic condition member has not been modified
+	UE_NET_ASSERT_EQ(ClientObject->DynamicConditionInt, PrevValue);
+}
+
+UE_NET_TEST_FIXTURE(FTestDynamicConditionFixture, DynamicConditionSupports_COND_SimulatedOrPhysics)
+{
+	// Add client
+	FReplicationSystemTestClient* Client = CreateClient();
+
+	UTestReplicatedIrisObjectWithDynamicCondition* ServerObject = Server->CreateObject<UTestReplicatedIrisObjectWithDynamicCondition>();
+
+	// Set the connection to be autonomous so we know that only the Physics condition is at play
+	Server->ReplicationSystem->SetReplicationConditionConnectionFilter(ServerObject->NetRefHandle, EReplicationCondition::RoleAutonomous, Client->ConnectionIdOnServer, true);
+
+	// Enable ReplicatePhysics condition
+	Server->ReplicationSystem->SetReplicationCondition(ServerObject->NetRefHandle, EReplicationCondition::ReplicatePhysics, true);
+
+	// Set value
+	ServerObject->DynamicConditionInt ^= 13;
+
+	// Set condition to COND_SimulatedOrPhysics. 
+	ServerObject->SetDynamicCondition(COND_SimulatedOrPhysics);
+
+	// Send and deliver packet
+	Server->PreSendUpdate();
+	Server->SendAndDeliverTo(Client, DeliverPacket);
+	Server->PostSendUpdate();
+	
+	const UTestReplicatedIrisObjectWithDynamicCondition* ClientObject = Cast<UTestReplicatedIrisObjectWithDynamicCondition>(Client->GetReplicationBridge()->GetReplicatedObject(ServerObject->NetRefHandle));
+	UE_NET_ASSERT_NE(ClientObject, nullptr);
+
+	// Check that the dynamic condition member is the same as on the server
+	UE_NET_ASSERT_EQ(ClientObject->DynamicConditionInt, ServerObject->DynamicConditionInt);
+
+	// Disable ReplicatePhysics condition
+	Server->ReplicationSystem->SetReplicationCondition(ServerObject->NetRefHandle, EReplicationCondition::ReplicatePhysics, false);
+
+	// Modify value
+	const int32 PrevValue = ServerObject->DynamicConditionInt;
+	ServerObject->DynamicConditionInt += 13;
+
+	// Send and deliver packet
+	Server->PreSendUpdate();
+	Server->SendAndDeliverTo(Client, DeliverPacket);
+	Server->PostSendUpdate();
+
+	// Check that the dynamic condition member is unmodified
+	UE_NET_ASSERT_EQ(ClientObject->DynamicConditionInt, PrevValue);
+
+	// Set the connection to no longer be autonomous, thus simulated
+	Server->ReplicationSystem->SetReplicationConditionConnectionFilter(ServerObject->NetRefHandle, EReplicationCondition::RoleAutonomous, Client->ConnectionIdOnServer, false);
+
+	// Modify value
+	ServerObject->DynamicConditionInt += 17;
+
+	// Send and deliver packet
+	Server->PreSendUpdate();
+	Server->SendAndDeliverTo(Client, DeliverPacket);
+	Server->PostSendUpdate();
+
+	// Check that the dynamic condition member is the same as on the server
+	UE_NET_ASSERT_EQ(ClientObject->DynamicConditionInt, ServerObject->DynamicConditionInt);
+}
+
+UE_NET_TEST_FIXTURE(FTestDynamicConditionFixture, DynamicConditionSupportsGoingFromNotReplicatedToReplicated)
+{
+	// Add client
+	FReplicationSystemTestClient* Client = CreateClient();
+
+	// Spawn object on server
+	UTestReplicatedIrisObjectWithDynamicCondition* ServerObject = Server->CreateObject<UTestReplicatedIrisObjectWithDynamicCondition>();
+
+	// Set value
+	const int32 OriginalValue = ServerObject->DynamicConditionInt;
+	ServerObject->DynamicConditionInt ^= 13;
+
+	// Set condition to COND_Never
+	ServerObject->SetDynamicCondition(COND_Never);
+
+	// Send and deliver packet
+	Server->PreSendUpdate();
+	Server->SendAndDeliverTo(Client, DeliverPacket);
+	Server->PostSendUpdate();
+	
+	// Object should have been created on the client
+	const UTestReplicatedIrisObjectWithDynamicCondition* ClientObject = Cast<UTestReplicatedIrisObjectWithDynamicCondition>(Client->GetReplicationBridge()->GetReplicatedObject(ServerObject->NetRefHandle));
+	UE_NET_ASSERT_NE(ClientObject, nullptr);
+
+	// Check that the dynamic condition member is unmodified
+	UE_NET_ASSERT_EQ(ClientObject->DynamicConditionInt, OriginalValue);
+
+	// Set condition to COND_None
+	ServerObject->SetDynamicCondition(COND_None);
+
+	// Send and deliver packet
+	Server->PreSendUpdate();
+	Server->SendAndDeliverTo(Client, DeliverPacket);
+	Server->PostSendUpdate();
+
+	// Check that the dynamic condition member is the same as on the server
+	UE_NET_ASSERT_EQ(ClientObject->DynamicConditionInt, ServerObject->DynamicConditionInt);
+}
+
+}
diff --git a/Engine/Plugins/Runtime/ReplicationSystemTestPlugin/Source/Private/Tests/ReplicationSystem/ReplicatedTestObject.cpp b/Engine/Plugins/Runtime/ReplicationSystemTestPlugin/Source/Private/Tests/ReplicationSystem/ReplicatedTestObject.cpp
--- a/Engine/Plugins/Runtime/ReplicationSystemTestPlugin/Source/Private/Tests/ReplicationSystem/ReplicatedTestObject.cpp
+++ b/Engine/Plugins/Runtime/ReplicationSystemTestPlugin/Source/Private/Tests/ReplicationSystem/ReplicatedTestObject.cpp
@@ -301,15 +301,23 @@ void FTestReplicatedIrisComponent::ApplyReplicationState(const FFakeGeneratedRep
 }
 
 //////////////////////////////////////////////////////////////////////////
-// Implementation for UTestReplicationSystem_TestClass
+// Implementation for UTestReplicatedIrisObject
 //////////////////////////////////////////////////////////////////////////
-void UTestReplicatedIrisObject::GetLifetimeReplicatedProps( TArray< class FLifetimeProperty > & OutLifetimeProps ) const
+UTestReplicatedIrisObject::UTestReplicatedIrisObject()
+: UReplicatedTestObject()
 {
 }
 
-UTestReplicatedIrisObject::UTestReplicatedIrisObject()
-: UReplicatedTestObject()
+void UTestReplicatedIrisObject::GetLifetimeReplicatedProps(TArray<FLifetimeProperty>& OutLifetimeProps) const
 {
+	FDoRepLifetimeParams Params;
+	Params.bIsPushBased = false;
+
+	Params.Condition = COND_None;
+	DOREPLIFETIME_WITH_PARAMS_FAST(ThisClass, IntA, Params);
+	DOREPLIFETIME_WITH_PARAMS_FAST(ThisClass, IntB, Params);
+	DOREPLIFETIME_WITH_PARAMS_FAST(ThisClass, IntC, Params);
+	DOREPLIFETIME_WITH_PARAMS_FAST(ThisClass, StructD, Params);
 }
 
 void UTestReplicatedIrisObject::AddComponents(const UTestReplicatedIrisObject::FComponents& InComponents)
@@ -414,7 +422,7 @@ void UTestReplicatedIrisObject::RegisterReplicationFragments(UE::Net::FFragmentR
 }
 
 //////////////////////////////////////////////////////////////////////////
-// Implementation for UTestReplicationSystem_TestClass
+// Implementation for UTestReplicatedIrisObjectWithObjectReference
 //////////////////////////////////////////////////////////////////////////
 
 void UTestReplicatedIrisObjectWithObjectReference::GetLifetimeReplicatedProps( TArray< class FLifetimeProperty > & OutLifetimeProps ) const
@@ -466,6 +474,44 @@ void UReplicatedSubObjectOrderObject::RegisterReplicationFragments(UE::Net::FFra
 	}
 }
 
+//////////////////////////////////////////////////////////////////////////
+// Implementation for UTestReplicatedIrisObjectWithDynamicCondition
+//////////////////////////////////////////////////////////////////////////
+UTestReplicatedIrisObjectWithDynamicCondition::UTestReplicatedIrisObjectWithDynamicCondition()
+: UReplicatedTestObject()
+{
+}
+
+void UTestReplicatedIrisObjectWithDynamicCondition::RegisterReplicationFragments(UE::Net::FFragmentRegistrationContext& Context, UE::Net::EFragmentRegistrationFlags RegistrationFlags)
+{
+	Super::RegisterReplicationFragments(Context, RegistrationFlags);
+
+	// Base object owns the fragment in this case
+	{
+		this->ReplicationFragments.Reset();
+		UE::Net::FReplicationFragmentUtil::CreateAndRegisterFragmentsForObject(this, Context, RegistrationFlags, &this->ReplicationFragments);
+	}
+}
+
+void UTestReplicatedIrisObjectWithDynamicCondition::GetLifetimeReplicatedProps(TArray<FLifetimeProperty>& OutLifetimeProps) const
+{
+	FDoRepLifetimeParams Params;
+	Params.bIsPushBased = false;
+
+	Params.Condition = COND_Dynamic;
+	DOREPLIFETIME_WITH_PARAMS_FAST(ThisClass, DynamicConditionInt, Params);
+}
+
+void UTestReplicatedIrisObjectWithDynamicCondition::SetDynamicCondition(ELifetimeCondition Condition)
+{
+	DOREPDYNAMICCONDITION_SETCONDITION_FAST(ThisClass, DynamicConditionInt, Condition);
+}
+
+void UTestReplicatedIrisObjectWithDynamicCondition::SetDynamicConditionCustomCondition(bool bActive)
+{
+	DOREPCUSTOMCONDITION_SETACTIVE_FAST(ThisClass, DynamicConditionInt, bActive);	
+}
+
 //////////////////////////////////////////////////////////////////////////
 // THIS SECTION WILL BE GENERATED FROM UHT
 //////////////////////////////////////////////////////////////////////////
diff --git a/Engine/Plugins/Runtime/ReplicationSystemTestPlugin/Source/Private/Tests/ReplicationSystem/ReplicatedTestObject.h b/Engine/Plugins/Runtime/ReplicationSystemTestPlugin/Source/Private/Tests/ReplicationSystem/ReplicatedTestObject.h
--- a/Engine/Plugins/Runtime/ReplicationSystemTestPlugin/Source/Private/Tests/ReplicationSystem/ReplicatedTestObject.h
+++ b/Engine/Plugins/Runtime/ReplicationSystemTestPlugin/Source/Private/Tests/ReplicationSystem/ReplicatedTestObject.h
@@ -347,6 +347,36 @@ public:
 	TArray<UE::Net::FReplicationFragment*> ReplicationFragments;
 };
 
+/**
+ *  A test class for Replication that itself uses Property based replication but also has "components" that uses a mix of property based replication and native ReplicationStates
+ */
+UCLASS()
+class UTestReplicatedIrisObjectWithDynamicCondition : public UReplicatedTestObject
+{
+	GENERATED_BODY()
+
+public:
+	UTestReplicatedIrisObjectWithDynamicCondition();
+
+	// Sets the condition for the DynamicConditionInt member.
+	void SetDynamicCondition(ELifetimeCondition Condition);
+
+	// Sets custom condition on DynamicConditionInt member.
+	void SetDynamicConditionCustomCondition(bool bActive);
+
+public:
+	UPROPERTY(Transient, Replicated)
+	int32 DynamicConditionInt;
+
+private:
+	// Network interface must be part of base.
+	virtual void RegisterReplicationFragments(UE::Net::FFragmentRegistrationContext& Fragments, UE::Net::EFragmentRegistrationFlags RegistrationFlags) override;
+
+public:
+	// Network data only for test
+	TArray<UE::Net::FReplicationFragment*> ReplicationFragments;
+};
+
 UCLASS()
 class UReplicatedSubObjectOrderObject : public UReplicatedTestObject
 {
diff --git a/Engine/Source/Runtime/CoreUObject/Public/UObject/CoreNet.h b/Engine/Source/Runtime/CoreUObject/Public/UObject/CoreNet.h
--- a/Engine/Source/Runtime/CoreUObject/Public/UObject/CoreNet.h
+++ b/Engine/Source/Runtime/CoreUObject/Public/UObject/CoreNet.h
@@ -565,8 +565,9 @@ class FCustomPropertyConditionState
 public:
 	FCustomPropertyConditionState() = delete;
 	FCustomPropertyConditionState(int32 NumProperties)
+	: CurrentState(true /* initial value of all bits */, NumProperties)
+	, DynamicConditionChangeCounter(0)
 	{
-		CurrentState.Init(true, NumProperties);
 	}
 
 	void SetActiveState(const uint16 RepIndex, const bool bIsActive)
@@ -579,11 +580,35 @@ public:
 		return CurrentState[RepIndex];
 	}
 
+	void SetDynamicCondition(const uint16 RepIndex, const ELifetimeCondition Condition)
+	{
+		static_assert(static_cast<__underlying_type(ELifetimeCondition)>(ELifetimeCondition::COND_Max - 1) <= int16(32767), "Unable to use int16 for ELifetimeCondition values");
+
+		++DynamicConditionChangeCounter;
+		DynamicConditions.Emplace(RepIndex, static_cast<int16>(Condition));
+	}
+
+	ELifetimeCondition GetDynamicCondition(const uint16 RepIndex) const
+	{
+		if (const int16* Condition = DynamicConditions.Find(RepIndex))
+		{
+			return static_cast<const ELifetimeCondition>(*Condition);
+		}
+
+		return COND_Dynamic;
+	}
+
+
 	int32 GetNumProperties() const
 	{
 		return CurrentState.Num();
 	}
 
+	uint32 GetDynamicConditionChangeCounter() const
+	{
+		return DynamicConditionChangeCounter;
+	}
+
 	void CountBytes(FArchive& Ar) const
 	{
 		CurrentState.CountBytes(Ar);
@@ -591,6 +616,9 @@ public:
 
 private:
 	TBitArray<> CurrentState;
+	// Storing int16 instead of int-sized ELifetimeCondition to save some memory.
+	TMap<uint16, int16> DynamicConditions;
+	uint32 DynamicConditionChangeCounter;
 };
 
 /**
diff --git a/Engine/Source/Runtime/CoreUObject/Public/UObject/CoreNetTypes.h b/Engine/Source/Runtime/CoreUObject/Public/UObject/CoreNetTypes.h
--- a/Engine/Source/Runtime/CoreUObject/Public/UObject/CoreNetTypes.h
+++ b/Engine/Source/Runtime/CoreUObject/Public/UObject/CoreNetTypes.h
@@ -32,6 +32,7 @@ enum ELifetimeCondition : int
 	COND_SimulatedOnlyNoReplay = 11			UMETA(DisplayName = "Simulated Only No Replay"),		// This property will send to actors only, but not to replay connections
 	COND_SimulatedOrPhysicsNoReplay = 12	UMETA(DisplayName = "Simulated Or Physics No Replay"),	// This property will send to simulated Or bRepPhysics actors, but not to replay connections
 	COND_SkipReplay = 13					UMETA(DisplayName = "Skip Replay"),						// This property will not send to the replay connection
+	COND_Dynamic = 14						UMETA(Hidden),											// This property wants to override the condition at runtime. Defaults to always replicate until you override it to a new condition.
 	COND_Never = 15							UMETA(Hidden),											// This property will never be replicated
 	COND_NetGroup = 16						UMETA(Hidden),											// This subobject will replicate to connections that are part of the same group the subobject is registered to. Not usable on properties.
 	COND_Max = 17							UMETA(Hidden)
diff --git a/Engine/Source/Runtime/Engine/Classes/Engine/EngineTypes.h b/Engine/Source/Runtime/Engine/Classes/Engine/EngineTypes.h
--- a/Engine/Source/Runtime/Engine/Classes/Engine/EngineTypes.h
+++ b/Engine/Source/Runtime/Engine/Classes/Engine/EngineTypes.h
@@ -3086,9 +3086,14 @@ struct FReplicationFlags
 			uint32 bUseCustomSubobjectReplication : 1;
 			/** True if this actor is replicating on a replay connection on a game client. */
 			uint32 bClientReplay : 1;
+			/** Padding bits to align to a 16-bit boundary. NOTE: Need to be adjusted if adding/removing bitfields. */
+			uint32 Padding0 : 4;
+
+			/** COND_Dynamic change counter in order to force rebuilding conditionals when appropriate. */
+			uint32 CondDynamicChangeCounter : 16;
 		};
 
-		uint32	Value;
+		uint32 Value;
 	};
 	FReplicationFlags()
 	{
diff --git a/Engine/Source/Runtime/Engine/Private/DataChannel.cpp b/Engine/Source/Runtime/Engine/Private/DataChannel.cpp
--- a/Engine/Source/Runtime/Engine/Private/DataChannel.cpp
+++ b/Engine/Source/Runtime/Engine/Private/DataChannel.cpp
@@ -3465,6 +3465,17 @@ int64 UActorChannel::ReplicateActor()
 	RepFlags.bReplay		= bReplay;
 	RepFlags.bClientReplay	= ActorWorld->IsRecordingClientReplay();
 	RepFlags.bForceInitialDirty = Connection->IsForceInitialDirty();
+	if (EnumHasAnyFlags(ActorReplicator->RepLayout->GetFlags(), ERepLayoutFlags::HasDynamicConditionProperties))
+	{
+		if (const FRepState* RepState = ActorReplicator->RepState.Get())
+		{
+			const FSendingRepState* SendingRepState = RepState->GetSendingRepState();
+			if (const FRepChangedPropertyTracker* PropertyTracker = SendingRepState ? SendingRepState->RepChangedPropertyTracker.Get() : nullptr)
+			{
+				RepFlags.CondDynamicChangeCounter = PropertyTracker->GetDynamicConditionChangeCounter();
+			}
+		}
+	}
 
 	UE_LOG(LogNetTraffic, Log, TEXT("Replicate %s, bNetInitial: %d, bNetOwner: %d"), *Actor->GetName(), RepFlags.bNetInitial, RepFlags.bNetOwner);
 
diff --git a/Engine/Source/Runtime/Engine/Private/DataReplication.cpp b/Engine/Source/Runtime/Engine/Private/DataReplication.cpp
--- a/Engine/Source/Runtime/Engine/Private/DataReplication.cpp
+++ b/Engine/Source/Runtime/Engine/Private/DataReplication.cpp
@@ -1614,7 +1614,16 @@ void FObjectReplicator::ReplicateCustomDeltaProperties( FNetBitWriter & Bunch, F
 	// Replicate those properties.
 	for (uint16 CustomDeltaProperty = 0; CustomDeltaProperty < NumLifetimeCustomDeltaProperties; ++CustomDeltaProperty)
 	{
-		const ELifetimeCondition RepCondition = FNetSerializeCB::GetLifetimeCustomDeltaPropertyCondition(LocalRepLayout, CustomDeltaProperty);
+		ELifetimeCondition RepCondition = FNetSerializeCB::GetLifetimeCustomDeltaPropertyCondition(LocalRepLayout, CustomDeltaProperty);
+		FProperty* Property = FNetSerializeCB::GetLifetimeCustomDeltaProperty(LocalRepLayout, CustomDeltaProperty);
+
+		if (RepCondition == COND_Dynamic)
+		{
+			if (const FRepChangedPropertyTracker* RepChangedPropertyTracker = SendingRepState->RepChangedPropertyTracker.Get())
+			{
+				RepCondition = RepChangedPropertyTracker->GetDynamicCondition(Property->RepIndex);
+			}
+		}
 
 		check(RepCondition >= 0 && RepCondition < COND_Max);
 
@@ -1625,8 +1634,6 @@ void FObjectReplicator::ReplicateCustomDeltaProperties( FNetBitWriter & Bunch, F
 			continue;
 		}
 
-		FProperty* Property = FNetSerializeCB::GetLifetimeCustomDeltaProperty(LocalRepLayout, CustomDeltaProperty);
-
 		// If this is a dynamic array, we do the delta here
 		TSharedPtr<INetDeltaBaseState> NewState;
 
diff --git a/Engine/Source/Runtime/Engine/Private/RepLayout.cpp b/Engine/Source/Runtime/Engine/Private/RepLayout.cpp
--- a/Engine/Source/Runtime/Engine/Private/RepLayout.cpp
+++ b/Engine/Source/Runtime/Engine/Private/RepLayout.cpp
@@ -1395,7 +1395,19 @@ static bool CompareParentProperty(
 	// Further, this will keep the last active state of the property in the shadow buffer,
 	// meaning the next time the property becomes active it will be sent to all connections.
 	const bool bIsActive = SharedParams.bForceCustomPropsActive || !SharedParams.RepChangedPropertyTracker || SharedParams.RepChangedPropertyTracker->IsParentActive(ParentIndex);
-	const bool bShouldSkip = !bIsLifetime || !bIsActive || (Parent.Condition == COND_InitialOnly && !SharedParams.bIsInitial);
+	bool bShouldSkip = !bIsLifetime || !bIsActive;
+	if (!bShouldSkip)
+	{
+		ELifetimeCondition Condition = Parent.Condition;
+		if (Condition == COND_Dynamic)
+		{
+			Condition = SharedParams.RepChangedPropertyTracker ? SharedParams.RepChangedPropertyTracker->GetDynamicCondition(ParentIndex) : COND_Dynamic;
+			bShouldSkip = Condition == COND_Never;
+		}
+
+		// Skip initial state if we're not replicating it.
+		bShouldSkip |= !SharedParams.bIsInitial && Condition == COND_InitialOnly;
+	}
 
 	if (bShouldSkip)
 	{
@@ -1510,7 +1522,8 @@ static void CompareParentProperties(
 
 			for (int32 ParentIndex = 0; ParentIndex < SharedParams.Parents.Num(); ++ParentIndex)
 			{
-				const bool bRecompareInitialProperties = SharedParams.bIsInitial && SharedParams.Parents[ParentIndex].Condition == COND_InitialOnly;
+				const ELifetimeCondition Condition = SharedParams.Parents[ParentIndex].Condition;
+				const bool bRecompareInitialProperties = SharedParams.bIsInitial && (Condition == COND_InitialOnly || (Condition == COND_Dynamic && SharedParams.RepChangedPropertyTracker && SharedParams.RepChangedPropertyTracker->GetDynamicCondition(ParentIndex) == COND_InitialOnly));
 
 				const bool bIsPropertyDirty = bRecompareInitialProperties ||
 												UE_RepLayout_Private::IsPropertyDirty(ParentIndex, bRecentlyCollectedGarbage, SharedParams, StackParams);
@@ -1522,7 +1535,7 @@ static void CompareParentProperties(
 		}
 #endif // WITH_PUSH_VALIDATION_SUPPORT
 
-		else if (UNLIKELY(SharedParams.bIsInitial && EnumHasAnyFlags(SharedParams.Flags, ERepLayoutFlags::HasInitialOnlyProperties)))
+		else if (UNLIKELY(SharedParams.bIsInitial && EnumHasAnyFlags(SharedParams.Flags, ERepLayoutFlags::HasInitialOnlyProperties | ERepLayoutFlags::HasDynamicConditionProperties)))
 		{
 			/*
 				Most replication conditions don't actually effect whether or not we compare properties,
@@ -1554,8 +1567,9 @@ static void CompareParentProperties(
 
 			for (int32 ParentIndex = 0; ParentIndex < SharedParams.Parents.Num(); ++ParentIndex)
 			{
-				if (SharedParams.Parents[ParentIndex].Condition == COND_InitialOnly ||
-					UE_RepLayout_Private::IsPropertyDirty(ParentIndex, bRecentlyCollectedGarbage, SharedParams, StackParams))
+				const ELifetimeCondition Condition = SharedParams.Parents[ParentIndex].Condition;
+				if (Condition == COND_InitialOnly || UE_RepLayout_Private::IsPropertyDirty(ParentIndex, bRecentlyCollectedGarbage, SharedParams, StackParams)
+					|| (Condition == COND_Dynamic && SharedParams.RepChangedPropertyTracker && SharedParams.RepChangedPropertyTracker->GetDynamicCondition(ParentIndex) == COND_InitialOnly))
 				{
 					UE_RepLayout_Private::CompareParentPropertyHelper(ParentIndex, SharedParams, StackParams);
 				}
@@ -6249,6 +6263,10 @@ void FRepLayout::InitFromClass(
 			{
 				Flags |= ERepLayoutFlags::HasInitialOnlyProperties;
 			}
+			else if (LifetimeProps[i].Condition == COND_Dynamic)
+			{
+				Flags |= ERepLayoutFlags::HasDynamicConditionProperties;
+			}
 
 			if (EnumHasAnyFlags(Parents[ParentIndex].Flags, ERepParentFlags::HasNetSerializeProperties | ERepParentFlags::HasObjectProperties))
 			{
@@ -7149,6 +7167,7 @@ TStaticBitArray<COND_Max> FSendingRepState::BuildConditionMapFromRepFlags(const
 	ConditionMap[COND_SkipReplay] = !bIsReplay;
 
 	ConditionMap[COND_Custom] = true;
+	ConditionMap[COND_Dynamic] = true;
 	ConditionMap[COND_Never] = false;
 
 	return ConditionMap;
@@ -7167,19 +7186,33 @@ bool FSendingRepState::HasAnyPendingRetirements() const
 	return false;
 }
 
-void FRepLayout::RebuildConditionalProperties(
-	FSendingRepState* RESTRICT RepState,
-	const FReplicationFlags& RepFlags) const
+void FRepLayout::RebuildConditionalProperties(FSendingRepState* RepState, const FReplicationFlags RepFlags) const
 {
 	SCOPE_CYCLE_COUNTER(STAT_NetRebuildConditionalTime);
 	
+	RepState->RepFlags = RepFlags;
+
 	TStaticBitArray<COND_Max> ConditionMap = FSendingRepState::BuildConditionMapFromRepFlags(RepFlags);
-	for (auto It = TBitArray<>::FIterator(RepState->InactiveParents); It; ++It)
+	if (EnumHasAnyFlags(Flags, ERepLayoutFlags::HasDynamicConditionProperties) && RepState->RepChangedPropertyTracker.IsValid())
 	{
-		It.GetValue() = !ConditionMap[Parents[It.GetIndex()].Condition];
+		const FRepChangedPropertyTracker* RepChangedPropertyTracker = RepState->RepChangedPropertyTracker.Get();
+		for (auto It = TBitArray<>::FIterator(RepState->InactiveParents); It; ++It)
+		{
+			ELifetimeCondition Condition = Parents[It.GetIndex()].Condition;
+			if (Condition == COND_Dynamic)
+			{
+				Condition = RepChangedPropertyTracker->GetDynamicCondition(It.GetIndex());
+			}
+			It.GetValue() = !ConditionMap[Condition];
+		}
+	}
+	else
+	{
+		for (auto It = TBitArray<>::FIterator(RepState->InactiveParents); It; ++It)
+		{
+			It.GetValue() = !ConditionMap[Parents[It.GetIndex()].Condition];
+		}
 	}
-
-	RepState->RepFlags = RepFlags;
 }
 
 PRAGMA_DISABLE_DEPRECATION_WARNINGS
@@ -8440,6 +8473,8 @@ const TCHAR* LexToString(ERepLayoutFlags Flag)
 		return TEXT("FullPushProperties");
 	case ERepLayoutFlags::HasInitialOnlyProperties:
 		return TEXT("HasInitialOnlyProperties");
+	case ERepLayoutFlags::HasDynamicConditionProperties:
+		return TEXT("HasDynamicConditionProperties");
 	default:
 		check(false);
 		return TEXT("Unknown");
diff --git a/Engine/Source/Runtime/Engine/Public/Net/RepLayout.h b/Engine/Source/Runtime/Engine/Public/Net/RepLayout.h
--- a/Engine/Source/Runtime/Engine/Public/Net/RepLayout.h
+++ b/Engine/Source/Runtime/Engine/Public/Net/RepLayout.h
@@ -1053,6 +1053,7 @@ enum class ERepLayoutFlags : uint8
 	NoReplicatedProperties				= (1 << 4), //! Will be set if the RepLayout has no lifetime properties, or they are all disabled.
 	FullPushProperties					= (1 << 5), //! All properties in this RepLayout use Push Model.
 	HasInitialOnlyProperties			= (1 << 6), //! There is at least 1 Initial Only Lifetime property on this RepLayout.
+	HasDynamicConditionProperties		= (1 << 7), //! There is at least 1 Dynamic lifetime property on this RepLayout.
 };
 ENUM_CLASS_FLAGS(ERepLayoutFlags);
 
@@ -1591,8 +1592,8 @@ private:
 		TArray<uint16>& MergedDirty) const;
 
 	void RebuildConditionalProperties(
-		FSendingRepState* RESTRICT RepState,
-		const FReplicationFlags& RepFlags) const;
+		FSendingRepState* RepState,
+		const FReplicationFlags RepFlags) const;
 
 	void UpdateChangelistHistory(
 		FSendingRepState* RepState,
diff --git a/Engine/Source/Runtime/Engine/Public/Net/UnrealNetwork.h b/Engine/Source/Runtime/Engine/Public/Net/UnrealNetwork.h
--- a/Engine/Source/Runtime/Engine/Public/Net/UnrealNetwork.h
+++ b/Engine/Source/Runtime/Engine/Public/Net/UnrealNetwork.h
@@ -344,6 +344,23 @@ static FProperty* GetReplicatedProperty(const UClass* CallingClass, const UClass
 	} \
 }
 
+/** To be used in a UObject's GetReplicatedCustomConditionState for example. */
+#define DOREPDYNAMICCONDITION_INITCONDITION_FAST(c,v,cond) \
+{ \
+	static const bool bIsValid_##c_##v = ValidateReplicatedClassInheritance(StaticClass(), c::StaticClass(), TEXT(#v)); \
+	OutActiveState.SetDynamicCondition((uint16)c::ENetFields_Private::v, cond); \
+}
+
+/**
+ * To be used when a condition is to be changed. This should be called as soon as the condition changes, not during PreReplication.
+ * @note Changing the condition of a custom delta property, such as a FastArraySerializer, after it has replicated can result in unexpected behavior.
+ */
+#define DOREPDYNAMICCONDITION_SETCONDITION_FAST(c,v,cond) \
+{ \
+	UE::Net::Private::FNetPropertyConditionManager::Get().SetPropertyDynamicCondition(this, (uint16)c::ENetFields_Private::v, cond); \
+}
+
+
 
 ENGINE_API void RegisterReplicatedLifetimeProperty(
 	const FProperty* ReplicatedProperty,
diff --git a/Engine/Source/Runtime/Experimental/Iris/Core/Private/Iris/ReplicationState/ReplicationStateDescriptorBuilder.cpp b/Engine/Source/Runtime/Experimental/Iris/Core/Private/Iris/ReplicationState/ReplicationStateDescriptorBuilder.cpp
--- a/Engine/Source/Runtime/Experimental/Iris/Core/Private/Iris/ReplicationState/ReplicationStateDescriptorBuilder.cpp
+++ b/Engine/Source/Runtime/Experimental/Iris/Core/Private/Iris/ReplicationState/ReplicationStateDescriptorBuilder.cpp
@@ -1784,6 +1784,7 @@ EMemberPropertyTraits FPropertyReplicationStateDescriptorBuilder::GetConnectionF
 		{COND_ReplayOnly, EMemberPropertyTraits::HasLifetimeConditionals},
 		{COND_SimulatedOnlyNoReplay, EMemberPropertyTraits::HasLifetimeConditionals},
 		{COND_SimulatedOrPhysicsNoReplay, EMemberPropertyTraits::HasLifetimeConditionals},
+		{COND_Dynamic, EMemberPropertyTraits::HasLifetimeConditionals},
 	};
 
 	static_assert(TConditionPairIsSorted<>(ConditionPairs), "ConditionPairs order needs to be updated.");
diff --git a/Engine/Source/Runtime/Experimental/Iris/Core/Private/Iris/ReplicationSystem/Conditionals/ReplicationConditionals.cpp b/Engine/Source/Runtime/Experimental/Iris/Core/Private/Iris/ReplicationSystem/Conditionals/ReplicationConditionals.cpp
--- a/Engine/Source/Runtime/Experimental/Iris/Core/Private/Iris/ReplicationSystem/Conditionals/ReplicationConditionals.cpp
+++ b/Engine/Source/Runtime/Experimental/Iris/Core/Private/Iris/ReplicationSystem/Conditionals/ReplicationConditionals.cpp
@@ -307,6 +307,136 @@ bool FReplicationConditionals::SetPropertyCustomCondition(FInternalNetRefIndex O
 	return false;
 }
 
+bool FReplicationConditionals::SetPropertyDynamicCondition(FInternalNetRefIndex ObjectIndex, const void* Owner, uint16 RepIndex, ELifetimeCondition Condition)
+{
+	IRIS_PROFILER_SCOPE(FReplicationConditionals_SetPropertyDynamicCondition);
+
+	const FNetRefHandleManager::FReplicatedObjectData& ReplicatedObjectData = NetRefHandleManager->GetReplicatedObjectDataNoCheck(ObjectIndex);
+	const FReplicationProtocol* Protocol = ReplicatedObjectData.Protocol;
+
+	if (NetRefHandleManager->GetReplicatedObjectStateBufferNoCheck(ObjectIndex) == nullptr || !EnumHasAnyFlags(Protocol->ProtocolTraits, EReplicationProtocolTraits::HasLifetimeConditionals))
+	{
+		return false;
+	}
+
+	if (Protocol->LifetimeConditionalsStateCount == 1U)
+	{
+		const SIZE_T StateIndex = Protocol->FirstLifetimeConditionalsStateIndex;
+		const FReplicationInstanceProtocol* InstanceProtocol = ReplicatedObjectData.InstanceProtocol;
+		const FReplicationInstanceProtocol::FFragmentData& Fragment = InstanceProtocol->FragmentData[StateIndex];
+		const FReplicationStateDescriptor* StateDescriptor = Protocol->ReplicationStateDescriptors[StateIndex];
+
+		if (RepIndex >= StateDescriptor->RepIndexCount)
+		{
+			UE_LOG(LogIris, Warning, TEXT("Trying to change non-existing dynamic conditional for RepIndex %u in protocol %s"), RepIndex, ToCStr(Protocol->DebugName));
+			return false;
+		}
+
+		const FReplicationStateMemberRepIndexToMemberIndexDescriptor& RepIndexToMemberIndexDescriptor = StateDescriptor->MemberRepIndexToMemberIndexDescriptors[RepIndex];
+		const uint16 MemberIndex = RepIndexToMemberIndexDescriptor.MemberIndex;
+		if ((MemberIndex == FReplicationStateMemberRepIndexToMemberIndexDescriptor::InvalidEntry))
+		{
+			UE_LOG(LogIris, Warning, TEXT("Trying to change non-existing dynamic conditional for RepIndex %u in protocol %s"), RepIndex, ToCStr(Protocol->DebugName));
+			return false;
+		}
+
+		if (StateDescriptor->MemberLifetimeConditionDescriptors[MemberIndex].Condition != COND_Dynamic)
+		{
+			UE_LOG(LogIris, Warning, TEXT("Trying to change condition for member %s with wrong condition in protocol %s"), ToCStr(StateDescriptor->MemberDebugDescriptors[MemberIndex].DebugName), ToCStr(Protocol->DebugName));
+			return false;
+		}
+
+		const ELifetimeCondition OldCondition = GetDynamicCondition(ObjectIndex, RepIndex);
+		SetDynamicCondition(ObjectIndex, RepIndex, Condition);
+
+		// If a condition may cause something to go from not replicated to replicated we mark the changemask as dirty and invalidate baselines.
+		if (DynamicConditionChangeRequiresBaselineInvalidation(OldCondition, Condition))
+		{
+			const FReplicationStateMemberChangeMaskDescriptor& ChangeMaskDescriptor = StateDescriptor->MemberChangeMaskDescriptors[MemberIndex];
+
+			FNetBitArrayView MemberChangeMask = UE::Net::Private::GetMemberChangeMask(Fragment.ExternalSrcBuffer, StateDescriptor);
+			FReplicationStateHeader& ReplicationStateHeader = UE::Net::Private::GetReplicationStateHeader(Fragment.ExternalSrcBuffer, StateDescriptor);
+			MarkDirty(ReplicationStateHeader, MemberChangeMask, ChangeMaskDescriptor);
+
+			// $TODO Consider more extensive checking to see if only a single connection requires baseline invalidation.
+			BaselineInvalidationTracker->InvalidateBaselines(ObjectIndex, BaselineInvalidationTracker->InvalidateBaselineForAllConnections);
+		}
+
+		return true;
+	}
+	else
+	{
+		constexpr uint32 MaxFragmentOwnerCount = 1U;
+		UObject* FragmentOwners[MaxFragmentOwnerCount] = {};
+		FReplicationStateOwnerCollector FragmentOwnerCollector(FragmentOwners, MaxFragmentOwnerCount);
+
+		const FReplicationInstanceProtocol* InstanceProtocol = ReplicatedObjectData.InstanceProtocol;
+
+		const SIZE_T FirstRelevantStateIndex = Protocol->FirstLifetimeConditionalsStateIndex;
+		for (const FReplicationStateDescriptor*& StateDescriptor : MakeArrayView(Protocol->ReplicationStateDescriptors + FirstRelevantStateIndex, Protocol->ReplicationStateCount - FirstRelevantStateIndex))
+		{
+			if (EnumHasAnyFlags(StateDescriptor->Traits, EReplicationStateTraits::HasLifetimeConditionals))
+			{
+				const SIZE_T StateIndex = &StateDescriptor - Protocol->ReplicationStateDescriptors;
+
+				// Is the passed Owner the owner of the fragment?
+				{
+					const FReplicationFragment* ReplicationFragment = InstanceProtocol->Fragments[StateIndex];
+
+					FragmentOwnerCollector.Reset();
+					ReplicationFragment->CollectOwner(&FragmentOwnerCollector);
+					if (FragmentOwnerCollector.GetOwnerCount() == 0U || FragmentOwnerCollector.GetOwners()[0] != Owner)
+					{
+						// Not the right owner.
+						continue;
+					}
+				}
+
+				// Can this state contain this property?
+				if (RepIndex >= StateDescriptor->RepIndexCount)
+				{
+					continue;
+				}
+
+				// Does this state contain this property?
+				const FReplicationStateMemberRepIndexToMemberIndexDescriptor& RepIndexToMemberIndexDescriptor = StateDescriptor->MemberRepIndexToMemberIndexDescriptors[RepIndex];
+				const uint16 MemberIndex = RepIndexToMemberIndexDescriptor.MemberIndex;
+				if (MemberIndex == FReplicationStateMemberRepIndexToMemberIndexDescriptor::InvalidEntry)
+				{
+					continue;
+				}
+
+				// We've found the relevant state. Verify it's a dynamic condition member.
+				if (StateDescriptor->MemberLifetimeConditionDescriptors[MemberIndex].Condition != COND_Dynamic)
+				{
+					UE_LOG(LogIris, Warning, TEXT("Trying to change condition for member %s with wrong condition in protocol %s"), ToCStr(StateDescriptor->MemberDebugDescriptors[MemberIndex].DebugName), ToCStr(Protocol->DebugName));
+					return false;
+				}
+
+				const ELifetimeCondition OldCondition = GetDynamicCondition(ObjectIndex, RepIndex);
+				SetDynamicCondition(ObjectIndex, RepIndex, Condition);
+
+				// If a condition may cause something to go from not replicated to replicated we mark the changemask as dirty and invalidate baselines.
+				if (DynamicConditionChangeRequiresBaselineInvalidation(OldCondition, Condition))
+				{
+					const FReplicationInstanceProtocol::FFragmentData& Fragment = InstanceProtocol->FragmentData[StateIndex];
+					const FReplicationStateMemberChangeMaskDescriptor& ChangeMaskDescriptor = StateDescriptor->MemberChangeMaskDescriptors[MemberIndex];
+
+					FNetBitArrayView MemberChangeMask = UE::Net::Private::GetMemberChangeMask(Fragment.ExternalSrcBuffer, StateDescriptor);
+					FReplicationStateHeader& ReplicationStateHeader = UE::Net::Private::GetReplicationStateHeader(Fragment.ExternalSrcBuffer, StateDescriptor);
+					MarkDirty(ReplicationStateHeader, MemberChangeMask, ChangeMaskDescriptor);
+
+					BaselineInvalidationTracker->InvalidateBaselines(ObjectIndex, BaselineInvalidationTracker->InvalidateBaselineForAllConnections);
+				}
+
+				return true;
+			}
+		}
+	}
+
+	return false;
+}
+
 void FReplicationConditionals::Update()
 {
 	UpdateObjectsInScope();
@@ -386,11 +516,12 @@ void FReplicationConditionals::GetSubObjectsToReplicate(uint32 ReplicationConnec
 	// For now, we do nothing to detect if a conditional has changed on the RootParent, we simply defer this until the next
 	// time the subobjects are marked as dirty. We might want to consider to explicitly mark object and subobjects as dirty when 
 	// the owning connections or conditionals such as bRepPhysics or Role is changed.
-	const FConditionalsMask LifetimeConditionals = GetLifetimeConditionals(ReplicationConnectionId, RootObjectIndex);
+	constexpr bool bInitialState = false;
+	const FConditionalsMask LifetimeConditionals = GetLifetimeConditionals(ReplicationConnectionId, RootObjectIndex, bInitialState);
 	GetChildSubObjectsToReplicate(ReplicationConnectionId, LifetimeConditionals, RootObjectIndex, OutSubObjectsToReplicate);
 }
 
-bool FReplicationConditionals::ApplyConditionalsToChangeMask(uint32 ReplicatingConnectionId, FInternalNetRefIndex ParentObjectIndex, FInternalNetRefIndex ObjectIndex, uint32* ChangeMaskData, const uint32* ConditionalChangeMaskData, const FReplicationProtocol* Protocol)
+bool FReplicationConditionals::ApplyConditionalsToChangeMask(uint32 ReplicatingConnectionId, bool bIsInitialState, FInternalNetRefIndex ParentObjectIndex, FInternalNetRefIndex ObjectIndex, uint32* ChangeMaskData, const uint32* ConditionalChangeMaskData, const FReplicationProtocol* Protocol)
 {
 	IRIS_PROFILER_SCOPE(FReplicationConditionals_ApplyConditionalsToChangeMask);
 
@@ -402,7 +533,7 @@ bool FReplicationConditionals::ApplyConditionalsToChangeMask(uint32 ReplicatingC
 	// Legacy lifetime conditionals support.
 	if (EnumHasAnyFlags(Protocol->ProtocolTraits, EReplicationProtocolTraits::HasLifetimeConditionals))
 	{
-		const FConditionalsMask LifetimeConditionals = GetLifetimeConditionals(ReplicatingConnectionId, ParentObjectIndex);
+		const FConditionalsMask LifetimeConditionals = GetLifetimeConditionals(ReplicatingConnectionId, ParentObjectIndex, bIsInitialState);
 		FConditionalsMask PrevLifeTimeConditions = ConnectionInfos[ReplicatingConnectionId].ObjectConditionals[ObjectIndex];
 		if (PrevLifeTimeConditions.IsUninitialized())
 		{
@@ -421,7 +552,16 @@ bool FReplicationConditionals::ApplyConditionalsToChangeMask(uint32 ReplicatingC
 			for (uint32 MemberIt = 0U, MemberEndIt = StateDescriptor->MemberCount; MemberIt != MemberEndIt; ++MemberIt)
 			{
 				const FReplicationStateMemberLifetimeConditionDescriptor& LifetimeConditionDescriptor = LifetimeConditionDescriptors[MemberIt];
-				const ELifetimeCondition Condition = static_cast<ELifetimeCondition>(LifetimeConditionDescriptor.Condition);
+				ELifetimeCondition Condition = static_cast<ELifetimeCondition>(LifetimeConditionDescriptor.Condition);
+				if (Condition == COND_Dynamic)
+				{
+					const FProperty* Property = StateDescriptor->MemberProperties[MemberIt];
+					if (ensure(Property != nullptr))
+					{
+						Condition = GetDynamicCondition(ObjectIndex, Property->RepIndex);
+					}
+				}
+				
 				// If condition was enabled we need to dirty changemask of relevant members. If it was disabled we clear the changemask of relevant members.
 				if (LifetimeConditionals.IsConditionEnabled(Condition))
 				{
@@ -460,7 +600,16 @@ bool FReplicationConditionals::ApplyConditionalsToChangeMask(uint32 ReplicatingC
 					for (uint32 MemberIt = 0U, MemberEndIt = StateDescriptor->MemberCount; MemberIt != MemberEndIt; ++MemberIt)
 					{
 						const FReplicationStateMemberLifetimeConditionDescriptor& LifetimeConditionDescriptor = LifetimeConditionDescriptors[MemberIt];
-						const ELifetimeCondition Condition = static_cast<ELifetimeCondition>(LifetimeConditionDescriptor.Condition);
+						ELifetimeCondition Condition = static_cast<ELifetimeCondition>(LifetimeConditionDescriptor.Condition);
+						if (Condition == COND_Dynamic)
+						{
+							const FProperty* Property = StateDescriptor->MemberProperties[MemberIt];
+							if (ensure(Property != nullptr))
+							{
+								Condition = GetDynamicCondition(ObjectIndex, Property->RepIndex);
+							}
+						}
+
 						// If the condition is fulfilled the changemask will remain intact so we can continue to the next member.
 						if (LifetimeConditionals.IsConditionEnabled(Condition))
 						{
@@ -580,7 +729,7 @@ void FReplicationConditionals::UpdateObjectsInScope()
 	}
 }
 
-FReplicationConditionals::FConditionalsMask FReplicationConditionals::GetLifetimeConditionals(uint32 ReplicatingConnectionId, FInternalNetRefIndex ParentObjectIndex) const
+FReplicationConditionals::FConditionalsMask FReplicationConditionals::GetLifetimeConditionals(uint32 ReplicatingConnectionId, FInternalNetRefIndex ParentObjectIndex, bool bIsInitialState) const
 {
 	FConditionalsMask ConditionalsMask{0};
 
@@ -594,12 +743,14 @@ FReplicationConditionals::FConditionalsMask FReplicationConditionals::GetLifetim
 
 	ConditionalsMask.SetConditionEnabled(COND_None, true);
 	ConditionalsMask.SetConditionEnabled(COND_Custom, true);
+	ConditionalsMask.SetConditionEnabled(COND_Dynamic, true);
 	ConditionalsMask.SetConditionEnabled(COND_OwnerOnly, bIsReplicatingToOwner);
 	ConditionalsMask.SetConditionEnabled(COND_SkipOwner, !bIsReplicatingToOwner);
 	ConditionalsMask.SetConditionEnabled(COND_SimulatedOnly, bRoleSimulated);
 	ConditionalsMask.SetConditionEnabled(COND_AutonomousOnly, bRoleAutonomous);
 	ConditionalsMask.SetConditionEnabled(COND_SimulatedOrPhysics, bRoleSimulated | bRepPhysics);
-	ConditionalsMask.SetConditionEnabled(COND_InitialOrOwner, bIsReplicatingToOwner);
+	ConditionalsMask.SetConditionEnabled(COND_InitialOnly, bIsInitialState);
+	ConditionalsMask.SetConditionEnabled(COND_InitialOrOwner, bIsReplicatingToOwner | bIsInitialState);
 	ConditionalsMask.SetConditionEnabled(COND_ReplayOrOwner, bIsReplicatingToOwner);
 	ConditionalsMask.SetConditionEnabled(COND_SimulatedOnlyNoReplay, bRoleSimulated);
 	ConditionalsMask.SetConditionEnabled(COND_SimulatedOrPhysicsNoReplay, bRoleSimulated | bRepPhysics);
@@ -611,6 +762,9 @@ void FReplicationConditionals::ClearPerObjectInfo(FInternalNetRefIndex ObjectInd
 {
 	FPerObjectInfo& PerObjectInfo = PerObjectInfos.GetData()[ObjectIndex];
 	PerObjectInfo = {};
+
+	// Remove any dynamic conditions information stored.
+	DynamicConditions.Remove(ObjectIndex);
 }
 
 void FReplicationConditionals::ClearConnectionInfosForObject(const FNetBitArray& ValidConnections, uint32 MaxConnectionId, FInternalNetRefIndex ObjectIndex)
@@ -631,4 +785,38 @@ void FReplicationConditionals::ClearConnectionInfosForObject(const FNetBitArray&
 	}
 }
 
+ELifetimeCondition FReplicationConditionals::GetDynamicCondition(FInternalNetRefIndex ObjectIndex, uint16 RepIndex) const
+{
+	const FObjectDynamicConditions* ObjectConditions = DynamicConditions.Find(ObjectIndex);
+	if (ObjectConditions == nullptr)
+	{
+		return COND_Dynamic;
+	}
+
+	const int16* Condition = ObjectConditions->DynamicConditions.Find(RepIndex);
+	if (Condition == nullptr)
+	{
+		return COND_Dynamic;
+	}
+
+	return static_cast<const ELifetimeCondition>(*Condition);
+}
+
+void FReplicationConditionals::SetDynamicCondition(FInternalNetRefIndex ObjectIndex, uint16 RepIndex, ELifetimeCondition Condition)
+{
+	FObjectDynamicConditions& ObjectConditions = DynamicConditions.FindOrAdd(ObjectIndex);
+	ObjectConditions.DynamicConditions.Emplace(RepIndex, static_cast<int16>(Condition));
+}
+
+bool FReplicationConditionals::DynamicConditionChangeRequiresBaselineInvalidation(ELifetimeCondition OldCondition, ELifetimeCondition NewCondition) const
+{
+	// If the old condition didn't cause the member to always be replicated it could have been not replicated to one or more connections.
+	const bool OldConditionMayHaveBeenDisabled = !(OldCondition == COND_None || OldCondition == COND_Dynamic);
+
+	// If the new condition is something other than never replicating then it may be replicated.
+	const bool NewConditionMayBeEnabled = (NewCondition != COND_Never);
+
+	return OldConditionMayHaveBeenDisabled && NewConditionMayBeEnabled;
+}
+
 }
diff --git a/Engine/Source/Runtime/Experimental/Iris/Core/Private/Iris/ReplicationSystem/Conditionals/ReplicationConditionals.h b/Engine/Source/Runtime/Experimental/Iris/Core/Private/Iris/ReplicationSystem/Conditionals/ReplicationConditionals.h
--- a/Engine/Source/Runtime/Experimental/Iris/Core/Private/Iris/ReplicationSystem/Conditionals/ReplicationConditionals.h
+++ b/Engine/Source/Runtime/Experimental/Iris/Core/Private/Iris/ReplicationSystem/Conditionals/ReplicationConditionals.h
@@ -4,8 +4,10 @@
 
 #include "CoreTypes.h"
 #include "Containers/Array.h"
+#include "Containers/Map.h"
 #include "Net/Core/NetBitArray.h"
 
+enum ELifetimeCondition : int;
 namespace UE::Net
 {
 	enum class EReplicationCondition : uint32;
@@ -59,10 +61,11 @@ public:
 	// For property custom conditions only
 	void InitPropertyCustomConditions(FInternalNetRefIndex ObjectIndex);
 	bool SetPropertyCustomCondition(FInternalNetRefIndex ObjectIndex, const void* Owner, uint16 RepIndex, bool bIsActive);
+	bool SetPropertyDynamicCondition(FInternalNetRefIndex ObjectIndex, const void* Owner, uint16 RepIndex, ELifetimeCondition Condition);
 
 	void Update();
 
-	bool ApplyConditionalsToChangeMask(uint32 ReplicatingConnectionId, FInternalNetRefIndex ParentObjectIndex, FInternalNetRefIndex ObjectIndex, uint32* ChangeMaskData, const uint32* ConditionalChangeMaskData, const FReplicationProtocol* Protocol);
+	bool ApplyConditionalsToChangeMask(uint32 ReplicatingConnectionId, bool bIsInitialState, FInternalNetRefIndex ParentObjectIndex, FInternalNetRefIndex ObjectIndex, uint32* ChangeMaskData, const uint32* ConditionalChangeMaskData, const FReplicationProtocol* Protocol);
 
 	using FSubObjectsToReplicateArray = TArray<FInternalNetRefIndex, TInlineAllocator<32>>;
 	void GetSubObjectsToReplicate(uint32 ReplicationConnectionId, FInternalNetRefIndex ParentObjectIndex, FSubObjectsToReplicateArray& OutSubObjectsToReplicate);
@@ -96,18 +99,27 @@ private:
 		int8 Condition;
 	};
 
+	struct FObjectDynamicConditions
+	{
+		// RepIndex and ELifetimeCondition expressed as int16
+		TMap<uint16, int16> DynamicConditions;
+	};
+
 private:
 	void UpdateObjectsInScope();
 
-	FConditionalsMask GetLifetimeConditionals(uint32 ReplicatingConnectionId, FInternalNetRefIndex ParentObjectIndex) const;
-	bool ObjectHasLifetimeConditionals(FInternalNetRefIndex ObjectIndex) const;	
+	FConditionalsMask GetLifetimeConditionals(uint32 ReplicatingConnectionId, FInternalNetRefIndex ParentObjectIndex, bool bInitialState) const;
 
 	FPerObjectInfo* GetPerObjectInfo(FInternalNetRefIndex ObjectIndex);
 	const FPerObjectInfo* GetPerObjectInfo(FInternalNetRefIndex ObjectIndex) const;
 	void ClearPerObjectInfo(FInternalNetRefIndex ObjectIndex);
 	void ClearConnectionInfosForObject(const FNetBitArray& ValidConnections, uint32 MaxConnectionId, FInternalNetRefIndex ObjectIndex);
 
-	void GetChildSubObjectsToReplicate(uint32 ReplicatingConnectionId, const FConditionalsMask& LifetimeConditionals,  const FInternalNetRefIndex ParentObjectIndex, FSubObjectsToReplicateArray& OutSubObjectsToReplicate);
+	void GetChildSubObjectsToReplicate(uint32 ReplicatingConnectionId, const FConditionalsMask& LifetimeConditionals, const FInternalNetRefIndex ParentObjectIndex, FSubObjectsToReplicateArray& OutSubObjectsToReplicate);
+
+	ELifetimeCondition GetDynamicCondition(FInternalNetRefIndex ObjectIndex, uint16 RepIndex) const;
+	void SetDynamicCondition(FInternalNetRefIndex ObjectIndex, uint16 RepIndex, ELifetimeCondition Condition);
+	bool DynamicConditionChangeRequiresBaselineInvalidation(ELifetimeCondition OldCondition, ELifetimeCondition NewCondition) const;
 
 private:
 	const FNetRefHandleManager* NetRefHandleManager = nullptr;
@@ -121,6 +133,7 @@ private:
 
 	TArray<FPerObjectInfo> PerObjectInfos;
 	TArray<FPerConnectionInfo> ConnectionInfos;
+	TMap<FInternalNetRefIndex, FObjectDynamicConditions> DynamicConditions;
 };
 
 inline FReplicationConditionals::FPerObjectInfo* FReplicationConditionals::GetPerObjectInfo(FInternalNetRefIndex ObjectIndex)
diff --git a/Engine/Source/Runtime/Experimental/Iris/Core/Private/Iris/ReplicationSystem/ObjectReplicationBridge.cpp b/Engine/Source/Runtime/Experimental/Iris/Core/Private/Iris/ReplicationSystem/ObjectReplicationBridge.cpp
--- a/Engine/Source/Runtime/Experimental/Iris/Core/Private/Iris/ReplicationSystem/ObjectReplicationBridge.cpp
+++ b/Engine/Source/Runtime/Experimental/Iris/Core/Private/Iris/ReplicationSystem/ObjectReplicationBridge.cpp
@@ -121,28 +121,15 @@ void UObjectReplicationBridge::Initialize(UReplicationSystem* InReplicationSyste
 
 	LoadConfig();
 
-	// Hookup delegate when a property custom conditions is changed
-	OnCustomConditionChangedHandle = UE::Net::Private::FPropertyConditionDelegates::GetOnPropertyCustomConditionChangedDelegate().AddLambda([this](const UObject* Owner, uint16 RepIndex, bool bEnable)
-	{
-		using namespace UE::Net;
-		using namespace UE::Net::Private;
-
-		const FNetRefHandle RefHandle = this->GetReplicatedRefHandle(Owner);
-		if (RefHandle.IsValid())
-		{
-			FReplicationSystemInternal* ReplicationSystemInternal = this->GetReplicationSystem()->GetReplicationSystemInternal();
-			const FNetRefHandleManager& LocalNetRefHandleManager = ReplicationSystemInternal->GetNetRefHandleManager();
-			FReplicationConditionals& Conditionals = ReplicationSystemInternal->GetConditionals();
-
-			Conditionals.SetPropertyCustomCondition(LocalNetRefHandleManager.GetInternalIndex(RefHandle), Owner, RepIndex, bEnable);
-		}
-	});
+	InitConditionalPropertyDelegates();
 }
 
 void UObjectReplicationBridge::Deinitialize()
 {
 	UE::Net::Private::FPropertyConditionDelegates::GetOnPropertyCustomConditionChangedDelegate().Remove(OnCustomConditionChangedHandle);
+	UE::Net::Private::FPropertyConditionDelegates::GetOnPropertyDynamicConditionChangedDelegate().Remove(OnDynamicConditionChangedHandle);
 	OnCustomConditionChangedHandle.Reset();
+	OnDynamicConditionChangedHandle.Reset();
 	PollFrequencyLimiter->Deinit();
 	Super::Deinitialize();
 }
@@ -1414,3 +1401,37 @@ void UObjectReplicationBridge::SetShouldSubclassUseSameFilterFunction(TFunction<
 
 	ShouldSubclassUseSameFilterFunction = InShouldSubclassUseSameFilterFunction;
 }
+
+void UObjectReplicationBridge::InitConditionalPropertyDelegates()
+{
+	using namespace UE::Net;
+	using namespace UE::Net::Private;
+
+	// Hookup delegate for when a property custom condition is changed
+	OnCustomConditionChangedHandle = UE::Net::Private::FPropertyConditionDelegates::GetOnPropertyCustomConditionChangedDelegate().AddLambda([this](const UObject* Owner, uint16 RepIndex, bool bEnable)
+	{
+		const FNetRefHandle RefHandle = this->GetReplicatedRefHandle(Owner);
+		if (RefHandle.IsValid())
+		{
+			FReplicationSystemInternal* ReplicationSystemInternal = this->GetReplicationSystem()->GetReplicationSystemInternal();
+			const FNetRefHandleManager& LocalNetRefHandleManager = ReplicationSystemInternal->GetNetRefHandleManager();
+			FReplicationConditionals& Conditionals = ReplicationSystemInternal->GetConditionals();
+
+			Conditionals.SetPropertyCustomCondition(LocalNetRefHandleManager.GetInternalIndex(RefHandle), Owner, RepIndex, bEnable);
+		}
+	});
+
+	// Hookup delegate for when a property dynamic condition is changed
+	OnDynamicConditionChangedHandle = UE::Net::Private::FPropertyConditionDelegates::GetOnPropertyDynamicConditionChangedDelegate().AddLambda([this](const UObject* Owner, uint16 RepIndex, ELifetimeCondition Condition)
+	{
+		const FNetRefHandle RefHandle = this->GetReplicatedRefHandle(Owner);
+		if (RefHandle.IsValid())
+		{
+			FReplicationSystemInternal* ReplicationSystemInternal = this->GetReplicationSystem()->GetReplicationSystemInternal();
+			const FNetRefHandleManager& LocalNetRefHandleManager = ReplicationSystemInternal->GetNetRefHandleManager();
+			FReplicationConditionals& Conditionals = ReplicationSystemInternal->GetConditionals();
+
+			Conditionals.SetPropertyDynamicCondition(LocalNetRefHandleManager.GetInternalIndex(RefHandle), Owner, RepIndex, Condition);
+		}
+	});
+}
diff --git a/Engine/Source/Runtime/Experimental/Iris/Core/Private/Iris/ReplicationSystem/ReplicationOperations.cpp b/Engine/Source/Runtime/Experimental/Iris/Core/Private/Iris/ReplicationSystem/ReplicationOperations.cpp
--- a/Engine/Source/Runtime/Experimental/Iris/Core/Private/Iris/ReplicationSystem/ReplicationOperations.cpp
+++ b/Engine/Source/Runtime/Experimental/Iris/Core/Private/Iris/ReplicationSystem/ReplicationOperations.cpp
@@ -44,7 +44,7 @@ static FNetDebugName s_LifetimeConditionDebugNames[] =
 	{TEXT("SimulatedOnlyNoReplay"), 0},
 	{TEXT("SimulatedOrPhysicsNoReplay"), 0},
 	{TEXT("SkipReplay"), 0},
-	{TEXT("ILLEGAL"), 0},
+	{TEXT("Dynamic"), 0},
 	{TEXT("Never"), 0},
 	{TEXT("NetGroup"), 0},
 };
diff --git a/Engine/Source/Runtime/Experimental/Iris/Core/Private/Iris/ReplicationSystem/ReplicationWriter.cpp b/Engine/Source/Runtime/Experimental/Iris/Core/Private/Iris/ReplicationSystem/ReplicationWriter.cpp
--- a/Engine/Source/Runtime/Experimental/Iris/Core/Private/Iris/ReplicationSystem/ReplicationWriter.cpp
+++ b/Engine/Source/Runtime/Experimental/Iris/Core/Private/Iris/ReplicationSystem/ReplicationWriter.cpp
@@ -1798,7 +1798,7 @@ FReplicationWriter::EWriteObjectStatus FReplicationWriter::WriteObjectInBatch(FN
 	const bool bNeedToFilterChangeMask = (bIsInitialState || Info.HasDirtyChangeMask) && Info.HasChangemaskFilter;
 	if (bNeedToFilterChangeMask)
 	{
-		ApplyFilterToChangeMask(OutBatchInfo.ParentInternalIndex, InternalIndex, Info, ObjectData.Protocol, ReplicatedObjectStateBuffer);
+		ApplyFilterToChangeMask(OutBatchInfo.ParentInternalIndex, InternalIndex, Info, ObjectData.Protocol, ReplicatedObjectStateBuffer, bIsInitialState);
 	}
 
 	const bool bIsObjectIndexForAttachment = IsObjectIndexForOOBAttachment(InternalIndex);
@@ -1884,7 +1884,7 @@ FReplicationWriter::EWriteObjectStatus FReplicationWriter::WriteObjectInBatch(FN
 						if (PatchupObjectChangeMaskWithInflightChanges(InternalIndex, Info))
 						{
 							// Mask off changemasks that may have been disabled due to conditionals.
-							ApplyFilterToChangeMask(OutBatchInfo.ParentInternalIndex, InternalIndex, Info, ObjectData.Protocol, ReplicatedObjectStateBuffer);
+							ApplyFilterToChangeMask(OutBatchInfo.ParentInternalIndex, InternalIndex, Info, ObjectData.Protocol, ReplicatedObjectStateBuffer, bIsInitialState);
 						}
 
 						UE_LOG_REPLICATIONWRITER_CONN(TEXT("Created new baseline %u for ( InternalObjectIndex: %u )"), CreatedBaselineIndex, InternalIndex);
@@ -2932,10 +2932,10 @@ void FReplicationWriter::SetupReplicationInfoForAttachmentsToObjectsNotInScope()
 	ReplicationRecord.ResetList(ReplicatedObjectsRecordInfoLists[ObjectIndexForOOBAttachment]);
 }
 
-void FReplicationWriter::ApplyFilterToChangeMask(uint32 ParentInternalIndex, uint32 InternalIndex, FReplicationInfo& Info, const FReplicationProtocol* Protocol, const uint8* InternalStateBuffer)
+void FReplicationWriter::ApplyFilterToChangeMask(uint32 ParentInternalIndex, uint32 InternalIndex, FReplicationInfo& Info, const FReplicationProtocol* Protocol, const uint8* InternalStateBuffer, bool bIsInitialState)
 {
 	const uint32* ConditionalChangeMaskPointer = (EnumHasAnyFlags(Protocol->ProtocolTraits, EReplicationProtocolTraits::HasConditionalChangeMask) ? reinterpret_cast<const uint32*>(InternalStateBuffer + Protocol->GetConditionalChangeMaskOffset()) : static_cast<const uint32*>(nullptr));
-	const bool bChangeMaskWasModified = ReplicationConditionals->ApplyConditionalsToChangeMask(Parameters.ConnectionId, ParentInternalIndex, InternalIndex, Info.GetChangeMaskStoragePointer(), ConditionalChangeMaskPointer, Protocol);
+	const bool bChangeMaskWasModified = ReplicationConditionals->ApplyConditionalsToChangeMask(Parameters.ConnectionId, bIsInitialState, ParentInternalIndex, InternalIndex, Info.GetChangeMaskStoragePointer(), ConditionalChangeMaskPointer, Protocol);
 	if (bChangeMaskWasModified && !MakeNetBitArrayView(Info.GetChangeMaskStoragePointer(), Info.ChangeMaskBitCount).IsAnyBitSet())
 	{
 		Info.HasDirtyChangeMask = 0;
diff --git a/Engine/Source/Runtime/Experimental/Iris/Core/Private/Iris/ReplicationSystem/ReplicationWriter.h b/Engine/Source/Runtime/Experimental/Iris/Core/Private/Iris/ReplicationSystem/ReplicationWriter.h
--- a/Engine/Source/Runtime/Experimental/Iris/Core/Private/Iris/ReplicationSystem/ReplicationWriter.h
+++ b/Engine/Source/Runtime/Experimental/Iris/Core/Private/Iris/ReplicationSystem/ReplicationWriter.h
@@ -418,7 +418,7 @@ private:
 	// Setup replication info to be able to send attachments to objects not in scope
 	void SetupReplicationInfoForAttachmentsToObjectsNotInScope();
 
-	void ApplyFilterToChangeMask(uint32 ParentInternalIndex, uint32 InternalIndex, FReplicationInfo& Info, const FReplicationProtocol* Protocol, const uint8* InternalStateBuffer);
+	void ApplyFilterToChangeMask(uint32 ParentInternalIndex, uint32 InternalIndex, FReplicationInfo& Info, const FReplicationProtocol* Protocol, const uint8* InternalStateBuffer, bool bIsInitialState);
 
 	// Patchup changemask to include any in-flight changes. Returns true if in-flight changes were added.
 	bool PatchupObjectChangeMaskWithInflightChanges(uint32 InternalIndex, FReplicationInfo& Info);
diff --git a/Engine/Source/Runtime/Experimental/Iris/Core/Public/Iris/ReplicationSystem/ObjectReplicationBridge.h b/Engine/Source/Runtime/Experimental/Iris/Core/Public/Iris/ReplicationSystem/ObjectReplicationBridge.h
--- a/Engine/Source/Runtime/Experimental/Iris/Core/Public/Iris/ReplicationSystem/ObjectReplicationBridge.h
+++ b/Engine/Source/Runtime/Experimental/Iris/Core/Public/Iris/ReplicationSystem/ObjectReplicationBridge.h
@@ -253,6 +253,8 @@ private:
 
 	FName GetConfigClassPathName(const UClass* Class);
 
+	void InitConditionalPropertyDelegates();
+
 private:
 
 	TMap<const UClass*, FName> ConfigClassPathNameCache;
@@ -294,6 +296,7 @@ private:
 	UE::Net::FNetObjectFilterHandle DefaultSpatialFilterHandle;
 
 	FDelegateHandle OnCustomConditionChangedHandle;
+	FDelegateHandle OnDynamicConditionChangedHandle;
 
 	bool bHasPollOverrides = false;
 	bool bHasDirtyClassesInPollPeriodOverrides = false;
diff --git a/Engine/Source/Runtime/Net/Core/Private/Net/Core/PropertyConditions/PropertyConditions.cpp b/Engine/Source/Runtime/Net/Core/Private/Net/Core/PropertyConditions/PropertyConditions.cpp
--- a/Engine/Source/Runtime/Net/Core/Private/Net/Core/PropertyConditions/PropertyConditions.cpp
+++ b/Engine/Source/Runtime/Net/Core/Private/Net/Core/PropertyConditions/PropertyConditions.cpp
@@ -14,6 +14,7 @@ namespace UE::Net::Private
 {
 #if UE_WITH_IRIS
 FPropertyConditionDelegates::FOnPropertyCustomConditionChanged FPropertyConditionDelegates::OnPropertyCustomConditionChangedDelegate;
+FPropertyConditionDelegates::FOnPropertyDynamicConditionChanged FPropertyConditionDelegates::OnPropertyDynamicConditionChangedDelegate;
 #endif // UE_WITH_IRIS
 
 FNetPropertyConditionManager::FNetPropertyConditionManager()
@@ -41,6 +42,15 @@ void FNetPropertyConditionManager::SetPropertyActive(const FObjectKey ObjectKey,
 	}
 }
 
+void FNetPropertyConditionManager::SetPropertyDynamicCondition(const FObjectKey ObjectKey, const uint16 RepIndex, const ELifetimeCondition Condition)
+{
+	TSharedPtr<FRepChangedPropertyTracker> Tracker = FindPropertyTracker(ObjectKey);
+	if (Tracker.IsValid())
+	{
+		Tracker->CallSetDynamicCondition(ObjectKey.ResolveObjectPtr(), RepIndex, Condition);
+	}
+}
+
 void FNetPropertyConditionManager::NotifyObjectDestroyed(const FObjectKey ObjectKey)
 {
 	PropertyTrackerMap.Remove(ObjectKey);
diff --git a/Engine/Source/Runtime/Net/Core/Private/Net/Core/PropertyConditions/RepChangedPropertyTracker.cpp b/Engine/Source/Runtime/Net/Core/Private/Net/Core/PropertyConditions/RepChangedPropertyTracker.cpp
--- a/Engine/Source/Runtime/Net/Core/Private/Net/Core/PropertyConditions/RepChangedPropertyTracker.cpp
+++ b/Engine/Source/Runtime/Net/Core/Private/Net/Core/PropertyConditions/RepChangedPropertyTracker.cpp
@@ -36,6 +36,31 @@ void FRepChangedPropertyTracker::SetCustomIsActiveOverride(UObject* OwningObject
 #endif // UE_WITH_IRIS
 }
 
+void FRepChangedPropertyTracker::CallSetDynamicCondition(const UObject* OwningObject, const uint16 RepIndex, const ELifetimeCondition Condition)
+{
+	using namespace UE::Net::Private;
+
+	const ELifetimeCondition OldCondition = ActiveState.GetDynamicCondition(RepIndex);
+	if (Condition == OldCondition)
+	{
+		return;
+	}
+
+	ActiveState.SetDynamicCondition(RepIndex, Condition);
+
+#if WITH_PUSH_MODEL
+	MARK_PROPERTY_DIRTY_UNSAFE(OwningObject, RepIndex);
+#endif
+
+#if UE_WITH_IRIS
+	FPropertyConditionDelegates::FOnPropertyDynamicConditionChanged& Delegate = FPropertyConditionDelegates::GetOnPropertyDynamicConditionChangedDelegate();
+	if (Delegate.IsBound())
+	{
+		Delegate.Broadcast(OwningObject, RepIndex, Condition);
+	}
+#endif
+}
+
 PRAGMA_DISABLE_DEPRECATION_WARNINGS
 void FRepChangedPropertyTracker::SetExternalData(const uint8* Src, const int32 NumBits)
 {
diff --git a/Engine/Source/Runtime/Net/Core/Public/Net/Core/PropertyConditions/PropertyConditions.h b/Engine/Source/Runtime/Net/Core/Public/Net/Core/PropertyConditions/PropertyConditions.h
--- a/Engine/Source/Runtime/Net/Core/Public/Net/Core/PropertyConditions/PropertyConditions.h
+++ b/Engine/Source/Runtime/Net/Core/Public/Net/Core/PropertyConditions/PropertyConditions.h
@@ -25,6 +25,8 @@ public:
 	static FNetPropertyConditionManager& Get();
 
 	void SetPropertyActive(const FObjectKey ObjectKey, const uint16 RepIndex, const bool bActive);
+	void SetPropertyDynamicCondition(const FObjectKey ObjectKey, const uint16 RepIndex, const ELifetimeCondition Condition);
+
 	void NotifyObjectDestroyed(const FObjectKey ObjectKey);
 
 	TSharedPtr<FRepChangedPropertyTracker> FindOrCreatePropertyTracker(const FObjectKey ObjectKey);
diff --git a/Engine/Source/Runtime/Net/Core/Public/Net/Core/PropertyConditions/PropertyConditionsDelegates.h b/Engine/Source/Runtime/Net/Core/Public/Net/Core/PropertyConditions/PropertyConditionsDelegates.h
--- a/Engine/Source/Runtime/Net/Core/Public/Net/Core/PropertyConditions/PropertyConditionsDelegates.h
+++ b/Engine/Source/Runtime/Net/Core/Public/Net/Core/PropertyConditions/PropertyConditionsDelegates.h
@@ -10,22 +10,23 @@
 
 #include "Delegates/Delegate.h"
 
-class FRepChangedPropertyTracker;
-
 namespace UE::Net::Private
 {
 
 class NETCORE_API FPropertyConditionDelegates
 {
 public:
-	DECLARE_MULTICAST_DELEGATE_ThreeParams(FOnPropertyCustomConditionChanged, const UObject*, const uint16, const bool);
+	using FOnPropertyCustomConditionChanged = TMulticastDelegate<void(const UObject* OwningObject, const uint16 RepIndex, const bool bActive)>;
+	using FOnPropertyDynamicConditionChanged = TMulticastDelegate<void(const UObject* OwningObject, const uint16 RepIndex, const ELifetimeCondition Condition)>;
 	
 	static FOnPropertyCustomConditionChanged& GetOnPropertyCustomConditionChangedDelegate() { return OnPropertyCustomConditionChangedDelegate; }
+	static FOnPropertyDynamicConditionChanged& GetOnPropertyDynamicConditionChangedDelegate() { return OnPropertyDynamicConditionChangedDelegate; }
 
 private:
 	static FOnPropertyCustomConditionChanged OnPropertyCustomConditionChangedDelegate;
+	static FOnPropertyDynamicConditionChanged OnPropertyDynamicConditionChangedDelegate;
 };
 
-}; // UE::Net::Private
+} // UE::Net::Private
 
-#endif
\ No newline at end of file
+#endif
diff --git a/Engine/Source/Runtime/Net/Core/Public/Net/Core/PropertyConditions/RepChangedPropertyTracker.h b/Engine/Source/Runtime/Net/Core/Public/Net/Core/PropertyConditions/RepChangedPropertyTracker.h
--- a/Engine/Source/Runtime/Net/Core/Public/Net/Core/PropertyConditions/RepChangedPropertyTracker.h
+++ b/Engine/Source/Runtime/Net/Core/Public/Net/Core/PropertyConditions/RepChangedPropertyTracker.h
@@ -73,7 +73,22 @@ public:
 		return ActiveState.GetNumProperties();
 	}
 
+	ELifetimeCondition GetDynamicCondition(uint16 ParentIndex) const
+	{
+		return ActiveState.GetDynamicCondition(ParentIndex);
+	}
+
+	uint32 GetDynamicConditionChangeCounter() const
+	{
+		return ActiveState.GetDynamicConditionChangeCounter();
+	}
+
 private:
+	friend UE::Net::Private::FNetPropertyConditionManager;
+
+	// Called from FNetPropertyConditionManager 
+	void CallSetDynamicCondition(const UObject* OwningObject, const uint16 RepIndex, const ELifetimeCondition Condition);
+
 	/** Activation data for top level Properties on the given Actor / Object. */
 	FCustomPropertyConditionState ActiveState;
 
