Fix some include paths in ReplicationSystemTestPlugin.

diff --git a/Engine/Plugins/Runtime/ReplicationSystemTestPlugin/Source/Private/Tests/ReplicationSystem/TestObjectReferences.cpp b/Engine/Plugins/Runtime/ReplicationSystemTestPlugin/Source/Private/Tests/ReplicationSystem/TestObjectReferences.cpp
--- a/Engine/Plugins/Runtime/ReplicationSystemTestPlugin/Source/Private/Tests/ReplicationSystem/TestObjectReferences.cpp
+++ b/Engine/Plugins/Runtime/ReplicationSystemTestPlugin/Source/Private/Tests/ReplicationSystem/TestObjectReferences.cpp
@@ -9,7 +9,7 @@
 #include "Iris/ReplicationSystem/ReplicationOperationsInternal.h"
 #include "Iris/ReplicationSystem/ReplicationSystemInternal.h"
 #include "Iris/ReplicationSystem/ReplicationFragmentUtil.h"
-#include "Engine/Public/Net/UnrealNetwork.h"
+#include "Net/UnrealNetwork.h"
 #include "Iris/ReplicationSystem/Filtering/NetObjectFilter.h"
 #include "Iris/Core/IrisDebugging.h"
 
diff --git a/Engine/Plugins/Runtime/ReplicationSystemTestPlugin/Source/Private/Tests/Serialization/TestPolymorphicStructNetSerializer.cpp b/Engine/Plugins/Runtime/ReplicationSystemTestPlugin/Source/Private/Tests/Serialization/TestPolymorphicStructNetSerializer.cpp
--- a/Engine/Plugins/Runtime/ReplicationSystemTestPlugin/Source/Private/Tests/Serialization/TestPolymorphicStructNetSerializer.cpp
+++ b/Engine/Plugins/Runtime/ReplicationSystemTestPlugin/Source/Private/Tests/Serialization/TestPolymorphicStructNetSerializer.cpp
@@ -11,7 +11,7 @@
 #include "Iris/ReplicationState/PropertyNetSerializerInfoRegistry.h"
 #include "Iris/Serialization/NetReferenceCollector.h"
 #include "Net/UnrealNetwork.h"
-#include "Core/Public/Iris/ReplicationSystem/ReplicationFragmentUtil.h"
+#include "Iris/ReplicationSystem/ReplicationFragmentUtil.h"
 
 namespace UE::Net
 {
diff --git a/Engine/Plugins/Runtime/ReplicationSystemTestPlugin/Source/Private/Tests/Serialization/TestPolymorphicStructNetSerializer.h b/Engine/Plugins/Runtime/ReplicationSystemTestPlugin/Source/Private/Tests/Serialization/TestPolymorphicStructNetSerializer.h
--- a/Engine/Plugins/Runtime/ReplicationSystemTestPlugin/Source/Private/Tests/Serialization/TestPolymorphicStructNetSerializer.h
+++ b/Engine/Plugins/Runtime/ReplicationSystemTestPlugin/Source/Private/Tests/Serialization/TestPolymorphicStructNetSerializer.h
@@ -5,8 +5,8 @@
 #include "Iris/Serialization/PolymorphicNetSerializer.h"
 #include "UObject/ObjectMacros.h"
 #include "Tests/ReplicationSystem/ReplicatedTestObject.h"
-#include "Core/Public/Iris/ReplicationState/IrisFastArraySerializer.h"
-#include "Core/Public/Iris/ReplicationState/Private/IrisFastArraySerializerInternal.h"
+#include "Iris/ReplicationState/IrisFastArraySerializer.h"
+#include "Iris/ReplicationState/Private/IrisFastArraySerializerInternal.h"
 #include "TestPolymorphicStructNetSerializer.generated.h"
 
 // Example BaseStruct used to support replication of polymorphic structs
