Fix some macros in ReplicationSystemTestPlugin.

diff --git a/Engine/Plugins/Runtime/ReplicationSystemTestPlugin/Source/Private/NetworkAutomationTest.h b/Engine/Plugins/Runtime/ReplicationSystemTestPlugin/Source/Private/NetworkAutomationTest.h
--- a/Engine/Plugins/Runtime/ReplicationSystemTestPlugin/Source/Private/NetworkAutomationTest.h
+++ b/Engine/Plugins/Runtime/ReplicationSystemTestPlugin/Source/Private/NetworkAutomationTest.h
@@ -3,6 +3,7 @@
 #pragma once
 
 #include "CoreTypes.h"
+#include "HAL/PreprocessorHelpers.h"
 #include "Misc/AutomationTest.h"
 
 #if UE_NET_WITH_LOW_LEVEL_TESTS
@@ -117,7 +118,7 @@ public:
 
 #define UE_NET_TEST_CATCH_INTERNAL(TestSuite, TestCase, BaseClass) \
 	constexpr TCHAR TestSuite ## _ ## TestCase ## Name[] = TEXT("Net." #TestSuite "." #TestCase); \
-	TEST_CASE_METHOD((TCatchFixtureWrapper<BaseClass, TestSuite ## _ ## TestCase ## Name>), #TestSuite ## "_" ## #TestCase)
+	TEST_CASE_METHOD((TCatchFixtureWrapper<BaseClass, TestSuite ## _ ## TestCase ## Name>), PREPROCESSOR_TO_STRING(TestSuite ## _ ## TestCase))
 
 #define UE_NET_TEST_FIXTURE(TestFixture, TestCase) UE_NET_TEST_CATCH_INTERNAL(TestFixture, TestCase, TestFixture)
 #define UE_NET_TEST(TestSuite, TestCase) UE_NET_TEST_CATCH_INTERNAL(TestSuite, TestCase, UE::Net::FNetworkAutomationTestSuiteFixture)
diff --git a/Engine/Plugins/Runtime/ReplicationSystemTestPlugin/Source/Private/NetworkAutomationTestMacros.h b/Engine/Plugins/Runtime/ReplicationSystemTestPlugin/Source/Private/NetworkAutomationTestMacros.h
--- a/Engine/Plugins/Runtime/ReplicationSystemTestPlugin/Source/Private/NetworkAutomationTestMacros.h
+++ b/Engine/Plugins/Runtime/ReplicationSystemTestPlugin/Source/Private/NetworkAutomationTestMacros.h
@@ -155,6 +155,9 @@ Name(const T& Value, U ExpectedValue, const char* ValueText, const char* Expecte
 #define UE_NET_EXPECT_TRUE_MSG(V, Message) UE_NET_TEST_MSG_INTERNAL(Message); CHECK(V)
 #define UE_NET_EXPECT_FALSE_MSG(V, Message) UE_NET_TEST_MSG_INTERNAL(Message); CHECK_FALSE(V)
 
+#define UE_NET_FAIL(Message) FAIL(Message)
+#define UE_NET_WARN(Message) FAIL_CHECK(Message)
+
 #else
 
 UE_NET_COMPARE_RETURN_TESTRESULT(TCmpEqual, ==)
