Fix ensures about fragmentless netobjects in ReplicationSystemTests

diff --git a/Engine/Plugins/Runtime/ReplicationSystemTestPlugin/Source/Private/Tests/ReplicationSystem/ReplicatedTestObject.cpp b/Engine/Plugins/Runtime/ReplicationSystemTestPlugin/Source/Private/Tests/ReplicationSystem/ReplicatedTestObject.cpp
--- a/Engine/Plugins/Runtime/ReplicationSystemTestPlugin/Source/Private/Tests/ReplicationSystem/ReplicatedTestObject.cpp
+++ b/Engine/Plugins/Runtime/ReplicationSystemTestPlugin/Source/Private/Tests/ReplicationSystem/ReplicatedTestObject.cpp
@@ -448,6 +448,11 @@ UTestReplicatedIrisObjectWithNoReplicatedMembers::UTestReplicatedIrisObjectWithN
 {
 }
 
+void UTestReplicatedIrisObjectWithNoReplicatedMembers::RegisterReplicationFragments(UE::Net::FFragmentRegistrationContext& Fragments, UE::Net::EFragmentRegistrationFlags RegistrationFlags)
+{
+	Fragments.SetIsFragmentlessNetObject(true);
+}
+
 void UReplicatedSubObjectOrderObject::GetLifetimeReplicatedProps(TArray< class FLifetimeProperty >& OutLifetimeProps) const
 {
 	DOREPLIFETIME(UReplicatedSubObjectOrderObject, IntA);
diff --git a/Engine/Plugins/Runtime/ReplicationSystemTestPlugin/Source/Private/Tests/ReplicationSystem/ReplicatedTestObject.h b/Engine/Plugins/Runtime/ReplicationSystemTestPlugin/Source/Private/Tests/ReplicationSystem/ReplicatedTestObject.h
--- a/Engine/Plugins/Runtime/ReplicationSystemTestPlugin/Source/Private/Tests/ReplicationSystem/ReplicatedTestObject.h
+++ b/Engine/Plugins/Runtime/ReplicationSystemTestPlugin/Source/Private/Tests/ReplicationSystem/ReplicatedTestObject.h
@@ -420,6 +420,8 @@ class UTestReplicatedIrisObjectWithNoReplicatedMembers : public UReplicatedTestO
 
 public:
 	UTestReplicatedIrisObjectWithNoReplicatedMembers();
+
+	virtual void RegisterReplicationFragments(UE::Net::FFragmentRegistrationContext& Fragments, UE::Net::EFragmentRegistrationFlags RegistrationFlags) override;
 };
 
 /**
diff --git a/Engine/Plugins/Runtime/ReplicationSystemTestPlugin/Source/Private/Tests/Serialization/ScriptInterfaceTestTypes.cpp b/Engine/Plugins/Runtime/ReplicationSystemTestPlugin/Source/Private/Tests/Serialization/ScriptInterfaceTestTypes.cpp
--- a/Engine/Plugins/Runtime/ReplicationSystemTestPlugin/Source/Private/Tests/Serialization/ScriptInterfaceTestTypes.cpp
+++ b/Engine/Plugins/Runtime/ReplicationSystemTestPlugin/Source/Private/Tests/Serialization/ScriptInterfaceTestTypes.cpp
@@ -20,3 +20,13 @@ void UTestObjectReferencingScriptInterface::GetLifetimeReplicatedProps(TArray<FL
 
 	DOREPLIFETIME(ThisClass, InterfaceObject);
 }
+
+void UTestScriptInterfaceReplicatedObject::RegisterReplicationFragments(UE::Net::FFragmentRegistrationContext& Fragments, UE::Net::EFragmentRegistrationFlags RegistrationFlags)
+{
+	Fragments.SetIsFragmentlessNetObject(true);
+}
+
+void UTestScriptInterfaceReplicatedObjectWithDefaultSubobject::RegisterReplicationFragments(UE::Net::FFragmentRegistrationContext& Fragments, UE::Net::EFragmentRegistrationFlags RegistrationFlags)
+{
+	Fragments.SetIsFragmentlessNetObject(true);
+}
\ No newline at end of file
diff --git a/Engine/Plugins/Runtime/ReplicationSystemTestPlugin/Source/Private/Tests/Serialization/ScriptInterfaceTestTypes.h b/Engine/Plugins/Runtime/ReplicationSystemTestPlugin/Source/Private/Tests/Serialization/ScriptInterfaceTestTypes.h
--- a/Engine/Plugins/Runtime/ReplicationSystemTestPlugin/Source/Private/Tests/Serialization/ScriptInterfaceTestTypes.h
+++ b/Engine/Plugins/Runtime/ReplicationSystemTestPlugin/Source/Private/Tests/Serialization/ScriptInterfaceTestTypes.h
@@ -42,6 +42,8 @@ class UTestScriptInterfaceReplicatedObject : public UReplicatedTestObject, publi
 public:
 	// IIrisTestInterface
 	virtual bool IsInterfaceFunctionOverridden() const { return true; }
+
+	virtual void RegisterReplicationFragments(UE::Net::FFragmentRegistrationContext& Fragments, UE::Net::EFragmentRegistrationFlags RegistrationFlags) override;
 };
 
 UCLASS()
@@ -52,6 +54,8 @@ class UTestScriptInterfaceReplicatedObjectWithDefaultSubobject : public UReplica
 public:
 	UTestScriptInterfaceReplicatedObjectWithDefaultSubobject(const FObjectInitializer& ObjectInitializer = FObjectInitializer::Get());
 
+	virtual void RegisterReplicationFragments(UE::Net::FFragmentRegistrationContext& Fragments, UE::Net::EFragmentRegistrationFlags RegistrationFlags) override;
+
 public:
 	TObjectPtr<UObject> DefaultSubobjectWithInterface;
 };
