Fix or suppress some warnings and errors in ReplicationSystemTestPlugin.

diff --git a/Engine/Plugins/Runtime/ReplicationSystemTestPlugin/Source/Private/Tests/DataStream/TestDataStream.cpp b/Engine/Plugins/Runtime/ReplicationSystemTestPlugin/Source/Private/Tests/DataStream/TestDataStream.cpp
--- a/Engine/Plugins/Runtime/ReplicationSystemTestPlugin/Source/Private/Tests/DataStream/TestDataStream.cpp
+++ b/Engine/Plugins/Runtime/ReplicationSystemTestPlugin/Source/Private/Tests/DataStream/TestDataStream.cpp
@@ -1,8 +1,10 @@
 // Copyright Epic Games, Inc. All Rights Reserved.
 
+#include "Logging/LogScopedVerbosityOverride.h"
 #include "NetworkAutomationTest.h"
 #include "NetworkAutomationTestMacros.h"
 #include "MockDataStream.h"
+#include "Iris/Core/IrisLog.h"
 #include "Iris/DataStream/DataStreamManager.h"
 #include "Iris/DataStream/DataStreamDefinitions.h"
 #include "Iris/DataStream/DataStreamManager.h"
@@ -72,6 +74,9 @@ UE_NET_TEST_FIXTURE(FTestDataStream, CannotCreateSameDataStreamTwice)
 	AddMockDataStreamDefinition(bAddValidDefinition);
 
 	DataStreamManager->CreateStream("Mock");
+	
+	// Suppress Iris internal warning, since we're intentionally creating duplicate streams.
+	LOG_SCOPE_VERBOSITY_OVERRIDE(LogIris, ELogVerbosity::Fatal);
 	ECreateDataStreamResult Result = DataStreamManager->CreateStream("Mock");
 	UE_NET_ASSERT_EQ(unsigned(Result), unsigned(ECreateDataStreamResult::Error_Duplicate));
 }
@@ -79,7 +84,12 @@ UE_NET_TEST_FIXTURE(FTestDataStream, CannotCreateSameDataStreamTwice)
 UE_NET_TEST_FIXTURE(FTestDataStream, CannotCreateInvalidDataStream)
 {
 	constexpr bool bAddValidDefinition = false;
-	AddMockDataStreamDefinition(bAddValidDefinition);
+
+	// Suppress Iris internal error, since we're intentionally creating an invalid stream.
+	{
+		LOG_SCOPE_VERBOSITY_OVERRIDE(LogIris, ELogVerbosity::Fatal);
+		AddMockDataStreamDefinition(bAddValidDefinition);
+	}
 
 	ECreateDataStreamResult Result = DataStreamManager->CreateStream("Mock");
 	UE_NET_ASSERT_EQ(unsigned(Result), unsigned(ECreateDataStreamResult::Error_InvalidDefinition));
diff --git a/Engine/Plugins/Runtime/ReplicationSystemTestPlugin/Source/Private/Tests/ReplicationState/TestFastArrayReplicationState.cpp b/Engine/Plugins/Runtime/ReplicationSystemTestPlugin/Source/Private/Tests/ReplicationState/TestFastArrayReplicationState.cpp
--- a/Engine/Plugins/Runtime/ReplicationSystemTestPlugin/Source/Private/Tests/ReplicationState/TestFastArrayReplicationState.cpp
+++ b/Engine/Plugins/Runtime/ReplicationSystemTestPlugin/Source/Private/Tests/ReplicationState/TestFastArrayReplicationState.cpp
@@ -4,6 +4,7 @@
 #include "Tests/ReplicationSystem/ReplicationSystemServerClientTestFixture.h"
 #include "NetworkAutomationTest.h"
 #include "NetworkAutomationTestMacros.h"
+#include "Iris/Core/IrisLog.h"
 #include "Iris/ReplicationState/InternalReplicationStateDescriptorUtils.h"
 #include "Iris/ReplicationState/ReplicationStateDescriptorBuilder.h"
 #include "Iris/ReplicationState/ReplicationStateUtil.h"
@@ -11,6 +12,7 @@
 #include "Iris/ReplicationSystem/ReplicationFragmentUtil.h"
 #include "Iris/Serialization/NetSerializers.h"
 #include "Iris/Serialization/InternalNetSerializerUtils.h"
+#include "Logging/LogScopedVerbosityOverride.h"
 #include "Misc/EnumClassFlags.h"
 #include "Net/UnrealNetwork.h"
 #include "Iris/ReplicationSystem/Filtering/NetObjectFilter.h"
@@ -187,7 +189,10 @@ UE_NET_TEST_FIXTURE(FTestFastArrayReplicationStateContext, BuildClassDescriptorF
 UE_NET_TEST_FIXTURE(FTestFastArrayReplicationStateContext, BuildClassDescriptorForSinglePropertyForFastArrayWithExtraProperty)
 {
 	// We expect this test to fail as we do not pass the flag allowing fastarrays to contain extra properties
-	InitDescriptor(UTestFastArrayReplicationState_FastArray_TestClassFastArrayWithExtraProperty::StaticClass(), (int32)UTestFastArrayReplicationState_FastArray_TestClassFastArrayWithExtraProperty::ENetFields_Private::FastArray);
+	{
+		LOG_SCOPE_VERBOSITY_OVERRIDE(LogIris, ELogVerbosity::Fatal);
+		InitDescriptor(UTestFastArrayReplicationState_FastArray_TestClassFastArrayWithExtraProperty::StaticClass(), (int32)UTestFastArrayReplicationState_FastArray_TestClassFastArrayWithExtraProperty::ENetFields_Private::FastArray);
+	}
 	UE_NET_ASSERT_FALSE(EnumHasAnyFlags(Descriptor->Traits, EReplicationStateTraits::IsFastArrayReplicationState));	
 }
 
diff --git a/Engine/Plugins/Runtime/ReplicationSystemTestPlugin/Source/Private/Tests/ReplicationState/TestPropertyReplicationState.cpp b/Engine/Plugins/Runtime/ReplicationSystemTestPlugin/Source/Private/Tests/ReplicationState/TestPropertyReplicationState.cpp
--- a/Engine/Plugins/Runtime/ReplicationSystemTestPlugin/Source/Private/Tests/ReplicationState/TestPropertyReplicationState.cpp
+++ b/Engine/Plugins/Runtime/ReplicationSystemTestPlugin/Source/Private/Tests/ReplicationState/TestPropertyReplicationState.cpp
@@ -16,6 +16,9 @@
 
 void UTestPropertyReplicationState_TestClass::GetLifetimeReplicatedProps( TArray< class FLifetimeProperty > & OutLifetimeProps ) const
 {
+	DOREPLIFETIME(ThisClass, IntA);
+	DOREPLIFETIME(ThisClass, IntB);
+	DOREPLIFETIME(ThisClass, IntC);
 }
 
 void UTestPropertyReplicationState_TestClassWithRepNotify::OnRep_IntA(int32 OldInt)
@@ -28,16 +31,17 @@ void UTestPropertyReplicationState_TestClassWithRepNotify::OnRep_IntB(int32 OldI
 
 void UTestPropertyReplicationState_TestClassWithRepNotify::GetLifetimeReplicatedProps( TArray< class FLifetimeProperty > & OutLifetimeProps ) const
 {
-	DOREPLIFETIME_CONDITION_NOTIFY(UTestPropertyReplicationState_TestClassWithRepNotify, IntA, COND_None, REPNOTIFY_Always);
-	DOREPLIFETIME_CONDITION_NOTIFY(UTestPropertyReplicationState_TestClassWithRepNotify, IntB, COND_None, REPNOTIFY_OnChanged);
+	DOREPLIFETIME_CONDITION_NOTIFY(ThisClass, IntA, COND_None, REPNOTIFY_Always);
+	DOREPLIFETIME_CONDITION_NOTIFY(ThisClass, IntB, COND_None, REPNOTIFY_OnChanged);
 }
 
 void UTestPropertyReplicationState_TestClassWithInitAndCArrays::GetLifetimeReplicatedProps(TArray<class FLifetimeProperty>& OutLifetimeProps) const
 {
-	DOREPLIFETIME_CONDITION(UTestPropertyReplicationState_TestClassWithInitAndCArrays, InitArrayOfFullyReplicatedStruct, COND_InitialOnly);
-	DOREPLIFETIME_CONDITION(UTestPropertyReplicationState_TestClassWithInitAndCArrays, InitArrayOfNotFullyReplicatedStruct, COND_InitialOnly);
-	DOREPLIFETIME_CONDITION(UTestPropertyReplicationState_TestClassWithInitAndCArrays, ArrayOfFullyReplicatedStruct, COND_None);
-	DOREPLIFETIME_CONDITION(UTestPropertyReplicationState_TestClassWithInitAndCArrays, ArrayOfNotFullyReplicatedStruct, COND_None);
+	DOREPLIFETIME_CONDITION(ThisClass, InitArrayOfFullyReplicatedStruct, COND_InitialOnly);
+	DOREPLIFETIME_CONDITION(ThisClass, InitArrayOfNotFullyReplicatedStruct, COND_InitialOnly);
+	DOREPLIFETIME_CONDITION(ThisClass, ArrayOfFullyReplicatedStruct, COND_None);
+	DOREPLIFETIME_CONDITION(ThisClass, ArrayOfNotFullyReplicatedStruct, COND_None);
+	DOREPLIFETIME(ThisClass, StructWithArrayOfNotFullyReplicatedStruct);
 }
 
 namespace UE::Net::Private
diff --git a/Engine/Plugins/Runtime/ReplicationSystemTestPlugin/Source/Private/Tests/ReplicationState/TestReplicationStateDescriptorBuilder.cpp b/Engine/Plugins/Runtime/ReplicationSystemTestPlugin/Source/Private/Tests/ReplicationState/TestReplicationStateDescriptorBuilder.cpp
--- a/Engine/Plugins/Runtime/ReplicationSystemTestPlugin/Source/Private/Tests/ReplicationState/TestReplicationStateDescriptorBuilder.cpp
+++ b/Engine/Plugins/Runtime/ReplicationSystemTestPlugin/Source/Private/Tests/ReplicationState/TestReplicationStateDescriptorBuilder.cpp
@@ -1182,79 +1182,127 @@ UE_NET_TEST_FIXTURE(FTestReplicationStateDescriptorBuilderFixture, UTestReplicat
 // Function definitions
 void UTestReplicationStateDescriptor_TestClass::GetLifetimeReplicatedProps( TArray< class FLifetimeProperty > & OutLifetimeProps ) const
 {
+	DOREPLIFETIME(ThisClass, Int8);
+	DOREPLIFETIME(ThisClass, Int16);
+	DOREPLIFETIME(ThisClass, Int32);
+	DOREPLIFETIME(ThisClass, Int64);
+	DOREPLIFETIME(ThisClass, bBitFieldBoolA);
+	DOREPLIFETIME(ThisClass, bBitFieldBoolB);
+	DOREPLIFETIME(ThisClass, bNativeBool);
 }
 
 void UTestReplicationStateDescriptor_TestClassWithNonReplicatedData::GetLifetimeReplicatedProps( TArray< class FLifetimeProperty > & OutLifetimeProps ) const
 {
+	DOREPLIFETIME(ThisClass, Int8);
+	DOREPLIFETIME(ThisClass, Int16);
+	DOREPLIFETIME(ThisClass, Int32);
+	DOREPLIFETIME(ThisClass, Int64);
+	DOREPLIFETIME(ThisClass, bBitFieldBoolA);
+	DOREPLIFETIME(ThisClass, bBitFieldBoolB);
+	DOREPLIFETIME(ThisClass, bNativeBool);
 }
 
 void UTestReplicationStateDescriptor_TestClassWithInheritance::GetLifetimeReplicatedProps( TArray< class FLifetimeProperty > & OutLifetimeProps ) const
 {
+	Super::GetLifetimeReplicatedProps(OutLifetimeProps);
+
+	DOREPLIFETIME(ThisClass, UInt8);
+	DOREPLIFETIME(ThisClass, UInt16);
 }
 
 void UTestReplicationStateDescriptor_TestClassWithReplicatedStruct::GetLifetimeReplicatedProps( TArray< class FLifetimeProperty > & OutLifetimeProps ) const
 {
+	DOREPLIFETIME(ThisClass, ReplicatedStruct);
 }
 
 void UTestReplicationStateDescriptor_TestClassWithTArray::GetLifetimeReplicatedProps( TArray< class FLifetimeProperty > & OutLifetimeProps ) const
 {
+	DOREPLIFETIME(ThisClass, ArrayWithStruct);
+	DOREPLIFETIME(ThisClass, ArrayWithPrimitiveType);
 }
 
 void UTestReplicationStateDescriptor_TestClassWithStructWithTArray::GetLifetimeReplicatedProps( TArray< class FLifetimeProperty > & OutLifetimeProps ) const
 {
+	DOREPLIFETIME(ThisClass, StructWithTArray);
 }
 
 void UTestReplicationStateDescriptor_TestClassWithCArray::GetLifetimeReplicatedProps(TArray<class FLifetimeProperty> & OutLifetimeProps) const
 {
+	DOREPLIFETIME(ThisClass, ArrayOfStruct);
+	DOREPLIFETIME(ThisClass, ArrayWithPrimitiveType);
 }
 
 void UTestReplicationStateDescriptor_TestClassWithStructWithCArray::GetLifetimeReplicatedProps(TArray<class FLifetimeProperty> & OutLifetimeProps) const
 {
+	DOREPLIFETIME(ThisClass, StructWithCArray);
 }
 
 void UTestReplicationStateDescriptor_TestClassWithEnums::GetLifetimeReplicatedProps(TArray<class FLifetimeProperty> & OutLifetimeProps) const
 {
+	DOREPLIFETIME(ThisClass, Int8Enum);
+	DOREPLIFETIME(ThisClass, Int16Enum);
+	DOREPLIFETIME(ThisClass, Int32Enum);
+	DOREPLIFETIME(ThisClass, Int64Enum);
+	DOREPLIFETIME(ThisClass, EnumAsByteEnum);
+	DOREPLIFETIME(ThisClass, Uint8Enum);
+	DOREPLIFETIME(ThisClass, Uint16Enum);
+	DOREPLIFETIME(ThisClass, Uint32Enum);
+	DOREPLIFETIME(ThisClass, Uint64Enum);
 }
 
 void UTestReplicationStateDescriptor_TestClassWithRoleAndRemoteRole::GetLifetimeReplicatedProps(TArray<class FLifetimeProperty> & OutLifetimeProps) const
 {
+	DOREPLIFETIME(ThisClass, Role);
+	DOREPLIFETIME(ThisClass, RemoteRole);
 }
 
 void UTestReplicationStateDescriptor_TestClassWithManyRoles::GetLifetimeReplicatedProps(TArray<class FLifetimeProperty> & OutLifetimeProps) const
 {
+	DOREPLIFETIME(ThisClass, Role);
+	DOREPLIFETIME(ThisClass, OtherRole);
+	DOREPLIFETIME(ThisClass, YetAnotherRole);
 }
 
 void UTestReplicationStateDescriptor_TestClassWithRPCs::GetLifetimeReplicatedProps(TArray<class FLifetimeProperty> & OutLifetimeProps) const
 {
+	DOREPLIFETIME(ThisClass, Int);
 }
 
 void UTestReplicationStateDescriptor_TestClassWithRef::GetLifetimeReplicatedProps(TArray<class FLifetimeProperty> & OutLifetimeProps) const
 {
+	DOREPLIFETIME(ThisClass, ObjectRef);
 }
 
 void UTestReplicationStateDescriptor_TestClassWithRefInStruct::GetLifetimeReplicatedProps(TArray<class FLifetimeProperty> & OutLifetimeProps) const
 {
+	DOREPLIFETIME(ThisClass, StructWithRef);
 }
 
 void UTestReplicationStateDescriptor_TestClassWithRefInStructWithNestedCArray::GetLifetimeReplicatedProps(TArray<class FLifetimeProperty> & OutLifetimeProps) const
 {
+	DOREPLIFETIME(ThisClass, StructWithRef);
 }
 
 void UTestReplicationStateDescriptor_TestClassWithRefInStructWithNestedTArray::GetLifetimeReplicatedProps(TArray<class FLifetimeProperty> & OutLifetimeProps) const
 {
+	DOREPLIFETIME(ThisClass, StructWithRef);
 }
 
 void UTestReplicationStateDescriptor_TestClassWithRefInCArray::GetLifetimeReplicatedProps(TArray<class FLifetimeProperty> & OutLifetimeProps) const
 {
+	DOREPLIFETIME(ThisClass, ObjectRef);
 }
 void UTestReplicationStateDescriptor_TestClassWithRefInTArray::GetLifetimeReplicatedProps(TArray<class FLifetimeProperty> & OutLifetimeProps) const
 {
+	DOREPLIFETIME(ThisClass, ObjectRefInArray);
 }
 void UTestReplicationStateDescriptor_TestClassWithRefInStructCArray::GetLifetimeReplicatedProps(TArray<class FLifetimeProperty> & OutLifetimeProps) const
 {
+	DOREPLIFETIME(ThisClass, StructWithRef);
 }
 void UTestReplicationStateDescriptor_TestClassWithRefInStructTArray::GetLifetimeReplicatedProps(TArray<class FLifetimeProperty> & OutLifetimeProps) const
 {
+	DOREPLIFETIME(ThisClass, StructWithRef);
 }
 
 void UTestReplicationStateDescriptor_TestClassWithRPCs::UnreliableRPCOnServerWithPrimitiveArgs_Implementation(bool bBool, int InInt)
diff --git a/Engine/Plugins/Runtime/ReplicationSystemTestPlugin/Source/Private/Tests/ReplicationSystem/ReplicatedTestObject.cpp b/Engine/Plugins/Runtime/ReplicationSystemTestPlugin/Source/Private/Tests/ReplicationSystem/ReplicatedTestObject.cpp
--- a/Engine/Plugins/Runtime/ReplicationSystemTestPlugin/Source/Private/Tests/ReplicationSystem/ReplicatedTestObject.cpp
+++ b/Engine/Plugins/Runtime/ReplicationSystemTestPlugin/Source/Private/Tests/ReplicationSystem/ReplicatedTestObject.cpp
@@ -425,6 +425,11 @@ void UTestReplicatedIrisObject::RegisterReplicationFragments(UE::Net::FFragmentR
 
 void UTestReplicatedIrisObjectWithObjectReference::GetLifetimeReplicatedProps( TArray< class FLifetimeProperty > & OutLifetimeProps ) const
 {
+	DOREPLIFETIME(ThisClass, IntA);
+	DOREPLIFETIME(ThisClass, IntB);
+	DOREPLIFETIME(ThisClass, IntC);
+	DOREPLIFETIME(ThisClass, RawObjectPtrRef);
+	DOREPLIFETIME(ThisClass, WeakObjectPtrObjectRef);
 }
 
 UTestReplicatedIrisObjectWithObjectReference::UTestReplicatedIrisObjectWithObjectReference()
diff --git a/Engine/Plugins/Runtime/ReplicationSystemTestPlugin/Source/Private/Tests/ReplicationSystem/TestObjectReferences.cpp b/Engine/Plugins/Runtime/ReplicationSystemTestPlugin/Source/Private/Tests/ReplicationSystem/TestObjectReferences.cpp
--- a/Engine/Plugins/Runtime/ReplicationSystemTestPlugin/Source/Private/Tests/ReplicationSystem/TestObjectReferences.cpp
+++ b/Engine/Plugins/Runtime/ReplicationSystemTestPlugin/Source/Private/Tests/ReplicationSystem/TestObjectReferences.cpp
@@ -31,6 +31,7 @@ void UTestObjectReferences_TestClassWithReferences::GetLifetimeReplicatedProps(T
 	DOREPLIFETIME(UTestObjectReferences_TestClassWithReferences, StructWithRef_CArray);
 	DOREPLIFETIME(UTestObjectReferences_TestClassWithReferences, StructWithRef_TArray);
 	DOREPLIFETIME(UTestObjectReferences_TestClassWithReferences, TestStructWithNestedRefTArray_TArray);
+	DOREPLIFETIME(UTestObjectReferences_TestClassWithReferences, TestStructWithNestedRefCArray_CArray);
 	DOREPLIFETIME(UTestObjectReferences_TestClassWithReferences, Ref_FastArray);
 	DOREPLIFETIME(UTestObjectReferences_TestClassWithReferences, Ref_NativeFastArray);
 }
diff --git a/Engine/Plugins/Runtime/ReplicationSystemTestPlugin/Source/Private/Tests/Serialization/TestNetBitStreamUtil.cpp b/Engine/Plugins/Runtime/ReplicationSystemTestPlugin/Source/Private/Tests/Serialization/TestNetBitStreamUtil.cpp
--- a/Engine/Plugins/Runtime/ReplicationSystemTestPlugin/Source/Private/Tests/Serialization/TestNetBitStreamUtil.cpp
+++ b/Engine/Plugins/Runtime/ReplicationSystemTestPlugin/Source/Private/Tests/Serialization/TestNetBitStreamUtil.cpp
@@ -3,9 +3,11 @@
 #include "NetworkAutomationTest.h"
 #include "NetworkAutomationTestMacros.h"
 #include "Containers/StringConv.h"
+#include "CoreGlobals.h"
 #include "Iris/Serialization/NetBitStreamReader.h"
 #include "Iris/Serialization/NetBitStreamWriter.h"
 #include "Iris/Serialization/NetBitStreamUtil.h"
+#include "Logging/LogScopedVerbosityOverride.h"
 
 namespace UE::Net::Private
 {
@@ -92,7 +94,13 @@ UE_NET_TEST_FIXTURE(FNetBitStreamUtilTest, TestTooLongStringWritesEmptyString)
 	TArray<uint8> VeryLargeBuffer;
 	VeryLargeBuffer.SetNumUninitialized(Align(VeryLongStringLength + 1024, 4));
 	Writer.InitBytes(VeryLargeBuffer.GetData(), VeryLargeBuffer.Num());
-	WriteString(&Writer, VeryLongString);
+	
+	// Suppress Iris internal error, since we're intentionally causing one.
+	{
+		LOG_SCOPE_VERBOSITY_OVERRIDE(LogSerialization, ELogVerbosity::Fatal);
+		WriteString(&Writer, VeryLongString);
+	}
+
 	Writer.CommitWrites();
 	UE_NET_ASSERT_FALSE(Writer.IsOverflown());
 	UE_NET_ASSERT_GT(Writer.GetPosBits(), 0U);
diff --git a/Engine/Plugins/Runtime/ReplicationSystemTestPlugin/Source/Private/Tests/Serialization/TestNetRoleNetSerializer.cpp b/Engine/Plugins/Runtime/ReplicationSystemTestPlugin/Source/Private/Tests/Serialization/TestNetRoleNetSerializer.cpp
--- a/Engine/Plugins/Runtime/ReplicationSystemTestPlugin/Source/Private/Tests/Serialization/TestNetRoleNetSerializer.cpp
+++ b/Engine/Plugins/Runtime/ReplicationSystemTestPlugin/Source/Private/Tests/Serialization/TestNetRoleNetSerializer.cpp
@@ -10,10 +10,14 @@
 #include "Iris/Serialization/NetSerializationContext.h"
 #include "Engine/EngineTypes.h"
 #include "EnumTestTypes.h"
+#include "Net/UnrealNetwork.h"
 #include "UObject/StrongObjectPtr.h"
 
 void UClassWithNetRoleSwapping::GetLifetimeReplicatedProps( TArray< class FLifetimeProperty > & OutLifetimeProps ) const
 {
+	DOREPLIFETIME(ThisClass, Role);
+	DOREPLIFETIME(ThisClass, JustSomeReplicatedPropertyInBetween);
+	DOREPLIFETIME(ThisClass, RemoteRole);
 }
 
 namespace UE::Net::Private
