Re-enable LastResortPropertyNetSerializer tests on non-Windows platforms since they're passing now.

diff --git a/Engine/Plugins/Runtime/ReplicationSystemTestPlugin/Source/Private/Tests/Serialization/TestLastResortPropertyNetSerializer.cpp b/Engine/Plugins/Runtime/ReplicationSystemTestPlugin/Source/Private/Tests/Serialization/TestLastResortPropertyNetSerializer.cpp
--- a/Engine/Plugins/Runtime/ReplicationSystemTestPlugin/Source/Private/Tests/Serialization/TestLastResortPropertyNetSerializer.cpp
+++ b/Engine/Plugins/Runtime/ReplicationSystemTestPlugin/Source/Private/Tests/Serialization/TestLastResortPropertyNetSerializer.cpp
@@ -60,14 +60,6 @@ const FText FTestLastResortPropertyNetSerializer::TestValues[] =
 
 const SIZE_T FTestLastResortPropertyNetSerializer::TestValueCount = sizeof(TestValues)/sizeof(TestValues[0]);
 
-// Having some issues with localization files not being present on most platforms
-
-#if !PLATFORM_WINDOWS
-UE_NET_TEST_FIXTURE(FTestLastResortPropertyNetSerializer, WarnNotTested)
-{
-	UE_NET_WARN("LastResortPropertyNetSerializer cannot be tested without localization files.");
-}
-#else
 UE_NET_TEST_FIXTURE(FTestLastResortPropertyNetSerializer, TestQuantize)
 {
 	TestQuantize();
@@ -87,7 +79,6 @@ UE_NET_TEST_FIXTURE(FTestLastResortPropertyNetSerializer, TestCloneDynamicState)
 {
 	TestCloneDynamicState();
 }
-#endif
 
 void FTestLastResortPropertyNetSerializer::SetUp()
 {
