Use std::numeric_limits

#rnx

diff --git a/Engine/Plugins/Runtime/ReplicationSystemTestPlugin/Source/Private/Tests/Core/TestBitTwiddling.cpp b/Engine/Plugins/Runtime/ReplicationSystemTestPlugin/Source/Private/Tests/Core/TestBitTwiddling.cpp
--- a/Engine/Plugins/Runtime/ReplicationSystemTestPlugin/Source/Private/Tests/Core/TestBitTwiddling.cpp
+++ b/Engine/Plugins/Runtime/ReplicationSystemTestPlugin/Source/Private/Tests/Core/TestBitTwiddling.cpp
@@ -3,7 +3,7 @@
 #include "NetworkAutomationTest.h"
 #include "NetworkAutomationTestMacros.h"
 #include "Iris/Core/BitTwiddling.h"
-#include "Math/NumericLimits.h"
+#include <limits>
 
 namespace UE::Net::Private
 {
@@ -11,118 +11,118 @@ namespace UE::Net::Private
 // GetBitsNeeded
 UE_NET_TEST(TestGetBitsNeeded, ReturnsExpectedValuesForInt8)
 {
-	UE_NET_ASSERT_EQ(GetBitsNeeded(TNumericLimits<int8>::Lowest()), 8U);
+	UE_NET_ASSERT_EQ(GetBitsNeeded(std::numeric_limits<int8>::min()), 8U);
 	UE_NET_ASSERT_EQ(GetBitsNeeded(int8(-64)), 7U);
 	UE_NET_ASSERT_EQ(GetBitsNeeded(int8(-1)), 1U);
 	UE_NET_ASSERT_EQ(GetBitsNeeded(int8(0)), 1U);
 	UE_NET_ASSERT_EQ(GetBitsNeeded(int8(1)), 2U);
-	UE_NET_ASSERT_EQ(GetBitsNeeded(TNumericLimits<int8>::Max()), 8U);
+	UE_NET_ASSERT_EQ(GetBitsNeeded(std::numeric_limits<int8>::max()), 8U);
 }
 
 UE_NET_TEST(TestGetBitsNeeded, ReturnsExpectedValuesForInt32)
 {
-	UE_NET_ASSERT_EQ(GetBitsNeeded(TNumericLimits<int32>::Lowest()), 32U);
+	UE_NET_ASSERT_EQ(GetBitsNeeded(std::numeric_limits<int32>::min()), 32U);
 	UE_NET_ASSERT_EQ(GetBitsNeeded(int32(-64)), 7U);
 	UE_NET_ASSERT_EQ(GetBitsNeeded(int32(-1)), 1U);
 	UE_NET_ASSERT_EQ(GetBitsNeeded(int32(0)), 1U);
 	UE_NET_ASSERT_EQ(GetBitsNeeded(int32(1)), 2U);
-	UE_NET_ASSERT_EQ(GetBitsNeeded(TNumericLimits<int32>::Max()), 32U);
+	UE_NET_ASSERT_EQ(GetBitsNeeded(std::numeric_limits<int32>::max()), 32U);
 }
 
 UE_NET_TEST(TestGetBitsNeeded, ReturnsExpectedValuesForInt64)
 {
-	UE_NET_ASSERT_EQ(GetBitsNeeded(TNumericLimits<int64>::Lowest()), 64U);
+	UE_NET_ASSERT_EQ(GetBitsNeeded(std::numeric_limits<int64>::min()), 64U);
 	UE_NET_ASSERT_EQ(GetBitsNeeded(int64(-64)), 7U);
 	UE_NET_ASSERT_EQ(GetBitsNeeded(int64(-1)), 1U);
 	UE_NET_ASSERT_EQ(GetBitsNeeded(int64(0)), 1U);
 	UE_NET_ASSERT_EQ(GetBitsNeeded(int64(1)), 2U);
-	UE_NET_ASSERT_EQ(GetBitsNeeded(TNumericLimits<int64>::Max()), 64U);
+	UE_NET_ASSERT_EQ(GetBitsNeeded(std::numeric_limits<int64>::max()), 64U);
 }
 
 UE_NET_TEST(TestGetBitsNeeded, ReturnsExpectedValuesForUint8)
 {
 	UE_NET_ASSERT_EQ(GetBitsNeeded(uint8(0)), 0U);
 	UE_NET_ASSERT_EQ(GetBitsNeeded(uint8(1)), 1U);
-	UE_NET_ASSERT_EQ(GetBitsNeeded(TNumericLimits<uint8>::Max()), 8U);
+	UE_NET_ASSERT_EQ(GetBitsNeeded(std::numeric_limits<uint8>::max()), 8U);
 }
 
 UE_NET_TEST(TestGetBitsNeeded, ReturnsExpectedValuesForUint32)
 {
 	UE_NET_ASSERT_EQ(GetBitsNeeded(uint32(0)), 0U);
 	UE_NET_ASSERT_EQ(GetBitsNeeded(uint32(1)), 1U);
-	UE_NET_ASSERT_EQ(GetBitsNeeded(TNumericLimits<uint32>::Max()), 32U);
+	UE_NET_ASSERT_EQ(GetBitsNeeded(std::numeric_limits<uint32>::max()), 32U);
 }
 
 UE_NET_TEST(TestGetBitsNeeded, ReturnsExpectedValuesForUint64)
 {
 	UE_NET_ASSERT_EQ(GetBitsNeeded(uint64(0)), 0U);
 	UE_NET_ASSERT_EQ(GetBitsNeeded(uint64(1)), 1U);
-	UE_NET_ASSERT_EQ(GetBitsNeeded(TNumericLimits<uint64>::Max()), 64U);
+	UE_NET_ASSERT_EQ(GetBitsNeeded(std::numeric_limits<uint64>::max()), 64U);
 }
 
 // GetBitsNeededForRange
 UE_NET_TEST(TestGetBitsNeededForRange, ReturnsExpectedValuesForInt8)
 {
 	UE_NET_ASSERT_EQ(GetBitsNeededForRange(int8(47), int8(47)), 0U);
-	UE_NET_ASSERT_EQ(GetBitsNeededForRange(MIN_int8, int8(-1)), 7U);
-	UE_NET_ASSERT_EQ(GetBitsNeededForRange(int8(MIN_int8 + 5), int8(4)), 7U);
-	UE_NET_ASSERT_EQ(GetBitsNeededForRange(MIN_int8, int8(0)), 8U);
+	UE_NET_ASSERT_EQ(GetBitsNeededForRange(std::numeric_limits<int8>::min(), int8(-1)), 7U);
+	UE_NET_ASSERT_EQ(GetBitsNeededForRange(int8(std::numeric_limits<int8>::min() + 5), int8(4)), 7U);
+	UE_NET_ASSERT_EQ(GetBitsNeededForRange(std::numeric_limits<int8>::min(), int8(0)), 8U);
 	UE_NET_ASSERT_EQ(GetBitsNeededForRange(int8(11), int8(47)), 6U);
 }
 
 UE_NET_TEST(TestGetBitsNeededForRange, ReturnsExpectedValuesForInt16)
 {
 	UE_NET_ASSERT_EQ(GetBitsNeededForRange(int16(4711), int16(4711)), 0U);
-	UE_NET_ASSERT_EQ(GetBitsNeededForRange(MIN_int16, int16(-1)), 15U);
-	UE_NET_ASSERT_EQ(GetBitsNeededForRange(int16(MIN_int16 + 5), int16(4)), 15U);
-	UE_NET_ASSERT_EQ(GetBitsNeededForRange(MIN_int16, int16(0)), 16U);
+	UE_NET_ASSERT_EQ(GetBitsNeededForRange(std::numeric_limits<int16>::min(), int16(-1)), 15U);
+	UE_NET_ASSERT_EQ(GetBitsNeededForRange(int16(std::numeric_limits<int16>::min() + 5), int16(4)), 15U);
+	UE_NET_ASSERT_EQ(GetBitsNeededForRange(std::numeric_limits<int16>::min(), int16(0)), 16U);
 	UE_NET_ASSERT_EQ(GetBitsNeededForRange(int16(11), int16(47)), 6U);
 }
 
 UE_NET_TEST(TestGetBitsNeededForRange, ReturnsExpectedValuesForInt32)
 {
 	UE_NET_ASSERT_EQ(GetBitsNeededForRange(int32(4711), int32(4711)), 0U);
-	UE_NET_ASSERT_EQ(GetBitsNeededForRange(MIN_int32, int32(-1)), 31U);
-	UE_NET_ASSERT_EQ(GetBitsNeededForRange(int32(MIN_int32 + 5), int32(4)), 31U);
-	UE_NET_ASSERT_EQ(GetBitsNeededForRange(MIN_int32, int32(0)), 32U);
+	UE_NET_ASSERT_EQ(GetBitsNeededForRange(std::numeric_limits<int32>::min(), int32(-1)), 31U);
+	UE_NET_ASSERT_EQ(GetBitsNeededForRange(int32(std::numeric_limits<int32>::min() + 5), int32(4)), 31U);
+	UE_NET_ASSERT_EQ(GetBitsNeededForRange(std::numeric_limits<int32>::min(), int32(0)), 32U);
 	UE_NET_ASSERT_EQ(GetBitsNeededForRange(int32(11), int32(47)), 6U);
 }
 
 UE_NET_TEST(TestGetBitsNeededForRange, ReturnsExpectedValuesForInt64)
 {
 	UE_NET_ASSERT_EQ(GetBitsNeededForRange(int64(4711), int64(4711)), 0U);
-	UE_NET_ASSERT_EQ(GetBitsNeededForRange(MIN_int64, int64(-1)), 63U);
-	UE_NET_ASSERT_EQ(GetBitsNeededForRange(int64(MIN_int64 + 5), int64(4)), 63U);
-	UE_NET_ASSERT_EQ(GetBitsNeededForRange(MIN_int64, int64(0)), 64U);
+	UE_NET_ASSERT_EQ(GetBitsNeededForRange(std::numeric_limits<int64>::min(), int64(-1)), 63U);
+	UE_NET_ASSERT_EQ(GetBitsNeededForRange(int64(std::numeric_limits<int64>::min() + 5), int64(4)), 63U);
+	UE_NET_ASSERT_EQ(GetBitsNeededForRange(std::numeric_limits<int64>::min(), int64(0)), 64U);
 	UE_NET_ASSERT_EQ(GetBitsNeededForRange(int64(11), int64(47)), 6U);
 }
 
 UE_NET_TEST(TestGetBitsNeededForRange, ReturnsExpectedValuesForUint8)
 {
 	UE_NET_ASSERT_EQ(GetBitsNeededForRange(uint8(47), uint8(47)), 0U);
-	UE_NET_ASSERT_EQ(GetBitsNeededForRange(uint8(MAX_uint8/2U + 1U), MAX_uint8), 7U);
-	UE_NET_ASSERT_EQ(GetBitsNeededForRange(MIN_uint8, MAX_uint8), 8U);
+	UE_NET_ASSERT_EQ(GetBitsNeededForRange(uint8(std::numeric_limits<uint8>::max()/2U + 1U), std::numeric_limits<uint8>::max()), 7U);
+	UE_NET_ASSERT_EQ(GetBitsNeededForRange(std::numeric_limits<uint8>::min(), std::numeric_limits<uint8>::max()), 8U);
 }
 
 UE_NET_TEST(TestGetBitsNeededForRange, ReturnsExpectedValuesForUint16)
 {
 	UE_NET_ASSERT_EQ(GetBitsNeededForRange(uint16(4711), uint16(4711)), 0U);
-	UE_NET_ASSERT_EQ(GetBitsNeededForRange(uint16(MAX_uint16/2U + 1U), MAX_uint16), 15U);
-	UE_NET_ASSERT_EQ(GetBitsNeededForRange(MIN_uint16, MAX_uint16), 16U);
+	UE_NET_ASSERT_EQ(GetBitsNeededForRange(uint16(std::numeric_limits<uint16>::max()/2U + 1U), std::numeric_limits<uint16>::max()), 15U);
+	UE_NET_ASSERT_EQ(GetBitsNeededForRange(std::numeric_limits<uint16>::min(), std::numeric_limits<uint16>::max()), 16U);
 }
 
 UE_NET_TEST(TestGetBitsNeededForRange, ReturnsExpectedValuesForUint32)
 {
 	UE_NET_ASSERT_EQ(GetBitsNeededForRange(uint32(4711), uint32(4711)), 0U);
-	UE_NET_ASSERT_EQ(GetBitsNeededForRange(uint32(MAX_uint32/2U + 1U), MAX_uint32), 31U);
-	UE_NET_ASSERT_EQ(GetBitsNeededForRange(MIN_uint32, MAX_uint32), 32U);
+	UE_NET_ASSERT_EQ(GetBitsNeededForRange(uint32(std::numeric_limits<uint32>::max() /2U + 1U), std::numeric_limits<uint32>::max()), 31U);
+	UE_NET_ASSERT_EQ(GetBitsNeededForRange(std::numeric_limits<uint32>::min(), std::numeric_limits<uint32>::max()), 32U);
 }
 
 UE_NET_TEST(TestGetBitsNeededForRange, ReturnsExpectedValuesForUint64)
 {
 	UE_NET_ASSERT_EQ(GetBitsNeededForRange(uint64(4711), uint64(4711)), 0U);
-	UE_NET_ASSERT_EQ(GetBitsNeededForRange(uint64(MAX_uint64/2U + 1U), MAX_uint64), 63U);
-	UE_NET_ASSERT_EQ(GetBitsNeededForRange(MIN_uint64, MAX_uint64), 64U);
+	UE_NET_ASSERT_EQ(GetBitsNeededForRange(uint64(std::numeric_limits<uint64>::max() /2U + 1U), std::numeric_limits<uint64>::max()), 63U);
+	UE_NET_ASSERT_EQ(GetBitsNeededForRange(std::numeric_limits<uint64>::min(), std::numeric_limits<uint64>::max()), 64U);
 }
 
 }
\ No newline at end of file
diff --git a/Engine/Plugins/Runtime/ReplicationSystemTestPlugin/Source/Private/Tests/Serialization/TestBitPacking.cpp b/Engine/Plugins/Runtime/ReplicationSystemTestPlugin/Source/Private/Tests/Serialization/TestBitPacking.cpp
--- a/Engine/Plugins/Runtime/ReplicationSystemTestPlugin/Source/Private/Tests/Serialization/TestBitPacking.cpp
+++ b/Engine/Plugins/Runtime/ReplicationSystemTestPlugin/Source/Private/Tests/Serialization/TestBitPacking.cpp
@@ -5,6 +5,7 @@
 #include "Iris/Serialization/BitPacking.h"
 #include "Iris/Serialization/NetBitStreamReader.h"
 #include "Iris/Serialization/NetBitStreamWriter.h"
+#include <limits>
 
 namespace UE::Net::Private
 {
@@ -46,9 +47,9 @@ template<typename SourceType, uint8 BitCount> const uint32 TTestBitPacking<Sourc
 
 template<typename SourceType, uint8 BitCount> const SourceType TTestBitPacking<SourceType, BitCount>::Values[] =
 {
-	TIsSigned<SourceType>::Value ? (SourceType(-(typename TSignedIntType<sizeof(SourceType)>::Type)(TNumericLimits<SourceType>::Max() >> (sizeof(SourceType)*8U - BitCount)) - SourceType(1))) : SourceType(0), // MinValue
+	TIsSigned<SourceType>::Value ? (SourceType(-(typename TSignedIntType<sizeof(SourceType)>::Type)(std::numeric_limits<SourceType>::max() >> (sizeof(SourceType)*8U - BitCount)) - SourceType(1))) : SourceType(0), // MinValue
 	SourceType(Values[0] + SourceType(1)),
-	TNumericLimits<SourceType>::Max() >> (sizeof(SourceType)*8U - BitCount), // Max value for desired bitcount
+	std::numeric_limits<SourceType>::max() >> (sizeof(SourceType)*8U - BitCount), // Max value for desired bitcount
 	SourceType(SourceType(Values[0] + Values[2])/SourceType(2)),
 };
 template<typename SourceType, uint8 BitCount> const SIZE_T TTestBitPacking<SourceType, BitCount>::ValueCount = sizeof(Values)/sizeof(Values[0]);
diff --git a/Engine/Plugins/Runtime/ReplicationSystemTestPlugin/Source/Private/Tests/Serialization/TestEnumNetSerializer.cpp b/Engine/Plugins/Runtime/ReplicationSystemTestPlugin/Source/Private/Tests/Serialization/TestEnumNetSerializer.cpp
--- a/Engine/Plugins/Runtime/ReplicationSystemTestPlugin/Source/Private/Tests/Serialization/TestEnumNetSerializer.cpp
+++ b/Engine/Plugins/Runtime/ReplicationSystemTestPlugin/Source/Private/Tests/Serialization/TestEnumNetSerializer.cpp
@@ -3,8 +3,8 @@
 #include "TestNetSerializerFixture.h"
 #include "Iris/Core/BitTwiddling.h"
 #include "Iris/Serialization/EnumNetSerializers.h"
-#include "Math/NumericLimits.h"
 #include "EnumTestTypes.h"
+#include <limits>
 
 namespace UE::Net::Private
 {
@@ -124,8 +124,8 @@ void FTestEnumIntNetSerializer<SerializerConfig, SourceType, EnumType>::SetUp()
 		else
 		{
 			// N.B. This code is designed to also work with all uint64 enums, which UEnum doesn't handle perfectly.
-			LargeIntegerType SmallestValue = TNumericLimits<LargeIntegerType>::Max();
-			LargeIntegerType LargestValue = TNumericLimits<LargeIntegerType>::Min();
+			LargeIntegerType SmallestValue = std::numeric_limits<LargeIntegerType>::max();
+			LargeIntegerType LargestValue = std::numeric_limits<LargeIntegerType>::min();
 			for (int32 EnumIt = 0, EnumEndIt = EnumValueCount; EnumIt != EnumEndIt; ++EnumIt)
 			{
 				const LargeIntegerType Value = static_cast<LargeIntegerType>(Enum->GetValueByIndex(EnumIt));
@@ -155,12 +155,12 @@ void FTestEnumIntNetSerializer<SerializerConfig, SourceType, EnumType>::SetUp()
 		// Invalid values
 		TArray<SourceType> TempInvalidValues;
 		TempInvalidValues.Reserve(3);
-		if (Config.LowerBound > TNumericLimits<SourceType>::Min())
+		if (Config.LowerBound > std::numeric_limits<SourceType>::min())
 		{
 			TempInvalidValues.Add(Config.LowerBound - SourceType(1));
 		}
 
-		if (Config.UpperBound < TNumericLimits<SourceType>::Max())
+		if (Config.UpperBound < std::numeric_limits<SourceType>::max())
 		{
 			TempInvalidValues.Add(Config.UpperBound + SourceType(1));
 		}
diff --git a/Engine/Plugins/Runtime/ReplicationSystemTestPlugin/Source/Private/Tests/Serialization/TestIntNetSerializer.cpp b/Engine/Plugins/Runtime/ReplicationSystemTestPlugin/Source/Private/Tests/Serialization/TestIntNetSerializer.cpp
--- a/Engine/Plugins/Runtime/ReplicationSystemTestPlugin/Source/Private/Tests/Serialization/TestIntNetSerializer.cpp
+++ b/Engine/Plugins/Runtime/ReplicationSystemTestPlugin/Source/Private/Tests/Serialization/TestIntNetSerializer.cpp
@@ -2,6 +2,7 @@
 
 #include "TestNetSerializerFixture.h"
 #include "Iris/Serialization/IntNetSerializers.h"
+#include <limits>
 
 namespace UE::Net::Private
 {
@@ -78,7 +79,7 @@ UE_NET_IMPLEMENT_INT_NETSERIALIZER_TEST(FTestInt8NetSerializer, FInt8NetSerializ
 //
 template<typename SourceType> const SourceType FTestIntNetSerializer<SourceType>::Values[] =
 {
-	TNumericLimits<SourceType>::Lowest(), TNumericLimits<SourceType>::Max(), SourceType(0), SourceType(-1), SourceType(-16), SourceType(2048458 % TNumericLimits<SourceType>::Max())
+	std::numeric_limits<SourceType>::min(), std::numeric_limits<SourceType>::max(), SourceType(0), SourceType(-1), SourceType(-16), SourceType(2048458 % std::numeric_limits<SourceType>::max())
 };
 template<typename SourceType> const SIZE_T FTestIntNetSerializer<SourceType>::ValueCount = sizeof(Values)/sizeof(Values[0]);
 
@@ -128,7 +129,7 @@ void FTestIntNetSerializer<SourceType>::TestValidate()
 	bool ExpectedResults[sizeof(Values)/sizeof(Values[0])];
 	for (uint8 BitCount = MinBitCount; BitCount <= MaxBitCount; ++BitCount)
 	{
-		constexpr SourceType MaxValue = TNumericLimits<SourceType>::Max();
+		constexpr SourceType MaxValue = std::numeric_limits<SourceType>::max();
 		const SourceType MaxValueForBitCount = MaxValue >> (sizeof(SourceType)*8 - BitCount);
 		const SourceType MinValueForBitCount = -MaxValueForBitCount - SourceType(1);
 		for (SIZE_T ValueIt = 0; ValueIt < ValueCount; ++ValueIt)
diff --git a/Engine/Plugins/Runtime/ReplicationSystemTestPlugin/Source/Private/Tests/Serialization/TestIntRangeNetSerializer.cpp b/Engine/Plugins/Runtime/ReplicationSystemTestPlugin/Source/Private/Tests/Serialization/TestIntRangeNetSerializer.cpp
--- a/Engine/Plugins/Runtime/ReplicationSystemTestPlugin/Source/Private/Tests/Serialization/TestIntRangeNetSerializer.cpp
+++ b/Engine/Plugins/Runtime/ReplicationSystemTestPlugin/Source/Private/Tests/Serialization/TestIntRangeNetSerializer.cpp
@@ -3,6 +3,7 @@
 #include "TestNetSerializerFixture.h"
 #include "Iris/Core/BitTwiddling.h"
 #include "Iris/Serialization/IntRangeNetSerializers.h"
+#include <limits>
 
 namespace UE::Net::Private
 {
@@ -212,7 +213,7 @@ UE_NET_TEST_FIXTURE(FTestInt64RangeNetSerializer, TestSerializeDelta)
 //
 template<typename SerializerConfig, typename SourceType> const SourceType FTestIntRangeNetSerializer<SerializerConfig, SourceType>::Values[] =
 {
-	TNumericLimits<SourceType>::Lowest(), TNumericLimits<SourceType>::Max(), SourceType(0), SourceType(-1), SourceType(-16), SourceType(2048458 % TNumericLimits<SourceType>::Max())
+	std::numeric_limits<SourceType>::min(), std::numeric_limits<SourceType>::max(), SourceType(0), SourceType(-1), SourceType(-16), SourceType(2048458 % std::numeric_limits<SourceType>::max())
 };
 template<typename SerializerConfig, typename SourceType> const SIZE_T FTestIntRangeNetSerializer<SerializerConfig, SourceType>::ValueCount = sizeof(Values)/sizeof(Values[0]);
 template<typename SerializerConfig, typename SourceType> const SIZE_T FTestIntRangeNetSerializer<SerializerConfig, SourceType>::ConfigCount = sizeof(Configs)/sizeof(Configs[0]);
@@ -224,8 +225,8 @@ void FTestIntRangeNetSerializer<SerializerConfig, SourceType>::SetupConfigs()
 	{
 		static_assert(FullRangeConfigIndex == 0, "Incorrect NetSerializerConfig setup");
 		SerializerConfig& Config = Configs[0];
-		Config.LowerBound = TNumericLimits<SourceType>::Lowest();
-		Config.UpperBound = TNumericLimits<SourceType>::Max();
+		Config.LowerBound = std::numeric_limits<SourceType>::min();
+		Config.UpperBound = std::numeric_limits<SourceType>::max();
 		Config.BitCount = GetBitsNeededForRange(Config.LowerBound, Config.UpperBound);
 	}
 
@@ -240,15 +241,15 @@ void FTestIntRangeNetSerializer<SerializerConfig, SourceType>::SetupConfigs()
 	// Arbitrary ranges
 	{
 		SerializerConfig& Config = Configs[2];
-		Config.LowerBound = TNumericLimits<SourceType>::Lowest()/2;
-		Config.UpperBound = TNumericLimits<SourceType>::Max()/4;
+		Config.LowerBound = std::numeric_limits<SourceType>::min()/2;
+		Config.UpperBound = std::numeric_limits<SourceType>::max()/4;
 		Config.BitCount = GetBitsNeededForRange(Config.LowerBound, Config.UpperBound);
 	}
 
 	{
 		SerializerConfig& Config = Configs[3];
-		Config.LowerBound = TNumericLimits<SourceType>::Max()/4;
-		Config.UpperBound = TNumericLimits<SourceType>::Max()/2;
+		Config.LowerBound = std::numeric_limits<SourceType>::max()/4;
+		Config.UpperBound = std::numeric_limits<SourceType>::max()/2;
 		Config.BitCount = GetBitsNeededForRange(Config.LowerBound, Config.UpperBound);
 	}
 }
diff --git a/Engine/Plugins/Runtime/ReplicationSystemTestPlugin/Source/Private/Tests/Serialization/TestPackedIntNetSerializer.cpp b/Engine/Plugins/Runtime/ReplicationSystemTestPlugin/Source/Private/Tests/Serialization/TestPackedIntNetSerializer.cpp
--- a/Engine/Plugins/Runtime/ReplicationSystemTestPlugin/Source/Private/Tests/Serialization/TestPackedIntNetSerializer.cpp
+++ b/Engine/Plugins/Runtime/ReplicationSystemTestPlugin/Source/Private/Tests/Serialization/TestPackedIntNetSerializer.cpp
@@ -2,6 +2,7 @@
 
 #include "TestNetSerializerFixture.h"
 #include "Iris/Serialization/PackedIntNetSerializers.h"
+#include <limits>
 
 namespace UE::Net::Private
 {
@@ -70,7 +71,7 @@ UE_NET_IMPLEMENT_PACKEDINT_NETSERIALIZER_TEST(FTestPackedInt32NetSerializer, FPa
 //
 template<typename SourceType> const SourceType TTestPackedIntNetSerializer<SourceType>::Values[] =
 {
-	TNumericLimits<SourceType>::Lowest(), TNumericLimits<SourceType>::Max(), SourceType(0), SourceType(-1), (TNumericLimits<SourceType>::Max() + TNumericLimits<SourceType>::Lowest())/SourceType(2) - SourceType(15), SourceType(2048458 % TNumericLimits<SourceType>::Max())
+	std::numeric_limits<SourceType>::min(), std::numeric_limits<SourceType>::max(), SourceType(0), SourceType(-1), (std::numeric_limits<SourceType>::max() + std::numeric_limits<SourceType>::min())/SourceType(2) - SourceType(15), SourceType(2048458 % std::numeric_limits<SourceType>::max())
 };
 template<typename SourceType> const SIZE_T TTestPackedIntNetSerializer<SourceType>::ValueCount = sizeof(Values)/sizeof(Values[0]);
 
@@ -133,7 +134,7 @@ template<typename SourceType>
 void TTestPackedIntNetSerializer<SourceType>::TestSmallValuesArePacked()
 {
 	constexpr SourceType SmallValue = SourceType(0);
-	constexpr SourceType LargeValue = TNumericLimits<SourceType>::Max();
+	constexpr SourceType LargeValue = std::numeric_limits<SourceType>::max();
 
 	const bool bSerializedSmallValue = this->Serialize(Config, NetSerializerValuePointer(&SmallValue));
 	UE_NET_ASSERT_TRUE_MSG(bSerializedSmallValue, "Failed serializing " << SmallValue);
diff --git a/Engine/Plugins/Runtime/ReplicationSystemTestPlugin/Source/Private/Tests/Serialization/TestPackedUintNetSerializer.cpp b/Engine/Plugins/Runtime/ReplicationSystemTestPlugin/Source/Private/Tests/Serialization/TestPackedUintNetSerializer.cpp
--- a/Engine/Plugins/Runtime/ReplicationSystemTestPlugin/Source/Private/Tests/Serialization/TestPackedUintNetSerializer.cpp
+++ b/Engine/Plugins/Runtime/ReplicationSystemTestPlugin/Source/Private/Tests/Serialization/TestPackedUintNetSerializer.cpp
@@ -2,6 +2,7 @@
 
 #include "TestNetSerializerFixture.h"
 #include "Iris/Serialization/PackedIntNetSerializers.h"
+#include <limits>
 
 namespace UE::Net::Private
 {
@@ -70,7 +71,7 @@ UE_NET_IMPLEMENT_PACKEDUINT_NETSERIALIZER_TEST(FTestPackedUint32NetSerializer, F
 //
 template<typename SourceType> const SourceType TTestPackedUintNetSerializer<SourceType>::Values[] =
 {
-	TNumericLimits<SourceType>::Lowest(), TNumericLimits<SourceType>::Max(), SourceType(0), (TNumericLimits<SourceType>::Max() + TNumericLimits<SourceType>::Lowest())/SourceType(2) - SourceType(15), SourceType(2048458 % TNumericLimits<SourceType>::Max())
+	std::numeric_limits<SourceType>::min(), std::numeric_limits<SourceType>::max(), SourceType(0), (std::numeric_limits<SourceType>::max() + std::numeric_limits<SourceType>::min())/SourceType(2) - SourceType(15), SourceType(2048458 % std::numeric_limits<SourceType>::max())
 };
 template<typename SourceType> const SIZE_T TTestPackedUintNetSerializer<SourceType>::ValueCount = sizeof(Values)/sizeof(Values[0]);
 
@@ -133,7 +134,7 @@ template<typename SourceType>
 void TTestPackedUintNetSerializer<SourceType>::TestSmallValuesArePacked()
 {
 	constexpr SourceType SmallValue = SourceType(0);
-	constexpr SourceType LargeValue = TNumericLimits<SourceType>::Max();
+	constexpr SourceType LargeValue = std::numeric_limits<SourceType>::max();
 
 	const bool bSerializedSmallValue = this->Serialize(Config, NetSerializerValuePointer(&SmallValue));
 	UE_NET_ASSERT_TRUE_MSG(bSerializedSmallValue, "Failed serializing " << SmallValue);
diff --git a/Engine/Plugins/Runtime/ReplicationSystemTestPlugin/Source/Private/Tests/Serialization/TestUintNetSerializer.cpp b/Engine/Plugins/Runtime/ReplicationSystemTestPlugin/Source/Private/Tests/Serialization/TestUintNetSerializer.cpp
--- a/Engine/Plugins/Runtime/ReplicationSystemTestPlugin/Source/Private/Tests/Serialization/TestUintNetSerializer.cpp
+++ b/Engine/Plugins/Runtime/ReplicationSystemTestPlugin/Source/Private/Tests/Serialization/TestUintNetSerializer.cpp
@@ -2,6 +2,7 @@
 
 #include "TestNetSerializerFixture.h"
 #include "Iris/Serialization/UintNetSerializers.h"
+#include <limits>
 
 namespace UE::Net::Private
 {
@@ -78,7 +79,7 @@ UE_NET_IMPLEMENT_UINT_NETSERIALIZER_TEST(FTestUint8NetSerializer, FUint8NetSeria
 //
 template<typename SourceType> const SourceType FTestUintNetSerializer<SourceType>::Values[] =
 {
-	TNumericLimits<SourceType>::Lowest(), TNumericLimits<SourceType>::Max(), SourceType(0), SourceType(16), SourceType(2048458 % TNumericLimits<SourceType>::Max())
+	std::numeric_limits<SourceType>::min(), std::numeric_limits<SourceType>::max(), SourceType(0), SourceType(16), SourceType(2048458 % std::numeric_limits<SourceType>::max())
 };
 template<typename SourceType> const SIZE_T FTestUintNetSerializer<SourceType>::ValueCount = sizeof(Values)/sizeof(Values[0]);
 
@@ -129,7 +130,7 @@ void FTestUintNetSerializer<SourceType>::TestValidate()
 	bool ExpectedResults[sizeof(Values)/sizeof(Values[0])];
 	for (uint8 BitCount = MinBitCount; BitCount <= MaxBitCount; ++BitCount)
 	{
-		const SourceType MaxValue = TNumericLimits<SourceType>::Max() >> (MaxBitCount - BitCount);
+		const SourceType MaxValue = std::numeric_limits<SourceType>::max() >> (MaxBitCount - BitCount);
 		for (SIZE_T ValueIt = 0; ValueIt < ValueCount; ++ValueIt)
 		{
 			const SourceType Value = Values[ValueIt];
diff --git a/Engine/Plugins/Runtime/ReplicationSystemTestPlugin/Source/Private/Tests/Serialization/TestUintRangeNetSerializer.cpp b/Engine/Plugins/Runtime/ReplicationSystemTestPlugin/Source/Private/Tests/Serialization/TestUintRangeNetSerializer.cpp
--- a/Engine/Plugins/Runtime/ReplicationSystemTestPlugin/Source/Private/Tests/Serialization/TestUintRangeNetSerializer.cpp
+++ b/Engine/Plugins/Runtime/ReplicationSystemTestPlugin/Source/Private/Tests/Serialization/TestUintRangeNetSerializer.cpp
@@ -3,6 +3,7 @@
 #include "TestNetSerializerFixture.h"
 #include "Iris/Core/BitTwiddling.h"
 #include "Iris/Serialization/UintRangeNetSerializers.h"
+#include <limits>
 
 namespace UE::Net::Private
 {
@@ -94,7 +95,7 @@ UE_NET_IMPLEMENT_UINTRANGE_NETSERIALIZER_TEST(FTestUint64RangeNetSerializer, FUi
 //
 template<typename SerializerConfig, typename SourceType> const SourceType FTestUintRangeNetSerializer<SerializerConfig, SourceType>::Values[] =
 {
-	TNumericLimits<SourceType>::Lowest(), TNumericLimits<SourceType>::Max(), SourceType(0), SourceType(1), SourceType(2048458 % TNumericLimits<SourceType>::Max())
+	std::numeric_limits<SourceType>::min(), std::numeric_limits<SourceType>::max(), SourceType(0), SourceType(1), SourceType(2048458 % std::numeric_limits<SourceType>::max())
 };
 template<typename SerializerConfig, typename SourceType> const SIZE_T FTestUintRangeNetSerializer<SerializerConfig, SourceType>::ValueCount = sizeof(Values)/sizeof(Values[0]);
 
@@ -105,8 +106,8 @@ void FTestUintRangeNetSerializer<SerializerConfig, SourceType>::SetupConfigs()
 	{
 		static_assert(FullRangeConfigIndex == 0, "Incorrect NetSerializerConfig setup");
 		SerializerConfig& Config = Configs[0];
-		Config.LowerBound = TNumericLimits<SourceType>::Lowest();
-		Config.UpperBound = TNumericLimits<SourceType>::Max();
+		Config.LowerBound = std::numeric_limits<SourceType>::min();
+		Config.UpperBound = std::numeric_limits<SourceType>::max();
 		Config.BitCount = GetBitsNeededForRange(Config.LowerBound, Config.UpperBound);
 	}
 
@@ -121,15 +122,15 @@ void FTestUintRangeNetSerializer<SerializerConfig, SourceType>::SetupConfigs()
 	// Arbitrary ranges
 	{
 		SerializerConfig& Config = Configs[2];
-		Config.LowerBound = TNumericLimits<SourceType>::Lowest()/2;
-		Config.UpperBound = TNumericLimits<SourceType>::Max()/4;
+		Config.LowerBound = std::numeric_limits<SourceType>::min()/2;
+		Config.UpperBound = std::numeric_limits<SourceType>::max()/4;
 		Config.BitCount = GetBitsNeededForRange(Config.LowerBound, Config.UpperBound);
 	}
 
 	{
 		SerializerConfig& Config = Configs[3];
-		Config.LowerBound = TNumericLimits<SourceType>::Max()/4;
-		Config.UpperBound = TNumericLimits<SourceType>::Max()/2;
+		Config.LowerBound = std::numeric_limits<SourceType>::max()/4;
+		Config.UpperBound = std::numeric_limits<SourceType>::max()/2;
 		Config.BitCount = GetBitsNeededForRange(Config.LowerBound, Config.UpperBound);
 	}
 }
diff --git a/Engine/Plugins/Runtime/ReplicationSystemTestPlugin/Source/Private/Tests/Serialization/TestVectorNetSerializer.cpp b/Engine/Plugins/Runtime/ReplicationSystemTestPlugin/Source/Private/Tests/Serialization/TestVectorNetSerializer.cpp
--- a/Engine/Plugins/Runtime/ReplicationSystemTestPlugin/Source/Private/Tests/Serialization/TestVectorNetSerializer.cpp
+++ b/Engine/Plugins/Runtime/ReplicationSystemTestPlugin/Source/Private/Tests/Serialization/TestVectorNetSerializer.cpp
@@ -2,6 +2,7 @@
 
 #include "TestNetSerializerFixture.h"
 #include "Iris/Serialization/VectorNetSerializers.h"
+#include <limits>
 
 namespace UE::Net::Private
 {
@@ -54,8 +55,8 @@ const FVector FTestVectorNetSerializer::Values[] =
 	*(const FVector*)(&BadVectorValues[3]),
 	FVector(-0.0f),
 	FVector(0.0f),
-	FVector(TNumericLimits<float>::Max()),
-	FVector(TNumericLimits<float>::Lowest()),
+	FVector(std::numeric_limits<FVector::FReal>::max()),
+	FVector(std::numeric_limits<FVector::FReal>::lowest()),
 	FVector(47.11f),
 	FVector(1.26698556E-23)
 };
