Iris - Struct derived from struct with custom NetSerializer is no longer required to implement a NetSerializer.

diff --git a/Engine/Plugins/Runtime/ReplicationSystemTestPlugin/Source/Private/Tests/ReplicationState/TestDerivedStruct.cpp b/Engine/Plugins/Runtime/ReplicationSystemTestPlugin/Source/Private/Tests/ReplicationState/TestDerivedStruct.cpp
new file mode 100644
--- /dev/null
+++ b/Engine/Plugins/Runtime/ReplicationSystemTestPlugin/Source/Private/Tests/ReplicationState/TestDerivedStruct.cpp
@@ -0,0 +1,384 @@
+// Copyright Epic Games, Inc. All Rights Reserved.
+
+#include "TestDerivedStruct.h"
+#include "NetworkAutomationTest.h"
+#include "NetworkAutomationTestMacros.h"
+#include "Iris/Serialization/NetSerializers.h"
+#include "Iris/ReplicationState/PropertyNetSerializerInfoRegistry.h"
+#include "Iris/ReplicationState/ReplicationStateDescriptorBuilder.h"
+#include "Iris/ReplicationSystem/ReplicationFragmentUtil.h"
+#include "Net/Core/Trace/NetDebugName.h"
+#include "Net/UnrealNetwork.h"
+#include "Tests/ReplicationSystem/ReplicationSystemServerClientTestFixture.h"
+
+namespace UE::Net::Private
+{
+
+struct FTestDerivedStruct_Inherited_WithNetSerializer_NetSerializer
+{
+	struct FQuantizedData
+	{
+		uint8 ByteMember0;
+		uint8 ByteMember1;
+	};
+
+	typedef FTestDerivedStruct_Inherited_WithNetSerializer SourceType;
+	typedef FQuantizedData QuantizedType;
+
+	typedef FNetSerializerConfig ConfigType;
+
+	static const uint32 Version = 0;
+	inline static const ConfigType DefaultConfig;
+
+	static void Serialize(FNetSerializationContext&, const FNetSerializeArgs&);
+	static void Deserialize(FNetSerializationContext&, const FNetDeserializeArgs&);
+
+	static void Quantize(FNetSerializationContext&, const FNetQuantizeArgs&);
+	static void Dequantize(FNetSerializationContext&, const FNetDequantizeArgs&);
+
+	static bool IsEqual(FNetSerializationContext&, const FNetIsEqualArgs&);
+
+
+	class FNetSerializerRegistryDelegates final : private UE::Net::FNetSerializerRegistryDelegates
+	{
+	public:
+		virtual ~FNetSerializerRegistryDelegates();
+
+	private:
+		virtual void OnPreFreezeNetSerializerRegistry() override;
+
+		static inline const FName PropertyNetSerializerRegistry_NAME_TestDerivedStruct_Inherited_WithNetSerializer = FName("TestDerivedStruct_Inherited_WithNetSerializer");
+		UE_NET_IMPLEMENT_NAMED_STRUCT_NETSERIALIZER_INFO(PropertyNetSerializerRegistry_NAME_TestDerivedStruct_Inherited_WithNetSerializer, FTestDerivedStruct_Inherited_WithNetSerializer_NetSerializer);
+	};
+	
+	static inline FTestDerivedStruct_Inherited_WithNetSerializer_NetSerializer::FNetSerializerRegistryDelegates NetSerializerRegistryDelegates;
+};
+UE_NET_IMPLEMENT_SERIALIZER(FTestDerivedStruct_Inherited_WithNetSerializer_NetSerializer);
+
+class FTestDerivedStructFixture : public FNetworkAutomationTestSuiteFixture
+{
+};
+
+class FTestDerivedStructInObjectFixture : public FReplicationSystemServerClientTestFixture
+{
+protected:
+	UTestDerivedStruct_TestObject_Member* CreateObjectWithMember();
+	UTestDerivedStruct_TestObject_Array* CreateObjectWithArray();
+
+};
+
+// We want the struct deriving from the struct with custom NetSerializer to have the same size. This allows us to detect if the struct with the NetSerializer overwrites any of the properties in the derived struct.
+UE_NET_TEST_FIXTURE(FTestDerivedStructFixture, VerifyTestAssumptions)
+{
+	UE_NET_EXPECT_EQ_MSG(sizeof(FTestDerivedStruct_Inherited_WithNetSerializer), sizeof(FTestDerivedStruct_Inherited_WithNetSerializer_Inherited_WithoutNetSerializer), "Expected FTestDerivedStruct_Inherited_WithNetSerializer to have same size as FTestDerivedStruct_Inherited_WithNetSerializer_Inherited_WithoutNetSerializer.");
+}
+
+UE_NET_TEST_FIXTURE(FTestDerivedStructFixture, StructWithNetSerializerGetsNetSerializer)
+{
+	TRefCountPtr<const FReplicationStateDescriptor> StructDescriptor = FReplicationStateDescriptorBuilder::CreateDescriptorForStruct(FTestDerivedStruct_Inherited_WithNetSerializer::StaticStruct());
+	UE_NET_ASSERT_EQ(StructDescriptor->MemberCount, uint16(1));
+}
+
+UE_NET_TEST_FIXTURE(FTestDerivedStructFixture, DerivedStructWithoutNetSerializerMakesUseOfParentNetSerializer)
+{
+	TRefCountPtr<const FReplicationStateDescriptor> StructDescriptor = FReplicationStateDescriptorBuilder::CreateDescriptorForStruct(FTestDerivedStruct_Inherited_WithNetSerializer_Inherited_WithoutNetSerializer::StaticStruct());
+	
+	// The derived struct should have two members, the first using the NetSerializer of the parent struct.
+	UE_NET_ASSERT_EQ(StructDescriptor->MemberCount, uint16(2));
+	UE_NET_ASSERT_EQ(StructDescriptor->MemberSerializerDescriptors[0].Serializer, &UE_NET_GET_SERIALIZER(FTestDerivedStruct_Inherited_WithNetSerializer_NetSerializer));
+	UE_NET_ASSERT_TRUE(EnumHasAnyFlags(StructDescriptor->Traits, EReplicationStateTraits::IsDerivedStruct));
+}
+
+UE_NET_TEST_FIXTURE(FTestDerivedStructFixture, DerivedStructHasValidDebugDescriptor)
+{
+	TRefCountPtr<const FReplicationStateDescriptor> StructDescriptor = FReplicationStateDescriptorBuilder::CreateDescriptorForStruct(FTestDerivedStruct_Inherited_WithNetSerializer_Inherited_WithoutNetSerializer::StaticStruct());
+
+	const FReplicationStateMemberDebugDescriptor* DebugDescriptors = StructDescriptor->MemberDebugDescriptors;
+	UE_NET_ASSERT_NE(DebugDescriptors, nullptr);
+
+	for (const FReplicationStateMemberDebugDescriptor& MemberDebugDescriptor : MakeArrayView(DebugDescriptors, StructDescriptor->MemberCount))
+	{
+		UE_NET_ASSERT_NE_MSG(MemberDebugDescriptor.DebugName, nullptr, FString().Appendf(TEXT("Member %u lacking DebugName"), (&MemberDebugDescriptor - DebugDescriptors)));
+		UE_NET_ASSERT_NE_MSG(MemberDebugDescriptor.DebugName->Name, nullptr, FString().Appendf(TEXT("Member %u lacking DebugName.Name"), (&MemberDebugDescriptor - DebugDescriptors)));
+	}
+}
+
+UE_NET_TEST_FIXTURE(FTestDerivedStructFixture, DeepInheritanceOfStructWithNetSerializerMakesUseOfParentNetSerializer)
+{
+	TRefCountPtr<const FReplicationStateDescriptor> StructDescriptor = FReplicationStateDescriptorBuilder::CreateDescriptorForStruct(FTestDerivedStruct_DeepInheritanceOfStructWithNetSerializer::StaticStruct());
+
+	UE_NET_ASSERT_TRUE(EnumHasAnyFlags(StructDescriptor->Traits, EReplicationStateTraits::IsDerivedStruct));
+
+	// The derived struct should have three members, the first using the NetSerializer of the parent struct.
+	UE_NET_ASSERT_EQ(StructDescriptor->MemberCount, uint16(3));
+	UE_NET_ASSERT_EQ(StructDescriptor->MemberSerializerDescriptors[0].Serializer, &UE_NET_GET_SERIALIZER(FTestDerivedStruct_Inherited_WithNetSerializer_NetSerializer));
+
+	const FProperty* LastProperty = StructDescriptor->MemberProperties[StructDescriptor->MemberCount - 1];
+	UE_NET_ASSERT_NE(LastProperty, nullptr);
+	UE_NET_ASSERT_EQ(LastProperty->GetName(), FString(TEXT("YetAnotherMember")));
+}
+
+UE_NET_TEST_FIXTURE(FTestDerivedStructFixture, DeepInheritanceAndStructWithNetSerializerUsesSerializer)
+{
+	TRefCountPtr<const FReplicationStateDescriptor> StructDescriptor = FReplicationStateDescriptorBuilder::CreateDescriptorForStruct(FTestDerivedStruct_DeepInheritanceOfStructWithNetSerializer::StaticStruct());
+
+	UE_NET_ASSERT_TRUE(EnumHasAnyFlags(StructDescriptor->Traits, EReplicationStateTraits::IsDerivedStruct));
+
+	// The derived struct should have three members, the first using the NetSerializer of the parent struct.
+	UE_NET_ASSERT_EQ(StructDescriptor->MemberCount, uint16(3));
+	UE_NET_ASSERT_EQ(StructDescriptor->MemberSerializerDescriptors[0].Serializer, &UE_NET_GET_SERIALIZER(FTestDerivedStruct_Inherited_WithNetSerializer_NetSerializer));
+
+	const FProperty* LastProperty = StructDescriptor->MemberProperties[StructDescriptor->MemberCount - 1];
+	UE_NET_ASSERT_NE(LastProperty, nullptr);
+	UE_NET_ASSERT_EQ(LastProperty->GetName(), FString(TEXT("YetAnotherMember")));
+}
+
+UE_NET_TEST_FIXTURE(FTestDerivedStructFixture, StructUsesClosestNetSerializerInHierachy)
+{
+	TRefCountPtr<const FReplicationStateDescriptor> StructDescriptor = FReplicationStateDescriptorBuilder::CreateDescriptorForStruct(FTestDerivedStruct_Inherited_Inherited_WithNetSerializer_DeepInheritanceOfStructWithNetSerializer::StaticStruct());
+
+	UE_NET_ASSERT_TRUE(EnumHasAnyFlags(StructDescriptor->Traits, EReplicationStateTraits::IsDerivedStruct));
+
+	// The derived struct should have two members, the first using the NetSerializer of the parent struct.
+	UE_NET_ASSERT_EQ(StructDescriptor->MemberCount, uint16(2));
+	// We've registered the struct to use an already existing serializer so need to look for the serializer we're forwarding to.
+	UE_NET_ASSERT_EQ(StructDescriptor->MemberSerializerDescriptors[0].Serializer, &UE_NET_GET_SERIALIZER(FNopNetSerializer));
+
+	const FProperty* LastProperty = StructDescriptor->MemberProperties[StructDescriptor->MemberCount - 1];
+	UE_NET_ASSERT_NE(LastProperty, nullptr);
+	UE_NET_ASSERT_EQ(LastProperty->GetName(), FString(TEXT("ByteMember5")));
+}
+
+UE_NET_TEST_FIXTURE(FTestDerivedStructInObjectFixture, CanReplicateDerivedStructMember)
+{
+	FReplicationSystemTestClient* Client = CreateClient();
+
+	UTestDerivedStruct_TestObject_Member* ServerObject = CreateObjectWithMember();
+	ServerObject->DerivedStruct.ByteMember2 ^= 47;
+
+	// Replicate
+	Server->PreSendUpdate();
+	Server->SendAndDeliverTo(Client, DeliverPacket);
+	Server->PostSendUpdate();
+
+	const UTestDerivedStruct_TestObject_Member* ClientObject = Cast<UTestDerivedStruct_TestObject_Member>(Client->GetReplicationBridge()->GetReplicatedObject(ServerObject->NetRefHandle));
+	UE_NET_ASSERT_NE(ClientObject, nullptr);
+
+	// Check that the member has been replicated
+	UE_NET_ASSERT_EQ(ClientObject->DerivedStruct.ByteMember2, ServerObject->DerivedStruct.ByteMember2);
+}
+
+UE_NET_TEST_FIXTURE(FTestDerivedStructInObjectFixture, CanReplicateDerivedStructArray)
+{
+	FReplicationSystemTestClient* Client = CreateClient();
+
+	UTestDerivedStruct_TestObject_Array* ServerObject = CreateObjectWithArray();
+
+	// Add items to the array
+	{
+		FTestDerivedStruct_Inherited_WithNetSerializer_Inherited_WithoutNetSerializer Item;
+		Item.ByteMember2 ^= 11;
+		ServerObject->DerivedStructArray.Add(Item);
+		Item.ByteMember2 ^= 47;
+		ServerObject->DerivedStructArray.Add(Item);
+	}
+
+	// Replicate
+	Server->PreSendUpdate();
+	Server->SendAndDeliverTo(Client, DeliverPacket);
+	Server->PostSendUpdate();
+
+	const UTestDerivedStruct_TestObject_Array* ClientObject = Cast<UTestDerivedStruct_TestObject_Array>(Client->GetReplicationBridge()->GetReplicatedObject(ServerObject->NetRefHandle));
+	UE_NET_ASSERT_NE(ClientObject, nullptr);
+
+	// Check that the members have been replicated
+	{
+		UE_NET_ASSERT_EQ(ClientObject->DerivedStructArray.Num(), ServerObject->DerivedStructArray.Num());
+		SIZE_T MemberIt = 0;
+		for (const FTestDerivedStruct_Inherited_WithNetSerializer_Inherited_WithoutNetSerializer& ServerItem : ServerObject->DerivedStructArray)
+		{
+			const FTestDerivedStruct_Inherited_WithNetSerializer_Inherited_WithoutNetSerializer& ClientItem = ClientObject->DerivedStructArray[MemberIt];
+			++MemberIt;
+			UE_NET_ASSERT_EQ(ClientItem.ByteMember2, ServerItem.ByteMember2);
+		}
+	}
+}
+
+UE_NET_TEST_FIXTURE(FTestDerivedStructInObjectFixture, NotReplicatedMemberInDerivedStructIsNotOverwritten)
+{
+	FReplicationSystemTestClient* Client = CreateClient();
+
+	UTestDerivedStruct_TestObject_Member* ServerObject = CreateObjectWithMember();
+
+	// Replicate
+	Server->PreSendUpdate();
+	Server->SendAndDeliverTo(Client, DeliverPacket);
+	Server->PostSendUpdate();
+
+	UTestDerivedStruct_TestObject_Member* ClientObject = Cast<UTestDerivedStruct_TestObject_Member>(Client->GetReplicationBridge()->GetReplicatedObject(ServerObject->NetRefHandle));
+	UE_NET_ASSERT_NE(ClientObject, nullptr);
+
+	// On the client modify the non-replicated property
+	constexpr uint8 ClientModifiedNotReplicatedPropertyValue = 255U;
+	ClientObject->DerivedStruct.ByteMember3_NotReplicated = ClientModifiedNotReplicatedPropertyValue;
+
+	// On the server modify all members
+	ServerObject->DerivedStruct.ByteMember0 = 0;
+	ServerObject->DerivedStruct.ByteMember1 = 1;
+	ServerObject->DerivedStruct.ByteMember2 = 2;
+	ServerObject->DerivedStruct.ByteMember3_NotReplicated = ClientModifiedNotReplicatedPropertyValue ^ 111;
+
+	Server->PreSendUpdate();
+	Server->SendAndDeliverTo(Client, DeliverPacket);
+	Server->PostSendUpdate();
+
+	// Verify state
+	UE_NET_ASSERT_EQ(ClientObject->DerivedStruct.ByteMember0, ServerObject->DerivedStruct.ByteMember0);
+	UE_NET_ASSERT_EQ(ClientObject->DerivedStruct.ByteMember1, ServerObject->DerivedStruct.ByteMember1);
+	UE_NET_ASSERT_EQ(ClientObject->DerivedStruct.ByteMember2, ServerObject->DerivedStruct.ByteMember2);
+	UE_NET_ASSERT_EQ(ClientObject->DerivedStruct.ByteMember3_NotReplicated, ClientModifiedNotReplicatedPropertyValue);
+}
+
+UE_NET_TEST_FIXTURE(FTestDerivedStructInObjectFixture, CanReplicateDerivedStructMembersInArbitraryOrder)
+{
+	FReplicationSystemTestClient* Client = CreateClient();
+
+	UTestDerivedStruct_TestObject_Member* ServerObject = CreateObjectWithMember();
+
+	// Replicate
+	Server->PreSendUpdate();
+	Server->SendAndDeliverTo(Client, DeliverPacket);
+	Server->PostSendUpdate();
+
+	const UTestDerivedStruct_TestObject_Member* ClientObject = Cast<UTestDerivedStruct_TestObject_Member>(Client->GetReplicationBridge()->GetReplicatedObject(ServerObject->NetRefHandle));
+	UE_NET_ASSERT_NE(ClientObject, nullptr);
+
+	// Modify the property not replicated by the custom serializer.
+	ServerObject->DerivedStruct.ByteMember2 ^= 47;
+	Server->PreSendUpdate();
+	Server->SendAndDeliverTo(Client, DeliverPacket);
+	Server->PostSendUpdate();
+
+	// Verify state
+	UE_NET_ASSERT_EQ(ClientObject->DerivedStruct.ByteMember0, ServerObject->DerivedStruct.ByteMember0);
+	UE_NET_ASSERT_EQ(ClientObject->DerivedStruct.ByteMember1, ServerObject->DerivedStruct.ByteMember1);
+	UE_NET_ASSERT_EQ(ClientObject->DerivedStruct.ByteMember2, ServerObject->DerivedStruct.ByteMember2);
+
+	// Modify properties replicated by custom serializer
+	ServerObject->DerivedStruct.ByteMember0 ^= 13;
+	ServerObject->DerivedStruct.ByteMember1 ^= 37;
+	Server->PreSendUpdate();
+	Server->SendAndDeliverTo(Client, DeliverPacket);
+	Server->PostSendUpdate();
+
+	// Verify state
+	UE_NET_ASSERT_EQ(ClientObject->DerivedStruct.ByteMember0, ServerObject->DerivedStruct.ByteMember0);
+	UE_NET_ASSERT_EQ(ClientObject->DerivedStruct.ByteMember1, ServerObject->DerivedStruct.ByteMember1);
+	UE_NET_ASSERT_EQ(ClientObject->DerivedStruct.ByteMember2, ServerObject->DerivedStruct.ByteMember2);
+}
+// FTestDerivedStruct_Inherited_WithNetSerializer_NetSerializer implementation
+void FTestDerivedStruct_Inherited_WithNetSerializer_NetSerializer::Serialize(FNetSerializationContext& Context, const FNetSerializeArgs& Args)
+{
+	const QuantizedType& Value = *reinterpret_cast<const QuantizedType*>(Args.Source);
+
+	FNetBitStreamWriter* Writer = Context.GetBitStreamWriter();
+	Writer->WriteBits(Value.ByteMember0, 8U);
+	Writer->WriteBits(Value.ByteMember1, 8U);
+}
+
+void FTestDerivedStruct_Inherited_WithNetSerializer_NetSerializer::Deserialize(FNetSerializationContext& Context, const FNetDeserializeArgs& Args)
+{
+	QuantizedType& Target = *reinterpret_cast<QuantizedType*>(Args.Target);
+
+	FNetBitStreamReader* Reader = Context.GetBitStreamReader();
+	Target.ByteMember0 = Reader->ReadBits(8U);
+	Target.ByteMember1 = Reader->ReadBits(8U);
+}
+
+void FTestDerivedStruct_Inherited_WithNetSerializer_NetSerializer::Quantize(FNetSerializationContext& Context, const FNetQuantizeArgs& Args)
+{
+	const SourceType& Source = *reinterpret_cast<const SourceType*>(Args.Source);
+	QuantizedType& Target = *reinterpret_cast<QuantizedType*>(Args.Target);
+
+	Target.ByteMember0 = Source.ByteMember0;
+	Target.ByteMember1 = Source.ByteMember1;
+}
+
+void FTestDerivedStruct_Inherited_WithNetSerializer_NetSerializer::Dequantize(FNetSerializationContext& Context, const FNetDequantizeArgs& Args)
+{
+	const QuantizedType& Source = *reinterpret_cast<const QuantizedType*>(Args.Source);
+	SourceType& Target = *reinterpret_cast<SourceType*>(Args.Target);
+
+	Target.ByteMember0 = Source.ByteMember0;
+	Target.ByteMember1 = Source.ByteMember1;
+}
+
+bool FTestDerivedStruct_Inherited_WithNetSerializer_NetSerializer::IsEqual(FNetSerializationContext& Context, const FNetIsEqualArgs& Args)
+{
+	if (Args.bStateIsQuantized)
+	{
+		const QuantizedType& Value0 = *reinterpret_cast<const QuantizedType*>(Args.Source0);
+		const QuantizedType& Value1 = *reinterpret_cast<const QuantizedType*>(Args.Source1);
+		return Value0.ByteMember0 == Value1.ByteMember0 && Value0.ByteMember1 == Value1.ByteMember1;
+	}
+	else
+	{
+		const SourceType& Value0 = *reinterpret_cast<const SourceType*>(Args.Source0);
+		const SourceType& Value1 = *reinterpret_cast<const SourceType*>(Args.Source1);
+		return Value0.ByteMember0 == Value1.ByteMember0 && Value0.ByteMember1 == Value1.ByteMember1;
+	}
+}
+
+FTestDerivedStruct_Inherited_WithNetSerializer_NetSerializer::FNetSerializerRegistryDelegates::~FNetSerializerRegistryDelegates()
+{
+	UE_NET_UNREGISTER_NETSERIALIZER_INFO(PropertyNetSerializerRegistry_NAME_TestDerivedStruct_Inherited_WithNetSerializer);
+}
+
+void FTestDerivedStruct_Inherited_WithNetSerializer_NetSerializer::FNetSerializerRegistryDelegates::OnPreFreezeNetSerializerRegistry()
+{
+	UE_NET_REGISTER_NETSERIALIZER_INFO(PropertyNetSerializerRegistry_NAME_TestDerivedStruct_Inherited_WithNetSerializer);
+
+}
+
+// FTestDerivedStruct_Inherited_WithNetSerializer_DeepInheritanceOfStructWithNetSerializer will not be used for replication testing so it's ok to use the NopNetSerializer.
+UE_NET_IMPLEMENT_FORWARDING_NETSERIALIZER_AND_REGISTRY_DELEGATES(TestDerivedStruct_Inherited_WithNetSerializer_DeepInheritanceOfStructWithNetSerializer, FNopNetSerializer);
+
+UTestDerivedStruct_TestObject_Member* FTestDerivedStructInObjectFixture::CreateObjectWithMember()
+{
+	return Server->CreateObject<UTestDerivedStruct_TestObject_Member>();
+}
+
+UTestDerivedStruct_TestObject_Array* FTestDerivedStructInObjectFixture::CreateObjectWithArray()
+{
+	return Server->CreateObject<UTestDerivedStruct_TestObject_Array>();
+}
+
+}
+
+FTestDerivedStruct_Base::~FTestDerivedStruct_Base() = default;
+
+void UTestDerivedStruct_TestObject_Member::GetLifetimeReplicatedProps(TArray<FLifetimeProperty>& OutLifetimeProps) const
+{
+	FDoRepLifetimeParams Params;
+	Params.bIsPushBased = false;
+
+	DOREPLIFETIME_WITH_PARAMS_FAST(ThisClass, DerivedStruct, Params);
+}
+
+void UTestDerivedStruct_TestObject_Member::RegisterReplicationFragments(UE::Net::FFragmentRegistrationContext& Context, UE::Net::EFragmentRegistrationFlags RegistrationFlags)
+{
+	UE::Net::FReplicationFragmentUtil::CreateAndRegisterFragmentsForObject(this, Context, RegistrationFlags);
+}
+
+void UTestDerivedStruct_TestObject_Array::GetLifetimeReplicatedProps(TArray<FLifetimeProperty>& OutLifetimeProps) const
+{
+	FDoRepLifetimeParams Params;
+	Params.bIsPushBased = false;
+
+	DOREPLIFETIME_WITH_PARAMS_FAST(ThisClass, DerivedStructArray, Params);
+}
+
+void UTestDerivedStruct_TestObject_Array::RegisterReplicationFragments(UE::Net::FFragmentRegistrationContext& Context, UE::Net::EFragmentRegistrationFlags RegistrationFlags)
+{
+	UE::Net::FReplicationFragmentUtil::CreateAndRegisterFragmentsForObject(this, Context, RegistrationFlags);
+}
diff --git a/Engine/Plugins/Runtime/ReplicationSystemTestPlugin/Source/Private/Tests/ReplicationState/TestDerivedStruct.h b/Engine/Plugins/Runtime/ReplicationSystemTestPlugin/Source/Private/Tests/ReplicationState/TestDerivedStruct.h
new file mode 100644
--- /dev/null
+++ b/Engine/Plugins/Runtime/ReplicationSystemTestPlugin/Source/Private/Tests/ReplicationState/TestDerivedStruct.h
@@ -0,0 +1,104 @@
+// Copyright Epic Games, Inc. All Rights Reserved.
+
+#pragma once
+
+#include "NetworkAutomationTest.h"
+#include "Iris/Serialization/NetSerializer.h"
+#include "UObject/ObjectMacros.h"
+#include "Tests/ReplicationSystem/ReplicatedTestObject.h"
+#include "TestDerivedStruct.generated.h"
+
+USTRUCT()
+struct alignas(16) FTestDerivedStruct_Base
+{
+	GENERATED_BODY()
+	
+	virtual ~FTestDerivedStruct_Base();
+
+	UPROPERTY(Transient)
+	uint8 ByteMember0 = 5;
+};
+
+USTRUCT()
+struct FTestDerivedStruct_Inherited_WithNetSerializer : public FTestDerivedStruct_Base
+{
+	GENERATED_BODY()
+	
+	UPROPERTY(Transient)
+	uint8 ByteMember1 = 25;
+};
+
+USTRUCT()
+struct FTestDerivedStruct_Inherited_WithNetSerializer_Inherited_WithoutNetSerializer : public FTestDerivedStruct_Inherited_WithNetSerializer
+{
+	GENERATED_BODY()
+
+	UPROPERTY(Transient)
+	uint8 ByteMember2 = 125;
+
+	UPROPERTY(NotReplicated)
+	uint8 ByteMember3_NotReplicated = 0;
+};
+
+USTRUCT()
+struct FTestDerivedStruct_DeepInheritanceOfStructWithNetSerializer : public FTestDerivedStruct_Inherited_WithNetSerializer_Inherited_WithoutNetSerializer
+{
+	GENERATED_BODY()
+
+	UPROPERTY(Transient)
+	float YetAnotherMember;
+};
+
+// This struct will use a forwarding serializer rather than implement a separate custom one.
+USTRUCT()
+struct FTestDerivedStruct_Inherited_WithNetSerializer_DeepInheritanceOfStructWithNetSerializer : public FTestDerivedStruct_DeepInheritanceOfStructWithNetSerializer
+{
+	GENERATED_BODY()
+
+	UPROPERTY(Transient)
+	uint8 ByteMember4;
+};
+
+USTRUCT()
+struct FTestDerivedStruct_Inherited_Inherited_WithNetSerializer_DeepInheritanceOfStructWithNetSerializer : public FTestDerivedStruct_Inherited_WithNetSerializer_DeepInheritanceOfStructWithNetSerializer
+{
+	GENERATED_BODY()
+
+	UPROPERTY(Transient)
+	uint8 ByteMember5;
+};
+
+UCLASS()
+class UTestDerivedStruct_TestObject_Member : public UReplicatedTestObject
+{
+	GENERATED_BODY()
+
+public:
+	UPROPERTY(Replicated, Transient)
+	FTestDerivedStruct_Inherited_WithNetSerializer_Inherited_WithoutNetSerializer DerivedStruct;
+
+protected:
+	virtual void GetLifetimeReplicatedProps(TArray<FLifetimeProperty>& OutLifetimeProps) const override;
+	virtual void RegisterReplicationFragments(UE::Net::FFragmentRegistrationContext& Context, UE::Net::EFragmentRegistrationFlags RegistrationFlags) override;
+};
+
+UCLASS()
+class UTestDerivedStruct_TestObject_Array : public UReplicatedTestObject
+{
+	GENERATED_BODY()
+
+public:
+	UPROPERTY(Replicated, Transient)
+	TArray<FTestDerivedStruct_Inherited_WithNetSerializer_Inherited_WithoutNetSerializer> DerivedStructArray;
+
+protected:
+	virtual void GetLifetimeReplicatedProps(TArray<FLifetimeProperty>& OutLifetimeProps) const override;
+	virtual void RegisterReplicationFragments(UE::Net::FFragmentRegistrationContext& Context, UE::Net::EFragmentRegistrationFlags RegistrationFlags) override;
+};
+
+namespace UE::Net::Private
+{
+
+UE_NET_DECLARE_SERIALIZER(FTestDerivedStruct_Inherited_WithNetSerializer_NetSerializer, REPLICATIONSYSTEMTESTPLUGIN_API);
+
+}
diff --git a/Engine/Plugins/Runtime/ReplicationSystemTestPlugin/Source/Private/Tests/ReplicationState/TestReplicationStateDescriptorBuilder.cpp b/Engine/Plugins/Runtime/ReplicationSystemTestPlugin/Source/Private/Tests/ReplicationState/TestReplicationStateDescriptorBuilder.cpp
--- a/Engine/Plugins/Runtime/ReplicationSystemTestPlugin/Source/Private/Tests/ReplicationState/TestReplicationStateDescriptorBuilder.cpp
+++ b/Engine/Plugins/Runtime/ReplicationSystemTestPlugin/Source/Private/Tests/ReplicationState/TestReplicationStateDescriptorBuilder.cpp
@@ -177,17 +177,18 @@ static FReplicationStateMemberTraitsDescriptor TestStructMemberTraitsDescriptors
 
 static const FReplicationStateDescriptor TestStructReplictionStateDescriptor =
 {
-	TestStructMemberDescriptors,			//FReplicationStateMemberDescriptor* MemberDescriptors;
-	static_cast<const FReplicationStateMemberChangeMaskDescriptor*>(nullptr),	//FReplicationStateMemberChangeMaskDescriptor* MemberChangeMaskDescriptors;
-	TestStructMemberSerializerDescriptors,	//FReplicationStateMemberSerializerDescriptor* MemberSerializerDescriptors;
-	TestStructMemberTraitsDescriptors, // MemberTraitsDescriptors;
+	TestStructMemberDescriptors,			// MemberDescriptors
+	static_cast<const FReplicationStateMemberChangeMaskDescriptor*>(nullptr),	// MemberChangeMaskDescriptors
+	TestStructMemberSerializerDescriptors,	// MemberSerializerDescriptors
+	TestStructMemberTraitsDescriptors, // MemberTraitsDescriptors
 	static_cast<const FReplicationStateMemberFunctionDescriptor*>(nullptr),
-	static_cast<const FReplicationStateMemberTagDescriptor*>(nullptr),	// MemberTagDescriptors;
-	static_cast<const FReplicationStateMemberReferenceDescriptor*>(nullptr),	// MemberReferenceDescriptors;
-	static_cast<const FProperty**>(nullptr),	// FProperty** MemberProperties;
+	static_cast<const FReplicationStateMemberTagDescriptor*>(nullptr),	// MemberTagDescriptors
+	static_cast<const FReplicationStateMemberReferenceDescriptor*>(nullptr),	// MemberReferenceDescriptors
+	static_cast<const FProperty**>(nullptr),	// MemberProperties
 	static_cast<const FReplicationStateMemberPropertyDescriptor*>(nullptr), // MemberPropertyDescriptors
 	static_cast<const FReplicationStateMemberLifetimeConditionDescriptor*>(nullptr),
 	static_cast<const FReplicationStateMemberRepIndexToMemberIndexDescriptor*>(nullptr),
+	static_cast<const UScriptStruct*>(nullptr), // BaseStruct
 	static_cast<const FNetDebugName*>(nullptr), // DebugName
 	static_cast<const FReplicationStateMemberDebugDescriptor*>(nullptr), // MemberDebugDescriptors
 	24,			//uint32 ExternalSize;
@@ -257,17 +258,18 @@ static FReplicationStateMemberChangeMaskDescriptor TestClassMemberChangeMaskDesc
 // Should be the same
 static const FReplicationStateDescriptor TestClassReplictionStateDescriptor
 {
-	TestClassMemberDescriptors,			//FReplicationStateMemberDescriptor* MemberDescriptors;
-	TestClassMemberChangeMaskDescriptors,	//FReplicationStateMemberChangeMaskDescriptor* MemberChangeMaskDescriptors;
-	TestClassMemberSerializerDescriptors,	//FReplicationStateMemberSerializerDescriptor* MemberSerializerDescriptors;
-	TestClassMemberTraitsDescriptors, // MemberTraitsDescriptors;
+	TestClassMemberDescriptors,			// MemberDescriptors
+	TestClassMemberChangeMaskDescriptors,	// MemberChangeMaskDescriptors
+	TestClassMemberSerializerDescriptors,	// MemberSerializerDescriptors
+	TestClassMemberTraitsDescriptors, // MemberTraitsDescriptors
 	static_cast<const FReplicationStateMemberFunctionDescriptor*>(nullptr),
 	static_cast<const FReplicationStateMemberTagDescriptor*>(nullptr), // MemberTagDescriptors
-	static_cast<const FReplicationStateMemberReferenceDescriptor*>(nullptr), // MemberReferenceDescriptors;
-	static_cast<const FProperty**>(nullptr), // FProperty** MemberProperties;
+	static_cast<const FReplicationStateMemberReferenceDescriptor*>(nullptr), // MemberReferenceDescriptors
+	static_cast<const FProperty**>(nullptr), // MemberProperties
 	static_cast<const FReplicationStateMemberPropertyDescriptor*>(nullptr), // MemberPropertyDescriptors
 	static_cast<const FReplicationStateMemberLifetimeConditionDescriptor*>(nullptr),
 	static_cast<const FReplicationStateMemberRepIndexToMemberIndexDescriptor*>(nullptr),
+	static_cast<const UScriptStruct*>(nullptr), // BaseStruct
 	static_cast<const FNetDebugName*>(nullptr), // DebugName
 	static_cast<const FReplicationStateMemberDebugDescriptor*>(nullptr), // MemberDebugDescriptors
 	32,			//uint32 ExternalSize;
@@ -292,17 +294,18 @@ static const FReplicationStateDescriptor TestClassReplictionStateDescriptor
 // Should be the same
 static const FReplicationStateDescriptor TestClassReplictionStateDescriptorNotReferenceCounted
 {
-	TestClassMemberDescriptors,			//FReplicationStateMemberDescriptor* MemberDescriptors;
-	TestClassMemberChangeMaskDescriptors,	//FReplicationStateMemberChangeMaskDescriptor* MemberChangeMaskDescriptors;
-	TestClassMemberSerializerDescriptors,	//FReplicationStateMemberSerializerDescriptor* MemberSerializerDescriptors;
-	TestClassMemberTraitsDescriptors, // MemberTraitsDescriptors;
+	TestClassMemberDescriptors,			// MemberDescriptors
+	TestClassMemberChangeMaskDescriptors,	// MemberChangeMaskDescriptors
+	TestClassMemberSerializerDescriptors,	// MemberSerializerDescriptors
+	TestClassMemberTraitsDescriptors, // MemberTraitsDescriptors
 	static_cast<const FReplicationStateMemberFunctionDescriptor*>(nullptr),
 	static_cast<const FReplicationStateMemberTagDescriptor*>(nullptr), // MemberTagDescriptors
-	static_cast<const FReplicationStateMemberReferenceDescriptor*>(nullptr), // MemberReferenceDescriptors;
-	static_cast<const FProperty**>(nullptr), // FProperty** MemberProperties;
+	static_cast<const FReplicationStateMemberReferenceDescriptor*>(nullptr), // MemberReferenceDescriptors
+	static_cast<const FProperty**>(nullptr), // MemberProperties
 	static_cast<const FReplicationStateMemberPropertyDescriptor*>(nullptr), // MemberPropertyDescriptors
 	static_cast<const FReplicationStateMemberLifetimeConditionDescriptor*>(nullptr),
 	static_cast<const FReplicationStateMemberRepIndexToMemberIndexDescriptor*>(nullptr),
+	static_cast<const UScriptStruct*>(nullptr), // BaseStruct
 	static_cast<const FNetDebugName*>(nullptr), // DebugName
 	static_cast<const FReplicationStateMemberDebugDescriptor*>(nullptr), // MemberDebugDescriptors
 	32,			//uint32 ExternalSize;
@@ -327,17 +330,18 @@ static const FReplicationStateDescriptor TestClassReplictionStateDescriptorNotRe
 // Should be the same
 static const FReplicationStateDescriptor TestClassWithNonReplicatedDataReplicationStateDescriptor =
 {
-	TestClassMemberDescriptors,			//FReplicationStateMemberDescriptor* MemberDescriptors;
-	TestClassMemberChangeMaskDescriptors,	//FReplicationStateMemberChangeMaskDescriptor* MemberChangeMaskDescriptors;
-	TestClassMemberSerializerDescriptors,	//FReplicationStateMemberSerializerDescriptor* MemberSerializerDescriptors;
+	TestClassMemberDescriptors,			// MemberDescriptors
+	TestClassMemberChangeMaskDescriptors,	// MemberChangeMaskDescriptors
+	TestClassMemberSerializerDescriptors,	// MemberSerializerDescriptors
 	TestClassMemberTraitsDescriptors,
 	static_cast<const FReplicationStateMemberFunctionDescriptor*>(nullptr),
 	static_cast<const FReplicationStateMemberTagDescriptor*>(nullptr), // MemberTagDescriptors
-	static_cast<const FReplicationStateMemberReferenceDescriptor*>(nullptr), // MemberReferenceDescriptors;
-	static_cast<const FProperty**>(nullptr), // FProperty** MemberProperties;
+	static_cast<const FReplicationStateMemberReferenceDescriptor*>(nullptr), // MemberReferenceDescriptors
+	static_cast<const FProperty**>(nullptr), // MemberProperties
 	static_cast<const FReplicationStateMemberPropertyDescriptor*>(nullptr), // MemberPropertyDescriptors
 	static_cast<const FReplicationStateMemberLifetimeConditionDescriptor*>(nullptr),
 	static_cast<const FReplicationStateMemberRepIndexToMemberIndexDescriptor*>(nullptr),
+	static_cast<const UScriptStruct*>(nullptr), // BaseStruct
 	static_cast<const FNetDebugName*>(nullptr), // DebugName
 	static_cast<const FReplicationStateMemberDebugDescriptor*>(nullptr), // MemberDebugDescriptors
 	32,			//uint32 ExternalSize;
@@ -480,11 +484,12 @@ static FReplicationStateDescriptor TestClassWithInheritanceReplictionStateDescri
 	TestClassWithInheritanceMemberTraitsDescriptors,
 	static_cast<const FReplicationStateMemberFunctionDescriptor*>(nullptr),
 	static_cast<const FReplicationStateMemberTagDescriptor*>(nullptr), // MemberTagDescriptors
-	static_cast<const FReplicationStateMemberReferenceDescriptor*>(nullptr), // MemberReferenceDescriptors;
-	static_cast<const FProperty**>(nullptr), // MemberProperties;
+	static_cast<const FReplicationStateMemberReferenceDescriptor*>(nullptr), // MemberReferenceDescriptors
+	static_cast<const FProperty**>(nullptr), // MemberProperties
 	static_cast<const FReplicationStateMemberPropertyDescriptor*>(nullptr), // MemberPropertyDescriptors
 	static_cast<const FReplicationStateMemberLifetimeConditionDescriptor*>(nullptr),
 	static_cast<const FReplicationStateMemberRepIndexToMemberIndexDescriptor*>(nullptr),
+	static_cast<const UScriptStruct*>(nullptr), // BaseStruct
 	static_cast<const FNetDebugName*>(nullptr), // DebugName
 	static_cast<const FReplicationStateMemberDebugDescriptor*>(nullptr), // MemberDebugDescriptors
 	32, // ExternalSize (rounded up since it also contains statemask)
@@ -552,11 +557,12 @@ static FReplicationStateDescriptor TestClassWithInheritanceNoSuperReplictionStat
 	TestClassWithInheritanceNoSuperMemberTraitsDescriptors,
 	static_cast<const FReplicationStateMemberFunctionDescriptor*>(nullptr),
 	static_cast<const FReplicationStateMemberTagDescriptor*>(nullptr), // MemberTagDescriptors
-	static_cast<const FReplicationStateMemberReferenceDescriptor*>(nullptr), // MemberReferenceDescriptors;
-	static_cast<const FProperty**>(nullptr), // MemberProperties;
+	static_cast<const FReplicationStateMemberReferenceDescriptor*>(nullptr), // MemberReferenceDescriptors
+	static_cast<const FProperty**>(nullptr), // MemberProperties
 	static_cast<const FReplicationStateMemberPropertyDescriptor*>(nullptr), // MemberPropertyDescriptors
 	static_cast<const FReplicationStateMemberLifetimeConditionDescriptor*>(nullptr),
 	static_cast<const FReplicationStateMemberRepIndexToMemberIndexDescriptor*>(nullptr),
+	static_cast<const UScriptStruct*>(nullptr), // BaseStruct
 	static_cast<const FNetDebugName*>(nullptr), // DebugName
 	static_cast<const FReplicationStateMemberDebugDescriptor*>(nullptr), // MemberDebugDescriptors
 	12, // ExternalSize (rounded up since it also contains statemask)
@@ -617,13 +623,14 @@ static FReplicationStateDescriptor TestClassWithReplicatedStructReplictionStateD
 	TestClassWithReplicatedStructMemberChangeMaskDescriptors,
 	TestClassWithReplicatedStructMemberSerializerDescriptors,
 	TestClassWithReplicatedStructMemberTraitsDescriptors,
-	static_cast<const FReplicationStateMemberFunctionDescriptor*>(nullptr), // MemberTagDescriptors;
+	static_cast<const FReplicationStateMemberFunctionDescriptor*>(nullptr), // MemberTagDescriptors
 	static_cast<const FReplicationStateMemberTagDescriptor*>(nullptr), // MemberTagDescriptors
-	static_cast<const FReplicationStateMemberReferenceDescriptor*>(nullptr), // MemberReferenceDescriptors;
-	static_cast<const FProperty**>(nullptr), // MemberProperties;
+	static_cast<const FReplicationStateMemberReferenceDescriptor*>(nullptr), // MemberReferenceDescriptors
+	static_cast<const FProperty**>(nullptr), // MemberProperties
 	static_cast<const FReplicationStateMemberPropertyDescriptor*>(nullptr), // MemberPropertyDescriptors
 	static_cast<const FReplicationStateMemberLifetimeConditionDescriptor*>(nullptr),
 	static_cast<const FReplicationStateMemberRepIndexToMemberIndexDescriptor*>(nullptr),
+	static_cast<const UScriptStruct*>(nullptr), // BaseStruct
 	static_cast<const FNetDebugName*>(nullptr), // DebugName
 	static_cast<const FReplicationStateMemberDebugDescriptor*>(nullptr), // MemberDebugDescriptors
 	32, // ExternalSize (rounded up since it also contains statemask)
diff --git a/Engine/Plugins/Runtime/ReplicationSystemTestPlugin/Source/Private/Tests/Serialization/TestPolymorphicStructNetSerializer.cpp b/Engine/Plugins/Runtime/ReplicationSystemTestPlugin/Source/Private/Tests/Serialization/TestPolymorphicStructNetSerializer.cpp
--- a/Engine/Plugins/Runtime/ReplicationSystemTestPlugin/Source/Private/Tests/Serialization/TestPolymorphicStructNetSerializer.cpp
+++ b/Engine/Plugins/Runtime/ReplicationSystemTestPlugin/Source/Private/Tests/Serialization/TestPolymorphicStructNetSerializer.cpp
@@ -43,11 +43,9 @@ struct FTestPolymorphicStructNetSerializer : public TPolymorphicStructNetSeriali
 	}
 
 	static const uint32 Version = 0;
-	static const ConfigType DefaultConfig;
+	static inline const ConfigType DefaultConfig;
 };
 
-const FTestPolymorphicStructNetSerializer::ConfigType FTestPolymorphicStructNetSerializer::DefaultConfig;
-
 UE_NET_IMPLEMENT_SERIALIZER(FTestPolymorphicStructNetSerializer);
 
 }
@@ -101,6 +99,7 @@ protected:
 	FExamplePolymorphicStructB* StructB;
 	FExamplePolymorphicStructC* StructC;
 	FExamplePolymorphicStructD* StructD;
+	FExamplePolymorphicStructD_Derived* StructD_Derived;
 
 	alignas(8) uint8 QuantizedBuffer[2][2048];
 	alignas(8) uint8 ClonedQuantizedBuffer[2][2048];
@@ -153,7 +152,7 @@ struct FExamplePolymorphicStructDNetSerializer
 	typedef uint32 QuantizedType;
 	typedef FExamplePolymorphicStructDNetSerializerConfig ConfigType;
 
-	static const ConfigType DefaultConfig;
+	static inline const ConfigType DefaultConfig;
 
 	static void Serialize(FNetSerializationContext&, const FNetSerializeArgs& Args);
 	static void Deserialize(FNetSerializationContext&, const FNetDeserializeArgs& Args);
@@ -164,8 +163,6 @@ struct FExamplePolymorphicStructDNetSerializer
 	static bool IsEqual(FNetSerializationContext& Context, const FNetIsEqualArgs& Args);
 };
 
-const FExamplePolymorphicStructDNetSerializer::ConfigType FExamplePolymorphicStructDNetSerializer::DefaultConfig;
-
 void FExamplePolymorphicStructDNetSerializer::Serialize(FNetSerializationContext& Context, const FNetSerializeArgs& Args)
 {
 	const uint32 Value = *reinterpret_cast<const uint32*>(Args.Source);
@@ -415,11 +412,16 @@ void FTestPolymorphicArrayStructNetSerializerFixture::SetExpectedState(FExampleP
 	StructD = new FExamplePolymorphicStructD;
 	StructD->SomeValue = 0x7E577E57;
 	
+	StructD_Derived = new FExamplePolymorphicStructD_Derived;
+	StructD_Derived->SomeValue = 0xC0DEC0DE;
+	StructD_Derived->FloatInD_Derived = 4711.0f;
+
 	// Add some data, we need some structs to send as well
 	Target.Add(StructA);
 	Target.Add(StructB);
 	Target.Add(StructC);
 	Target.Add(StructD);
+	Target.Add(StructD_Derived);
 }
 
 
@@ -470,12 +472,14 @@ void FTestPolymorphicArrayStructNetSerializerFixture::ValidateExpectedState(cons
 	auto BInstance = static_cast<const FExamplePolymorphicStructB*>(StructInstance.Get(1));
 	auto CInstance = static_cast<const FExamplePolymorphicStructC*>(StructInstance.Get(2));
 	auto DInstance = static_cast<const FExamplePolymorphicStructD*>(StructInstance.Get(3));
+	auto D_DerivedInstance = static_cast<const FExamplePolymorphicStructD_Derived*>(StructInstance.Get(4));
 
 	// Validate types
 	UE_NET_ASSERT_TRUE(AInstance->GetScriptStruct() == FExamplePolymorphicStructA::StaticStruct());
 	UE_NET_ASSERT_TRUE(BInstance->GetScriptStruct() == FExamplePolymorphicStructB::StaticStruct());
 	UE_NET_ASSERT_TRUE(CInstance->GetScriptStruct() == FExamplePolymorphicStructC::StaticStruct());
 	UE_NET_ASSERT_TRUE(DInstance->GetScriptStruct() == FExamplePolymorphicStructD::StaticStruct());
+	UE_NET_ASSERT_TRUE(D_DerivedInstance->GetScriptStruct() == FExamplePolymorphicStructD_Derived::StaticStruct());
 
 	// Validate actual data	
 	UE_NET_ASSERT_EQ(StructA->SomeInt, AInstance->SomeInt);
@@ -483,6 +487,8 @@ void FTestPolymorphicArrayStructNetSerializerFixture::ValidateExpectedState(cons
 	UE_NET_ASSERT_EQ(StructC->SomeBool, CInstance->SomeBool);
 	UE_NET_ASSERT_EQ(ToRawPtr(StructC->SomeObjectRef), ToRawPtr(CInstance->SomeObjectRef));
 	UE_NET_ASSERT_EQ(StructD->SomeValue, DInstance->SomeValue);
+	UE_NET_ASSERT_EQ(StructD_Derived->SomeValue, D_DerivedInstance->SomeValue);
+	UE_NET_ASSERT_EQ(StructD_Derived->FloatInD_Derived, D_DerivedInstance->FloatInD_Derived);
 }
 
 void FTestPolymorphicArrayStructNetSerializerFixture::TearDown()
@@ -862,10 +868,15 @@ void FTestPolymorphicArrayStructNetSerializerDeltaSerializationFixture::SetArbit
 	FExamplePolymorphicStructD* StructD = new FExamplePolymorphicStructD;
 	StructD->SomeValue = 0xBEEEEEEF;
 
+	FExamplePolymorphicStructD_Derived* StructD_Derived = new FExamplePolymorphicStructD_Derived;
+	StructD_Derived->SomeValue = 0xDEEDDEED;
+	StructD_Derived->FloatInD_Derived = 12345.0f;
+
 	// Add values to target
 	Target.Add(StructA);
 	Target.Add(StructB);
 	Target.Add(StructC);
+	Target.Add(StructD_Derived);
 	Target.Add(StructD);
 }
 
diff --git a/Engine/Plugins/Runtime/ReplicationSystemTestPlugin/Source/Private/Tests/Serialization/TestPolymorphicStructNetSerializer.h b/Engine/Plugins/Runtime/ReplicationSystemTestPlugin/Source/Private/Tests/Serialization/TestPolymorphicStructNetSerializer.h
--- a/Engine/Plugins/Runtime/ReplicationSystemTestPlugin/Source/Private/Tests/Serialization/TestPolymorphicStructNetSerializer.h
+++ b/Engine/Plugins/Runtime/ReplicationSystemTestPlugin/Source/Private/Tests/Serialization/TestPolymorphicStructNetSerializer.h
@@ -455,6 +455,21 @@ struct TStructOpsTypeTraits<FExamplePolymorphicStructD> : public TStructOpsTypeT
 	};
 };
 
+// Derives from a struct with custom NetSerializer and extra properties
+USTRUCT()
+struct FExamplePolymorphicStructD_Derived : public FExamplePolymorphicStructD
+{
+	GENERATED_BODY()
+
+	UPROPERTY()
+	float FloatInD_Derived = 0.0f;
+
+	virtual UScriptStruct* GetScriptStruct() const override
+	{
+		return FExamplePolymorphicStructD_Derived::StaticStruct();
+	}
+};
+
 namespace UE::Net
 {
 
diff --git a/Engine/Source/Runtime/Experimental/Iris/Core/Private/Iris/ReplicationState/InternalPropertyReplicationState.cpp b/Engine/Source/Runtime/Experimental/Iris/Core/Private/Iris/ReplicationState/InternalPropertyReplicationState.cpp
--- a/Engine/Source/Runtime/Experimental/Iris/Core/Private/Iris/ReplicationState/InternalPropertyReplicationState.cpp
+++ b/Engine/Source/Runtime/Experimental/Iris/Core/Private/Iris/ReplicationState/InternalPropertyReplicationState.cpp
@@ -151,23 +151,7 @@ bool InternalCompareMember(const FReplicationStateDescriptor* Descriptor, uint32
 		const FReplicationStateDescriptor* StructDescriptor = StructConfig->StateDescriptor;
 		if (!EnumHasAnyFlags(StructDescriptor->Traits, EReplicationStateTraits::AllMembersAreReplicated))
 		{
-			const FReplicationStateMemberDescriptor* MemberDescriptors = StructDescriptor->MemberDescriptors;
-			const FProperty** MemberProperties = StructDescriptor->MemberProperties;
-			const FReplicationStateMemberPropertyDescriptor* MemberPropertyDescriptors = StructDescriptor->MemberPropertyDescriptors;
-
-			for (uint32 StructMemberIt = 0, StructMemberEndIt = StructDescriptor->MemberCount; StructMemberIt != StructMemberEndIt; ++StructMemberIt)
-			{
-				const FReplicationStateMemberDescriptor& MemberDescriptor = MemberDescriptors[StructMemberIt];
-				const FReplicationStateMemberPropertyDescriptor& MemberPropertyDescriptor = MemberPropertyDescriptors[StructMemberIt];
-				const FProperty* MemberProperty = MemberProperties[StructMemberIt];
-				const SIZE_T MemberOffset = MemberProperty->GetOffset_ForGC() + MemberProperty->ElementSize*MemberPropertyDescriptor.ArrayIndex;
-				if (!InternalCompareMember(StructDescriptor, StructMemberIt, static_cast<const uint8*>(ValueA) + MemberOffset, static_cast<const uint8*>(ValueB) + MemberOffset))
-				{
-					return false;
-				}
-			}
-
-			return true;
+			return InternalCompareStructProperty(StructDescriptor, ValueA, ValueB);
 		}
 	}
 	else if (IsUsingArrayPropertyNetSerializer(MemberSerializerDescriptor))
@@ -192,19 +176,42 @@ bool InternalCompareMember(const FReplicationStateDescriptor* Descriptor, uint32
 	return Property->Identical(ValueA, ValueB);
 }
 
-bool InternalCompareStructProperty(const FReplicationStateDescriptor* StructDescriptor, void* RESTRICT Dst, const void* RESTRICT Src)
+bool InternalCompareStructProperty(const FReplicationStateDescriptor* StructDescriptor, const void* RESTRICT ValueA, const void* RESTRICT ValueB)
 {
+	uint32 FirstStructMemberForCompare = 0U;
+
 	const FReplicationStateMemberDescriptor* MemberDescriptors = StructDescriptor->MemberDescriptors;
+	const FReplicationStateMemberSerializerDescriptor* MemberSerializerDescriptors = StructDescriptor->MemberSerializerDescriptors;
 	const FProperty** MemberProperties = StructDescriptor->MemberProperties;
 	const FReplicationStateMemberPropertyDescriptor* MemberPropertyDescriptors = StructDescriptor->MemberPropertyDescriptors;
 
-	for (uint32 StructMemberIt = 0, StructMemberEndIt = StructDescriptor->MemberCount; StructMemberIt != StructMemberEndIt; ++StructMemberIt)
+	// For derived structs the first property is a NetSerializer for some super struct. Its property will be null so we can't use standard iteration over properties. Let's use serializer equal for it. 
+	if (EnumHasAnyFlags(StructDescriptor->Traits, EReplicationStateTraits::IsDerivedStruct))
+	{
+		FirstStructMemberForCompare = 1U;
+
+		const FReplicationStateMemberDescriptor& MemberDescriptor = MemberDescriptors[0];
+		const FReplicationStateMemberSerializerDescriptor& MemberSerializerDescriptor = MemberSerializerDescriptors[0];
+
+		FNetSerializationContext Context;
+		FNetIsEqualArgs IsEqualArgs;
+		IsEqualArgs.NetSerializerConfig = MemberSerializerDescriptor.SerializerConfig;
+		IsEqualArgs.Source0 = NetSerializerValuePointer(ValueA) + MemberDescriptor.ExternalMemberOffset;
+		IsEqualArgs.Source1 = NetSerializerValuePointer(ValueB) + MemberDescriptor.ExternalMemberOffset;
+		IsEqualArgs.bStateIsQuantized = false;
+		if (!MemberSerializerDescriptor.Serializer->IsEqual(Context, IsEqualArgs))
+		{
+			return false;
+		}
+	}
+
+	for (uint32 StructMemberIt = FirstStructMemberForCompare, StructMemberEndIt = StructDescriptor->MemberCount; StructMemberIt < StructMemberEndIt; ++StructMemberIt)
 	{
 		const FReplicationStateMemberDescriptor& MemberDescriptor = MemberDescriptors[StructMemberIt];
 		const FReplicationStateMemberPropertyDescriptor& MemberPropertyDescriptor = MemberPropertyDescriptors[StructMemberIt];
 		const FProperty* MemberProperty = MemberProperties[StructMemberIt];
 		const SIZE_T MemberOffset = MemberProperty->GetOffset_ForGC() + MemberProperty->ElementSize*MemberPropertyDescriptor.ArrayIndex;
-		if (!InternalCompareMember(StructDescriptor, StructMemberIt, static_cast<uint8*>(Dst) + MemberOffset, static_cast<const uint8*>(Src) + MemberOffset))
+		if (!InternalCompareMember(StructDescriptor, StructMemberIt, static_cast<const uint8*>(ValueA) + MemberOffset, static_cast<const uint8*>(ValueB) + MemberOffset))
 		{
 			return false;
 		}
@@ -215,11 +222,21 @@ bool InternalCompareStructProperty(const FReplicationStateDescriptor* StructDesc
 
 void InternalCopyStructProperty(const FReplicationStateDescriptor* StructDescriptor, void* RESTRICT Dst, const void* RESTRICT Src)
 {
+	uint32 FirstStructMemberForCopy = 0U;
+	if (EnumHasAnyFlags(StructDescriptor->Traits, EReplicationStateTraits::IsDerivedStruct))
+	{
+		const UScriptStruct* BaseStruct = StructDescriptor->BaseStruct;
+		BaseStruct->CopyScriptStruct(Dst, Src, 1);
+
+		// The first member in a derived struct descriptor is the base struct. We can skip it now that we've copied the value.
+		FirstStructMemberForCopy = 1U;
+	}
+
 	const FReplicationStateMemberDescriptor* MemberDescriptors = StructDescriptor->MemberDescriptors;
 	const FProperty** MemberProperties = StructDescriptor->MemberProperties;
 	const FReplicationStateMemberPropertyDescriptor* MemberPropertyDescriptors = StructDescriptor->MemberPropertyDescriptors;
 
-	for (uint32 StructMemberIt = 0, StructMemberEndIt = StructDescriptor->MemberCount; StructMemberIt != StructMemberEndIt; ++StructMemberIt)
+	for (uint32 StructMemberIt = FirstStructMemberForCopy, StructMemberEndIt = StructDescriptor->MemberCount; StructMemberIt < StructMemberEndIt; ++StructMemberIt)
 	{
 		const FReplicationStateMemberDescriptor& MemberDescriptor = MemberDescriptors[StructMemberIt];
 		const FReplicationStateMemberPropertyDescriptor& MemberPropertyDescriptor = MemberPropertyDescriptors[StructMemberIt];
@@ -260,10 +277,6 @@ void InternalCopyPropertyValue(const FReplicationStateDescriptor* Descriptor, ui
 			const uint32 ElementCount = ScriptArrayHelperSrc.Num();
 			ScriptArrayHelperDst.Resize(ElementCount);
 		
-			const FReplicationStateMemberDescriptor* MemberDescriptors = StructDescriptor->MemberDescriptors;
-			const FProperty** MemberProperties = StructDescriptor->MemberProperties;
-			const FReplicationStateMemberPropertyDescriptor* MemberPropertyDescriptors = StructDescriptor->MemberPropertyDescriptors;
-
 			// Iterate over array entries and copy the statedata using internal data
 			for (uint32 ElementIt = 0, ElementEndIt = ElementCount; ElementIt < ElementEndIt; ++ElementIt)
 			{
@@ -272,6 +285,7 @@ void InternalCopyPropertyValue(const FReplicationStateDescriptor* Descriptor, ui
 
 				InternalCopyStructProperty(StructDescriptor, ArrayDst, ArraySrc);
 			}
+
 			return;
 		}
 	}
diff --git a/Engine/Source/Runtime/Experimental/Iris/Core/Private/Iris/ReplicationState/InternalPropertyReplicationState.h b/Engine/Source/Runtime/Experimental/Iris/Core/Private/Iris/ReplicationState/InternalPropertyReplicationState.h
--- a/Engine/Source/Runtime/Experimental/Iris/Core/Private/Iris/ReplicationState/InternalPropertyReplicationState.h
+++ b/Engine/Source/Runtime/Experimental/Iris/Core/Private/Iris/ReplicationState/InternalPropertyReplicationState.h
@@ -25,7 +25,7 @@ IRISCORE_API void InternalCopyPropertyValue(const FReplicationStateDescriptor* D
 IRISCORE_API void InternalCopyStructProperty(const FReplicationStateDescriptor* StructDescriptor, void* RESTRICT Dst, const void* RESTRICT Src);
 
 /** Compare struct members using InternalCompareMember, this function will only compare replicated members of the struct */
-IRISCORE_API bool InternalCompareStructProperty(const FReplicationStateDescriptor* StructDescriptor, void* RESTRICT Dst, const void* RESTRICT Src);
+IRISCORE_API bool InternalCompareStructProperty(const FReplicationStateDescriptor* StructDescriptor, const void* RESTRICT ValueA, const void* RESTRICT ValueB);
 
 /** CompareMembers, if data is not fully replicated we will use per member compare for structs and arrays */
 IRISCORE_API bool InternalCompareMember(const FReplicationStateDescriptor* Descriptor, uint32 MemberIndex, const void* RESTRICT ValueA, const void* RESTRICT ValueB);
diff --git a/Engine/Source/Runtime/Experimental/Iris/Core/Private/Iris/ReplicationState/ReplicationStateDescriptorBuilder.cpp b/Engine/Source/Runtime/Experimental/Iris/Core/Private/Iris/ReplicationState/ReplicationStateDescriptorBuilder.cpp
--- a/Engine/Source/Runtime/Experimental/Iris/Core/Private/Iris/ReplicationState/ReplicationStateDescriptorBuilder.cpp
+++ b/Engine/Source/Runtime/Experimental/Iris/Core/Private/Iris/ReplicationState/ReplicationStateDescriptorBuilder.cpp
@@ -81,9 +81,20 @@ enum class EMemberPropertyTraits : uint32
 	HasConnectionSpecificSerialization	= IsInvalidFastArray << 1U,
 	HasPushBasedDirtiness				= HasConnectionSpecificSerialization << 1U,
 	UseSerializerIsEqual				= HasPushBasedDirtiness << 1U,
+	IsBaseStruct						= UseSerializerIsEqual << 1U,
 };
 ENUM_CLASS_FLAGS(EMemberPropertyTraits);
 
+enum class EStructNetSerializerType : unsigned
+{
+	// Struct should use StructNetSerializer
+	Struct,
+	// Struct has a custom NetSerializer
+	Custom,
+	// Struct has a parent struct in the hiearchy which has a custom NetSerializer
+	DerivedFromCustom,
+};
+
 struct FLazyGetPathNameHelper
 {
 	const FStringBuilderBase* GetPathName(const UObject* Object)
@@ -141,8 +152,9 @@ public:
 	};
 	static bool IsSupportedProperty(FMemberProperty& OutMemberProperty, const FProperty* Property, const FIsSupportedPropertyParams& Params);
 	static bool IsSupportedProperty(FMemberProperty& OutMemberProperty, const FProperty* Property) { FIsSupportedPropertyParams Params; return IsSupportedProperty(OutMemberProperty, Property, Params); }
-	static bool IsSupportedStructWithCustomSerializer(FMemberProperty& OutMemberProperty, const UStruct* InStruct);	
+	static EStructNetSerializerType IsSupportedStructWithCustomSerializer(FMemberProperty& OutMemberProperty, const UStruct* InStruct);
 	static bool IsStructWithCustomSerializer(const UStruct* InStruct);
+	static const UStruct* FindSuperStructWithCustomSerializer(const UStruct* Struct);
 
 	static bool CanStructUseStructNetSerializer(FName StructName);
 	
@@ -154,7 +166,7 @@ public:
 	static void GetIrisPropertyTraits(FMemberProperty& OutMemberProperty, const FProperty* Property, const TArray<FLifetimeProperty>* LifeTimeProperties, UClass* ObjectClass);
 
 public:
-	enum class EDescriptorType : uint32
+	enum class EDescriptorType : unsigned
 	{
 		Class,
 		Struct,
@@ -176,7 +188,9 @@ public:
 	void AddMemberProperty(const FMemberProperty& Info) { Members.Add(Info); }
 	void AddMemberFunction(const FMemberFunction& Info) { Functions.Add(Info); }
 
-	void SetStructSizeAndAlignment(SIZE_T Size, SIZE_T Alignment) { StructInfo.SizeAndAlignment.Size = Size; StructInfo.SizeAndAlignment.Alignment = Alignment; } 
+	void SetStructSizeAndAlignment(SIZE_T Size, SIZE_T Alignment) { StructInfo.SizeAndAlignment.Size = Size; StructInfo.SizeAndAlignment.Alignment = Alignment; }
+	void SetStruct(const UStruct* Struct) { StructInfo.Struct = Struct; }
+
 	bool HasDataToBuild() const { return (Members.Num() > 0) || (Functions.Num() > 0U); }
 	TRefCountPtr<const FReplicationStateDescriptor> Build(const FString& StateName, FReplicationStateDescriptorRegistry* DescriptorRegistry, const FBuildParameters& BuildParams);
 
@@ -277,6 +291,7 @@ private:
 	struct FStructInfo
 	{
 		FSizeAndAlignment SizeAndAlignment = {};
+		const UStruct* Struct = nullptr;
 	};
 
 	void AddSerializerConfigMemoryBlockToLayout(FMemoryLayoutUtil::FLayout& Layout, FOffsetAndSize& OutOffsetAndSize, const FMemberCache& MemberCache) const;
@@ -331,7 +346,7 @@ private:
 TRefCountPtr<const FReplicationStateDescriptor> FPropertyReplicationStateDescriptorBuilder::GetDescriptorForStructProperty(const FMemberProperty& MemberProperty, FReplicationStateDescriptorRegistry* DescriptorRegistry) const
 {
 	const FStructProperty* StructProp = CastFieldChecked<const FStructProperty>(MemberProperty.Property);
-	const UStruct* Struct = StructProp->Struct;
+	const UScriptStruct* Struct = StructProp->Struct;
 
 	FReplicationStateDescriptorBuilder::FParameters Params;
 	Params.DescriptorRegistry = DescriptorRegistry;
@@ -352,7 +367,7 @@ bool FPropertyReplicationStateDescriptorBuilder::ValidateStructPropertyDescripto
 	}
 
 	const FStructProperty* StructProp = CastFieldChecked<const FStructProperty>(MemberProperty.Property);
-	const UStruct* Struct = StructProp->Struct;
+	const UScriptStruct* Struct = StructProp->Struct;
 
 	if (FPropertyReplicationStateDescriptorBuilder::IsStructWithCustomSerializer(Struct))
 	{
@@ -478,7 +493,7 @@ void FPropertyReplicationStateDescriptorBuilder::BuildMemberCache(FBuilderContex
 	MemberCache.SetNum(Members.Num());
 	FMemberCacheEntry* CurrentCacheEntry = MemberCache.GetData();
 
-	bool bAreAllMembersProperties = true;
+	bool bSomeMembersAreProperties = false;
 	for (FMemberProperty& Member : Members)
 	{
 		const FNetSerializer* Serializer = Member.SerializerInfo->GetNetSerializer(Member.Property);
@@ -487,6 +502,8 @@ void FPropertyReplicationStateDescriptorBuilder::BuildMemberCache(FBuilderContex
 		// This is a property based struct, make sure we have a descriptor for it.
 		if (IsStructNetSerializer(Serializer))
 		{
+			bSomeMembersAreProperties = true;
+
 			// Build descriptor for struct
 			CurrentCacheEntry->Descriptor = GetDescriptorForStructProperty(Member, DescriptorRegistry);
 
@@ -545,6 +562,8 @@ void FPropertyReplicationStateDescriptorBuilder::BuildMemberCache(FBuilderContex
 		}
 		else if (IsArrayPropertyNetSerializer(Serializer))
 		{
+			bSomeMembersAreProperties = true;
+
 			// Build descriptor for inner property
 			CurrentCacheEntry->Descriptor = GetDescriptorForArrayProperty(Member.Property, DescriptorRegistry);
 
@@ -576,6 +595,8 @@ void FPropertyReplicationStateDescriptorBuilder::BuildMemberCache(FBuilderContex
 		}
 		else if (Member.Property)
 		{
+			bSomeMembersAreProperties = true;
+
 			CurrentCacheEntry->ExternalSizeAndAlignment = { (SIZE_T)Member.Property->ElementSize, (SIZE_T)Member.Property->GetMinAlignment() };
 			CurrentCacheEntry->InternalSizeAndAlignment = { Serializer->QuantizedTypeSize, Serializer->QuantizedTypeAlignment };
 		}
@@ -583,7 +604,6 @@ void FPropertyReplicationStateDescriptorBuilder::BuildMemberCache(FBuilderContex
 		{
 			CurrentCacheEntry->ExternalSizeAndAlignment = Member.ExternalSizeAndAlignment;
 			CurrentCacheEntry->InternalSizeAndAlignment = { Serializer->QuantizedTypeSize, Serializer->QuantizedTypeAlignment };
-			bAreAllMembersProperties = false;
 		}
 
 		CurrentCacheEntry->SerializerConfigSizeAndAlignment = { Serializer->ConfigTypeSize, Serializer->ConfigTypeAlignment };
@@ -599,12 +619,8 @@ void FPropertyReplicationStateDescriptorBuilder::BuildMemberCache(FBuilderContex
 	Context.CombinedPropertyTraits = PropertyTraits;
 	Context.SharedPropertyTraits = SharedPropertyTraits;
 
-	// Either all members should be property based or we have a single member that is not a property
-	if (!bAreAllMembersProperties)
-	{
-		check(Members.Num() <= 1);
-		Context.bGeneratePropertyDescriptors = false; 
-	}
+	// If some members are property based we generate the property descriptor.
+	Context.bGeneratePropertyDescriptors = bSomeMembersAreProperties;
 }
 
 void FPropertyReplicationStateDescriptorBuilder::BuildMemberTagCache(FMemberTagCache& MemberTagCache, const FMemberCache& MemberCache) const
@@ -724,9 +740,7 @@ void FPropertyReplicationStateDescriptorBuilder::BuildMemberFunctionCache(FBuild
 void FPropertyReplicationStateDescriptorBuilder::AllocateAndInitializeDefaultInternalStateBuffer(const uint8* RESTRICT SrcBuffer, const FReplicationStateDescriptor* Descriptor, const uint8*& OutDefaultStateBuffer) const
 {
 	// Allocate storage for incoming data, it will be freed when we destroy the descriptor
-	uint8* DstStateBuffer = (uint8*)FMemory::Malloc(Descriptor->InternalSize, Descriptor->InternalAlignment);
-	// Initialize to zero
-	FMemory::Memzero(DstStateBuffer, Descriptor->InternalSize);
+	uint8* DstStateBuffer = static_cast<uint8*>(FMemory::MallocZeroed(Descriptor->InternalSize, Descriptor->InternalAlignment));
 
 	// We setup a default context, we can however not capture object references as they need to resolve to the same value
 	// on both server and client
@@ -796,7 +810,7 @@ bool FPropertyReplicationStateDescriptorBuilder::CalculateDefaultStateChecksum(c
 	}
 
 	// State is too large, we cannot calculate a default hash for it
-	if (Writer.IsOverflown())
+	if (Context.HasErrorOrOverflow())
 	{
 		return false;
 	}
@@ -1166,7 +1180,7 @@ void FPropertyReplicationStateDescriptorBuilder::BuildMemberPropertyDescriptors(
 	uint16 ArrayIndex = 0;
 	for (const FMemberProperty& Info : Members)
 	{
-		ArrayIndex = (Info.Property == LastProperty ? uint16(ArrayIndex + 1) : uint16(0));
+		ArrayIndex = (LastProperty && Info.Property == LastProperty ? uint16(ArrayIndex + 1) : uint16(0));
 		LastProperty = Info.Property;
 
 		CurrentMemberPropertyDescriptor->RepNotifyFunction = Info.PropertyRepNotifyFunction;
@@ -1221,21 +1235,24 @@ void FPropertyReplicationStateDescriptorBuilder::BuildMemberRepIndexToMemberDesc
 void FPropertyReplicationStateDescriptorBuilder::BuildMemberDebugDescriptors(FReplicationStateMemberDebugDescriptor* MemberDebugDescriptors, FReplicationStateDescriptor* Descriptor, const FBuilderContext& Context) const
 {
 	FReplicationStateMemberDebugDescriptor* CurrentMemberDebugDescriptor = MemberDebugDescriptors;
-	if (Context.bGeneratePropertyDescriptors)
+
+	TCHAR StringBuffer[1024];
+	// Structs derived from a struct with a custom NetSerializer can contain some replicated properties.
+	for (const FMemberProperty& Info : Members)
 	{
-		TCHAR StringBuffer[1024];
-		for (const FMemberProperty& Info : Members)
+		if (const FProperty* Property = Info.Property)
 		{
-			Info.Property->GetFName().ToString(StringBuffer, 1024);
+			Property->GetFName().ToString(StringBuffer, 1024);
 			CurrentMemberDebugDescriptor->DebugName = CreatePersistentNetDebugName(StringBuffer);
-		
-			++CurrentMemberDebugDescriptor;
 		}
-	}
-	else
-	{
-		// If we dont have any property descriptor we just use the debugname of the descriptor
-		CurrentMemberDebugDescriptor->DebugName = Descriptor->DebugName;
+		else
+		{
+			// If we dont have a property tied to the member we use the debugname of the descriptor
+			CurrentMemberDebugDescriptor->DebugName = Descriptor->DebugName;
+
+		}
+
+		++CurrentMemberDebugDescriptor;
 	}
 
 	// Fill in descriptor
@@ -1418,6 +1435,12 @@ EReplicationStateTraits FPropertyReplicationStateDescriptorBuilder::BuildReplica
 		Traits |= EReplicationStateTraits::UseSerializerIsEqual;
 	}
 
+	// If one of the members have the IsBaseStruct trait we're building a descriptor for a derived struct.
+	if (EnumHasAnyFlags(PropertyTraits, EMemberPropertyTraits::IsBaseStruct))
+	{
+		Traits |= EReplicationStateTraits::IsDerivedStruct;
+	}
+
 	// Only set the PushBasedDirtiness trait if all properties are push based. An alternative could be to split into multiple states.
 	if (EnumHasAnyFlags(SharedPropertyTraits, EMemberPropertyTraits::HasPushBasedDirtiness))
 	{
@@ -1671,6 +1694,14 @@ FPropertyReplicationStateDescriptorBuilder::Build(const FString& StateName, FRep
 			BuildMemberRepIndexToMemberDescriptors(reinterpret_cast<FReplicationStateMemberRepIndexToMemberIndexDescriptor*>(Buffer + LayoutData.MemberRepIndexToMemberIndexDescriptorsSizeAndOffset.Offset), Descriptor);
 		}
 
+		// BaseStruct
+		if (EnumHasAnyFlags(Traits, EReplicationStateTraits::IsDerivedStruct))
+		{
+			const UScriptStruct* BaseStruct = Cast<const UScriptStruct>(FPropertyReplicationStateDescriptorBuilder::FindSuperStructWithCustomSerializer(StructInfo.Struct));
+			ensure(BaseStruct != nullptr);
+			Descriptor->BaseStruct = BaseStruct;
+		}
+
 		Descriptor->DebugName = CreatePersistentNetDebugName(StateName.GetCharArray().GetData());
 		if (LayoutData.MemberDebugDescriptorsSizeAndOffset.Size > 0)
 		{
@@ -1913,24 +1944,45 @@ bool FPropertyReplicationStateDescriptorBuilder::IsSupportedProperty(FMemberProp
 	return true;
 }
 
-bool FPropertyReplicationStateDescriptorBuilder::IsSupportedStructWithCustomSerializer(FMemberProperty& OutMemberProperty, const UStruct* InStruct)
+const UStruct* FPropertyReplicationStateDescriptorBuilder::FindSuperStructWithCustomSerializer(const UStruct* Struct)
 {
-	// See if we got a matching custom serializer for the entire struct
-	const FPropertyNetSerializerInfo* NetSerializerInfo = FPropertyNetSerializerInfoRegistry::FindStructSerializerInfo(InStruct->GetFName());
-	if (!NetSerializerInfo)
+	for (const UStruct* SuperStruct = Struct->GetSuperStruct(); SuperStruct != nullptr; SuperStruct = SuperStruct->GetSuperStruct())
 	{
-		if (InStruct->GetSuperStruct())
+		if (const FPropertyNetSerializerInfo* NetSerializerInfo = FPropertyNetSerializerInfoRegistry::FindStructSerializerInfo(SuperStruct->GetFName()))
 		{
-			const bool bSuperHasCustomNetSerializer = IsSupportedStructWithCustomSerializer(OutMemberProperty, InStruct->GetSuperStruct());
+			return SuperStruct;
+		}
+	}
 
-			// Warn if one of our supers have a CustomNetSerializer, if that is the case user must be made aware and either declare an explicit forwarding NetSerializer
-			// or implement a custom NetSerializer
-			const bool bIsAllowedToWarn = ShouldUseIrisReplication() || FApp::IsGame() || IsRunningDedicatedServer() || IsRunningClientOnly();
-			UE_CLOG(bSuperHasCustomNetSerializer && bIsAllowedToWarn && bWarnAboutStructsWithCustomSerialization, LogIris, Warning, TEXT("Detected derived struct %s which do not have a custom NetSerializer while Super %s has a custom NetSerializer. Silence warning by declaring an explicit forwarding serializer or implement a custom NetSerialzier for the derived struct"), ToCStr(InStruct->GetName()), ToCStr(InStruct->GetSuperStruct()->GetName()));
+	return nullptr;
+}
 
-			return bSuperHasCustomNetSerializer;
+EStructNetSerializerType FPropertyReplicationStateDescriptorBuilder::IsSupportedStructWithCustomSerializer(FMemberProperty& OutMemberProperty, const UStruct* InStruct)
+{
+	// See if we got a matching custom serializer for the entire struct
+	EStructNetSerializerType SerializerType = EStructNetSerializerType::Struct;
+	const FPropertyNetSerializerInfo* NetSerializerInfo = FPropertyNetSerializerInfoRegistry::FindStructSerializerInfo(InStruct->GetFName());
+	if (NetSerializerInfo)
+	{
+		SerializerType = EStructNetSerializerType::Custom;
+	}
+	else
+	{
+		// Find closest parent with custom serializer.
+		for (const UStruct* SuperStruct = InStruct->GetSuperStruct(); SuperStruct != nullptr; SuperStruct = SuperStruct->GetSuperStruct())
+		{
+			NetSerializerInfo = FPropertyNetSerializerInfoRegistry::FindStructSerializerInfo(SuperStruct->GetFName());
+			if (NetSerializerInfo != nullptr)
+			{
+				SerializerType = EStructNetSerializerType::DerivedFromCustom;
+				break;
+			}
 		}
-		return false;
+	}
+
+	if (!NetSerializerInfo)
+	{
+		return EStructNetSerializerType::Struct;
 	}
 
 	OutMemberProperty.Property = nullptr;
@@ -1943,9 +1995,13 @@ bool FPropertyReplicationStateDescriptorBuilder::IsSupportedStructWithCustomSeri
 	
 	// Init traits
 	OutMemberProperty.Traits = EMemberPropertyTraits::None;
+	if (SerializerType == EStructNetSerializerType::DerivedFromCustom)
+	{
+		OutMemberProperty.Traits |= EMemberPropertyTraits::IsBaseStruct;
+	}
 	GetSerializerTraits(OutMemberProperty, nullptr, NetSerializerInfo);
 
-	return true;
+	return SerializerType;
 }
 
 bool FPropertyReplicationStateDescriptorBuilder::IsStructWithCustomSerializer(const UStruct* InStruct)
@@ -2242,9 +2298,48 @@ TRefCountPtr<const FReplicationStateDescriptor> FReplicationStateDescriptorBuild
 	bool bAllMembersAreReplicated = true;
 
 	// We have a special case for structs with custom serializers
-	if (!Parameters.SkipCheckForCustomNetSerializerForStruct && FPropertyReplicationStateDescriptorBuilder::IsSupportedStructWithCustomSerializer(MemberProperty, InStruct))
+	EStructNetSerializerType SerializerType = EStructNetSerializerType::Struct;
+	if (!Parameters.SkipCheckForCustomNetSerializerForStruct)
+	{
+		SerializerType = FPropertyReplicationStateDescriptorBuilder::IsSupportedStructWithCustomSerializer(MemberProperty, InStruct);
+	}
+
+	if (SerializerType != EStructNetSerializerType::Struct)
 	{
 		Builder.AddMemberProperty(MemberProperty);
+		if (SerializerType == EStructNetSerializerType::DerivedFromCustom)
+		{
+			// Add properties introduced after parent struct with custom serializer
+			const UStruct* StructWithCustomSerializer = FPropertyReplicationStateDescriptorBuilder::FindSuperStructWithCustomSerializer(InStruct);
+			TArray<const UStruct*, TInlineAllocator<32>> Structs;
+			for (const UStruct* Struct = InStruct; Struct != StructWithCustomSerializer; Struct = Struct->GetSuperStruct())
+			{
+				// Insert struct first to traverse them in natural order when looking for replicated properties.
+				constexpr int32 InsertAtIndex = 0;
+				Structs.Insert(Struct, InsertAtIndex);
+			}
+
+			for (const UStruct* Struct : Structs)
+			{
+				for (TFieldIterator<FProperty> It(Struct, EFieldIteratorFlags::ExcludeSuper); It; ++It)
+				{
+					const FProperty* Property = *It;
+					if (EnumHasAnyFlags(Property->PropertyFlags, EPropertyFlags::CPF_RepSkip))
+					{
+						bAllMembersAreReplicated = false;
+						continue;
+					}
+
+					if (FPropertyReplicationStateDescriptorBuilder::IsSupportedProperty(MemberProperty, Property))
+					{
+						for (uint32 ArrayIt = 0, ArrayEndIt = Property->ArrayDim; ArrayIt < ArrayEndIt; ++ArrayIt)
+						{
+							Builder.AddMemberProperty(MemberProperty);
+						}
+					}
+				}
+			}
+		}
 	}
 	else
 	{
@@ -2313,6 +2408,7 @@ TRefCountPtr<const FReplicationStateDescriptor> FReplicationStateDescriptorBuild
 	}
 
 	Builder.SetStructSizeAndAlignment(InStruct->GetStructureSize(), InStruct->GetMinAlignment());
+	Builder.SetStruct(InStruct);
 
 	FLazyGetPathNameHelper LazyPathHelper;
 	FPropertyReplicationStateDescriptorBuilder::FBuildParameters BuildParameters = {};
diff --git a/Engine/Source/Runtime/Experimental/Iris/Core/Public/Iris/ReplicationState/ReplicationStateDescriptor.h b/Engine/Source/Runtime/Experimental/Iris/Core/Public/Iris/ReplicationState/ReplicationStateDescriptor.h
--- a/Engine/Source/Runtime/Experimental/Iris/Core/Public/Iris/ReplicationState/ReplicationStateDescriptor.h
+++ b/Engine/Source/Runtime/Experimental/Iris/Core/Public/Iris/ReplicationState/ReplicationStateDescriptor.h
@@ -178,7 +178,9 @@ enum class EReplicationStateTraits : uint32
 	HasPushBasedDirtiness				= HasConnectionSpecificSerialization << 1U,
 	// Whether delta compression is supported or not
 	SupportsDeltaCompression			= HasPushBasedDirtiness << 1U,
-	UseSerializerIsEqual				= SupportsDeltaCompression << 1U,	
+	UseSerializerIsEqual				= SupportsDeltaCompression << 1U,
+	// Whether this is a descriptor for a struct derived from something with a NetSerializer.
+	IsDerivedStruct						= UseSerializerIsEqual << 1U,
 };
 ENUM_CLASS_FLAGS(EReplicationStateTraits);
 
@@ -227,7 +229,10 @@ struct FReplicationStateDescriptor
 	// Non-null if RepIndexCount > 0
 	const FReplicationStateMemberRepIndexToMemberIndexDescriptor* MemberRepIndexToMemberIndexDescriptors;
 
-	// optional, 
+	// Non-null for derived struct descriptors.
+	const UScriptStruct* BaseStruct;
+
+	// Optional debug info
 	const FNetDebugName* DebugName;
 	const FReplicationStateMemberDebugDescriptor* MemberDebugDescriptors;
 
diff --git a/Engine/Source/Runtime/Experimental/Iris/Core/Public/Iris/ReplicationState/ReplicationStateDescriptorImplementationMacros.h b/Engine/Source/Runtime/Experimental/Iris/Core/Public/Iris/ReplicationState/ReplicationStateDescriptorImplementationMacros.h
--- a/Engine/Source/Runtime/Experimental/Iris/Core/Public/Iris/ReplicationState/ReplicationStateDescriptorImplementationMacros.h
+++ b/Engine/Source/Runtime/Experimental/Iris/Core/Public/Iris/ReplicationState/ReplicationStateDescriptorImplementationMacros.h
@@ -131,6 +131,7 @@ const UE::Net::FReplicationStateDescriptor StateName::sReplicationStateDescripto
 	static_cast<const UE::Net::FReplicationStateMemberPropertyDescriptor*>(nullptr), \
 	static_cast<const UE::Net::FReplicationStateMemberLifetimeConditionDescriptor*>(nullptr), \
 	static_cast<const UE::Net::FReplicationStateMemberRepIndexToMemberIndexDescriptor*>(nullptr), \
+	static_cast<const UScriptStruct*>(nullptr), \
 	UE::Net::CreatePersistentNetDebugName(TEXT(#StateName)), \
 	&sReplicationStateDescriptorMemberDebugDescriptors[0],\
 	sizeof(StateName), \
