FNetworkGUID now supports uint64 values

diff --git a/Engine/Plugins/Runtime/ReplicationSystemTestPlugin/Source/Private/TestMessage.h b/Engine/Plugins/Runtime/ReplicationSystemTestPlugin/Source/Private/TestMessage.h
--- a/Engine/Plugins/Runtime/ReplicationSystemTestPlugin/Source/Private/TestMessage.h
+++ b/Engine/Plugins/Runtime/ReplicationSystemTestPlugin/Source/Private/TestMessage.h
@@ -131,6 +131,11 @@ inline FTestMessage& operator<<(FTestMessage& Message, const FVector& Vector)
 	return Message << Vector.ToString();
 }
 
+inline FTestMessage& operator<<(FTestMessage& Message, const FVector3f& Vector)
+{
+	return Message << Vector.ToString();
+}
+
 FTestMessage& operator<<(FTestMessage& Message, const FNetHandle& NetHandle);
 
 #if defined(_WIN32) && !defined(_WIN64)
diff --git a/Engine/Plugins/Runtime/ReplicationSystemTestPlugin/Source/Private/Tests/Core/TestBitArchive.cpp b/Engine/Plugins/Runtime/ReplicationSystemTestPlugin/Source/Private/Tests/Core/TestBitArchive.cpp
new file mode 100644
--- /dev/null
+++ b/Engine/Plugins/Runtime/ReplicationSystemTestPlugin/Source/Private/Tests/Core/TestBitArchive.cpp
@@ -0,0 +1,207 @@
+// Copyright Epic Games, Inc. All Rights Reserved.
+
+#include "NetworkAutomationTest.h"
+#include "NetworkAutomationTestMacros.h"
+#include "Serialization/BitReader.h"
+#include "Serialization/BitWriter.h"
+#include <limits>
+
+namespace UE::Net::Private
+{
+	UE_NET_TEST(FNetBitArchive, SerializeIntPacked32)
+	{
+		{
+			const uint32 Values[] =
+			{
+				std::numeric_limits<uint32>::lowest(), std::numeric_limits<uint32>::max(), uint32(0), uint32(-1), (std::numeric_limits<uint32>::max() + std::numeric_limits<uint32>::lowest()) / uint32(2) - uint32(15), uint32(2048458 % std::numeric_limits<uint32>::max())
+			};
+
+			const SIZE_T ValueCount = sizeof(Values) / sizeof(Values[0]);
+
+			for (SIZE_T ValueIt = 0; ValueIt < ValueCount; ++ValueIt)
+			{
+				uint32 Value = Values[ValueIt];
+
+				FBitWriter Writer(0, true);
+				Writer.SerializeIntPacked(Value);
+
+				UE_NET_ASSERT_FALSE(Writer.IsError());
+
+				FBitReader Reader(Writer.GetData(), Writer.GetNumBits());
+
+				uint32 ResultValue = ~Value;
+				Reader.SerializeIntPacked(ResultValue);
+
+				UE_NET_ASSERT_FALSE(Reader.IsError());
+				UE_NET_ASSERT_EQ(Value, ResultValue);
+			}
+		}
+	}
+
+	UE_NET_TEST(FNetBitArchive, SerializeIntPacked64)
+	{
+		{
+			const uint64 Values[] =
+			{
+				std::numeric_limits<uint64>::lowest(), std::numeric_limits<uint64>::max(), uint64(0), uint64(-1), (std::numeric_limits<uint64>::max() + std::numeric_limits<uint64>::lowest()) / uint64(2) - uint64(15), uint64(2048458 % std::numeric_limits<uint64>::max())
+			};
+
+			const SIZE_T ValueCount = sizeof(Values) / sizeof(Values[0]);
+
+			for (SIZE_T ValueIt = 0; ValueIt < ValueCount; ++ValueIt)
+			{
+				uint64 Value = Values[ValueIt];
+
+				FBitWriter Writer(0, true);
+				Writer.SerializeIntPacked64(Value);
+
+				UE_NET_ASSERT_FALSE(Writer.IsError());
+
+				FBitReader Reader(Writer.GetData(), Writer.GetNumBits());
+
+				uint64 ResultValue = ~Value;
+				Reader.SerializeIntPacked64(ResultValue);
+
+				UE_NET_ASSERT_FALSE(Reader.IsError());
+				UE_NET_ASSERT_EQ(Value, ResultValue);
+			}
+		}
+	}
+
+	// Pulled from NetBitsTest.cpp
+	UE_NET_TEST(FNetBitArchive, SerializeInt)
+	{
+		// SerializeInt, Zero Value, 2 range
+		{
+			FBitWriter Writer(0, true);
+
+			uint32 WriteValue = 0;
+			Writer.SerializeInt(WriteValue, 2);
+
+			UE_NET_ASSERT_FALSE(Writer.IsError());
+			UE_NET_ASSERT_EQ(Writer.GetNumBits(), (int64)1);
+
+			FBitReader Reader(Writer.GetData(), Writer.GetNumBits());
+
+			uint32 ReadValue = ~WriteValue;
+			Reader.SerializeInt(ReadValue, 2);
+
+			UE_NET_ASSERT_FALSE(Reader.IsError());
+			UE_NET_ASSERT_EQ(ReadValue, WriteValue);
+		}
+
+		// SerializeInt, 1 Value, 2 range
+		{
+			FBitWriter Writer(0, true);
+
+			uint32 WriteValue = 1;
+			Writer.SerializeInt(WriteValue, 2);
+
+			UE_NET_ASSERT_FALSE(Writer.IsError());
+			UE_NET_ASSERT_EQ(Writer.GetNumBits(), (int64)1);
+
+			FBitReader Reader(Writer.GetData(), Writer.GetNumBits());
+
+			uint32 ReadValue = ~WriteValue;
+			Reader.SerializeInt(ReadValue, 2);
+
+			UE_NET_ASSERT_FALSE(Reader.IsError());
+			UE_NET_ASSERT_EQ(ReadValue, WriteValue);
+		}
+
+		// SerializeInt, 0 Value, 3 range
+		{
+			FBitWriter Writer(0, true);
+
+			uint32 WriteValue = 0;
+			Writer.SerializeInt(WriteValue, 3);
+
+			UE_NET_ASSERT_FALSE(Writer.IsError());
+			UE_NET_ASSERT_EQ(Writer.GetNumBits(), (int64)2);
+
+			FBitReader Reader(Writer.GetData(), Writer.GetNumBits());
+
+			uint32 ReadValue = ~WriteValue;
+			Reader.SerializeInt(ReadValue, 3);
+
+			UE_NET_ASSERT_FALSE(Reader.IsError());
+			UE_NET_ASSERT_EQ(ReadValue, WriteValue);
+		}
+
+		// SerializeInt, 1 Value, 3 range
+		{
+			FBitWriter Writer(0, true);
+			uint32 WriteValue = 1;
+
+			Writer.SerializeInt(WriteValue, 3);
+
+			UE_NET_ASSERT_FALSE(Writer.IsError());
+			UE_NET_ASSERT_EQ(Writer.GetNumBits(), (int64)1);
+
+			FBitReader Reader(Writer.GetData(), Writer.GetNumBits());
+
+			uint32 ReadValue = ~WriteValue;
+			Reader.SerializeInt(ReadValue, 3);
+
+			UE_NET_ASSERT_FALSE(Reader.IsError());
+			UE_NET_ASSERT_EQ(ReadValue, WriteValue);
+		}
+
+		// SerializeInt, 2 Value, 3 range
+		{
+			FBitWriter Writer(0, true);
+			uint32 WriteValue = 2;
+
+			Writer.SerializeInt(WriteValue, 3);
+
+			UE_NET_ASSERT_FALSE(Writer.IsError());
+			UE_NET_ASSERT_EQ(Writer.GetNumBits(), (int64)2);
+
+			FBitReader Reader(Writer.GetData(), Writer.GetNumBits());
+
+			uint32 ReadValue = ~WriteValue;
+			Reader.SerializeInt(ReadValue, 3);
+
+			UE_NET_ASSERT_FALSE(Reader.IsError());
+			UE_NET_ASSERT_EQ(ReadValue, WriteValue);
+		}
+
+		// SerializeInt, 0 Value, 4294967295U range
+		{
+			FBitWriter Writer(0, true);
+			uint32 WriteValue = 0;
+
+			Writer.SerializeInt(WriteValue, 4294967295U);
+
+			UE_NET_ASSERT_FALSE(Writer.IsError());
+			UE_NET_ASSERT_EQ(Writer.GetNumBits(), (int64)32);
+
+			FBitReader Reader(Writer.GetData(), Writer.GetNumBits());
+
+			uint32 ReadValue = ~WriteValue;
+			Reader.SerializeInt(ReadValue, 4294967295U);
+
+			UE_NET_ASSERT_FALSE(Reader.IsError());
+			UE_NET_ASSERT_EQ(ReadValue, WriteValue);
+		}
+
+		// SerializeInt, 4294967294U Value, 4294967295U range
+		{
+			FBitWriter Writer(0, true);
+			uint32 WriteValue = 4294967294U;
+
+			Writer.SerializeInt(WriteValue, 4294967295U);
+
+			UE_NET_ASSERT_FALSE(Writer.IsError());
+			UE_NET_ASSERT_EQ(Writer.GetNumBits(), (int64)32);
+
+			FBitReader Reader(Writer.GetData(), Writer.GetNumBits());
+
+			uint32 ReadValue = ~WriteValue;
+			Reader.SerializeInt(ReadValue, 4294967295U);
+
+			UE_NET_ASSERT_FALSE(Reader.IsError());
+			UE_NET_ASSERT_EQ(ReadValue, WriteValue);
+		}
+	}
+}
diff --git a/Engine/Plugins/Runtime/ReplicationSystemTestPlugin/Source/Private/Tests/Engine/TestNetSerialization.cpp b/Engine/Plugins/Runtime/ReplicationSystemTestPlugin/Source/Private/Tests/Engine/TestNetSerialization.cpp
new file mode 100644
--- /dev/null
+++ b/Engine/Plugins/Runtime/ReplicationSystemTestPlugin/Source/Private/Tests/Engine/TestNetSerialization.cpp
@@ -0,0 +1,273 @@
+// Copyright Epic Games, Inc. All Rights Reserved.
+
+#include "NetworkAutomationTest.h"
+#include "NetworkAutomationTestMacros.h"
+#include "Engine/NetSerialization.h"
+#include "Serialization/BitReader.h"
+#include "Serialization/BitWriter.h"
+#include <limits>
+
+namespace UE::Net::Private
+{
+	UE_NET_TEST(NetSerialization, PackedVectorFloat10)
+	{
+		{
+			const float Quantize10_Values[] =
+			{
+				0.0f,
+				-180817.42f,
+				47.11f,
+				-FMath::Exp2(25.0f),
+				std::numeric_limits<float>::infinity(), // non-finite
+			};
+
+			const int32 ValueCount = UE_ARRAY_COUNT(Quantize10_Values);
+
+			constexpr bool bAllowResize = false;
+			FBitWriter Writer(128, bAllowResize);
+
+			for (size_t ValueIt = 0, ValueEndIt = ValueCount; ValueIt != ValueEndIt; ++ValueIt)
+			{
+				Writer.Reset();
+
+				const float ScalarValue = Quantize10_Values[ValueIt];
+				const FVector3f WriteValue(ScalarValue);
+				FVector3f ReadValue;
+
+				const bool bOverflowOrNan = !WritePackedVector<10, 24>(WriteValue, Writer);
+
+				UE_NET_ASSERT_FALSE(Writer.GetError());
+
+				FBitReader Reader(Writer.GetData(), Writer.GetNumBits());
+				ReadPackedVector<10, 24>(ReadValue, Reader);
+
+				UE_NET_ASSERT_FALSE(Reader.GetError());
+
+				UE_NET_ASSERT_EQ(ReadValue.X, ReadValue.Y);
+				UE_NET_ASSERT_EQ(ReadValue.X, ReadValue.Z);
+
+				// At this point we should have similar values as the original ones, except for NaN and overflowed values
+				if (bOverflowOrNan)
+				{
+					UE_NET_ASSERT_TRUE(WriteValue.ContainsNaN());
+					UE_NET_ASSERT_EQ(ReadValue, FVector3f::ZeroVector);
+				}
+				else
+				{
+					const float ValueDiff = FMath::Abs(ReadValue.X - WriteValue.X);
+					UE_NET_ASSERT_LT(ValueDiff, 2.0f / 10); // The diff test might need some adjustment
+				}
+			}
+		}
+	}
+	
+	UE_NET_TEST(NetSerialization, PackedVectorFloat100)
+	{
+		{
+			const float Quantize100_Values[] =
+			{
+				0.0f,
+				+180720.42f,
+				-19751216.0f,
+				FMath::Exp2(31.0f),
+				-std::numeric_limits<float>::infinity(),
+			};
+
+			const int32 ValueCount = UE_ARRAY_COUNT(Quantize100_Values);
+
+			constexpr bool bAllowResize = false;
+			FBitWriter Writer(128, bAllowResize);
+
+			for (size_t ValueIt = 0, ValueEndIt = ValueCount; ValueIt != ValueEndIt; ++ValueIt)
+			{
+				Writer.Reset();
+
+				const float ScalarValue = Quantize100_Values[ValueIt];
+				const FVector3f WriteValue(ScalarValue);
+				FVector3f ReadValue;
+
+				const bool bOverflowOrNan = !WritePackedVector<100, 30>(WriteValue, Writer);
+
+				UE_NET_ASSERT_FALSE(Writer.GetError());
+
+				FBitReader Reader(Writer.GetData(), Writer.GetNumBits());
+				ReadPackedVector<100, 30>(ReadValue, Reader);
+
+				UE_NET_ASSERT_FALSE(Reader.GetError());
+
+				UE_NET_ASSERT_EQ(ReadValue.X, ReadValue.Y);
+				UE_NET_ASSERT_EQ(ReadValue.X, ReadValue.Z);
+
+				// At this point we should have similar values as the original ones, except for NaN and overflowed values
+				if (bOverflowOrNan)
+				{
+					UE_NET_ASSERT_TRUE(WriteValue.ContainsNaN());
+					UE_NET_ASSERT_EQ(ReadValue, FVector3f::ZeroVector);
+				}
+				else
+				{
+					const float ValueDiff = FMath::Abs(ReadValue.X - WriteValue.X);
+					UE_NET_ASSERT_LT(ValueDiff, 2.0f / 100); // The diff test might need some adjustment
+				}
+			}
+		}
+	}
+
+	UE_NET_TEST(NetSerialization, PackedVectorDouble10)
+	{
+		{
+			const double Quantize10_Values[] =
+			{
+				0.0f,
+				-180817.42,
+				47.11f,
+				-FMath::Exp2(25.0),
+				std::numeric_limits<double>::infinity(), // non-finite
+			};
+
+			const int32 ValueCount = UE_ARRAY_COUNT(Quantize10_Values);
+
+			constexpr bool bAllowResize = false;
+			FBitWriter Writer(512, bAllowResize);
+
+			for (size_t ValueIt = 0, ValueEndIt = ValueCount; ValueIt != ValueEndIt; ++ValueIt)
+			{
+				Writer.Reset();
+
+				const double ScalarValue = Quantize10_Values[ValueIt];
+				const FVector3d WriteValue(ScalarValue);
+				FVector3d ReadValue;
+
+				const bool bOverflowOrNan = !WritePackedVector<10, 24>(WriteValue, Writer);
+
+				UE_NET_ASSERT_FALSE(Writer.GetError());
+
+				FBitReader Reader(Writer.GetData(), Writer.GetNumBits());
+				ReadPackedVector<10, 24>(ReadValue, Reader);
+
+				UE_NET_ASSERT_FALSE(Reader.GetError());
+
+				UE_NET_ASSERT_EQ(ReadValue.X, ReadValue.Y);
+				UE_NET_ASSERT_EQ(ReadValue.X, ReadValue.Z);
+
+				// At this point we should have similar values as the original ones, except for NaN and overflowed values
+				if (bOverflowOrNan)
+				{
+					UE_NET_ASSERT_TRUE(WriteValue.ContainsNaN());
+					UE_NET_ASSERT_EQ(ReadValue, FVector3d::ZeroVector);
+				}
+				else
+				{
+					const double ValueDiff = FMath::Abs(ReadValue.X - WriteValue.X);
+					UE_NET_ASSERT_LT(ValueDiff, 2.0 / 10); // The diff test might need some adjustment
+				}
+			}
+		}
+	}
+
+	UE_NET_TEST(NetSerialization, PackedVectorDouble100)
+	{
+		{
+			const double Quantize100_Values[] =
+			{
+				0.0f,
+				+180720.42,
+				-19751216.0,
+				FMath::Exp2(40.0),
+				FMath::Exp2(59.0),
+				-std::numeric_limits<double>::infinity(),
+			};
+
+			const int32 ValueCount = UE_ARRAY_COUNT(Quantize100_Values);
+
+			constexpr bool bAllowResize = false;
+			FBitWriter Writer(512, bAllowResize);
+
+			for (size_t ValueIt = 0, ValueEndIt = ValueCount; ValueIt != ValueEndIt; ++ValueIt)
+			{
+				Writer.Reset();
+
+				const double ScalarValue = Quantize100_Values[ValueIt];
+				const FVector3d WriteValue(ScalarValue);
+				FVector3d ReadValue;
+
+				const bool bOverflowOrNan = !WritePackedVector<100, 30>(WriteValue, Writer);
+
+				UE_NET_ASSERT_FALSE(Writer.GetError());
+
+				FBitReader Reader(Writer.GetData(), Writer.GetNumBits());
+				ReadPackedVector<100, 30>(ReadValue, Reader);
+
+				UE_NET_ASSERT_FALSE(Reader.GetError());
+
+				UE_NET_ASSERT_EQ(ReadValue.X, ReadValue.Y);
+				UE_NET_ASSERT_EQ(ReadValue.X, ReadValue.Z);
+
+				// At this point we should have similar values as the original ones, except for NaN and overflowed values
+				if (bOverflowOrNan)
+				{
+					UE_NET_ASSERT_TRUE(WriteValue.ContainsNaN());
+					UE_NET_ASSERT_EQ(ReadValue, FVector3d::ZeroVector);
+				}
+				else
+				{
+					const double ValueDiff = FMath::Abs(ReadValue.X - WriteValue.X);
+					UE_NET_ASSERT_LT(ValueDiff, 2.0 / 100); // The diff test might need some adjustment
+				}
+			}
+		}
+	}
+
+	UE_NET_TEST(NetSerialization, PackedVectorWriteDoubleReadFloat)
+	{
+		{
+			const double Quantize100_Values[] =
+			{
+				0.0f,
+				+180720.42,
+				-19751216.0,
+				FMath::Exp2(40.0),
+				FMath::Exp2(59.0),
+				-std::numeric_limits<double>::infinity(),
+			};
+
+			const int32 ValueCount = UE_ARRAY_COUNT(Quantize100_Values);
+
+			constexpr bool bAllowResize = false;
+			FBitWriter Writer(512, bAllowResize);
+
+			for (size_t ValueIt = 0, ValueEndIt = ValueCount; ValueIt != ValueEndIt; ++ValueIt)
+			{
+				Writer.Reset();
+
+				const double ScalarValue = Quantize100_Values[ValueIt];
+				const FVector3d WriteValue(ScalarValue);
+				FVector3f ReadValue;
+
+				const bool bOverflowOrNan = !WritePackedVector<100, 30>(WriteValue, Writer);
+
+				UE_NET_ASSERT_FALSE(Writer.GetError());
+
+				FBitReader Reader(Writer.GetData(), Writer.GetNumBits());
+				ReadPackedVector<100, 30>(ReadValue, Reader);
+
+				UE_NET_ASSERT_FALSE(Reader.GetError());
+
+				UE_NET_ASSERT_EQ(ReadValue.X, ReadValue.Y);
+				UE_NET_ASSERT_EQ(ReadValue.X, ReadValue.Z);
+
+				// At this point we should have similar values as the original ones, except for NaN and overflowed values
+				if (bOverflowOrNan)
+				{
+					UE_NET_ASSERT_TRUE(WriteValue.ContainsNaN());
+					UE_NET_ASSERT_EQ(ReadValue, FVector3f::ZeroVector);
+				}
+				else
+				{
+					const double ValueDiff = FMath::Abs(ReadValue.X - WriteValue.X);
+					UE_NET_ASSERT_LT(ValueDiff, 2.0 / 100); // The diff test might need some adjustment
+				}
+			}
+		}
+	}
+}
diff --git a/Engine/Source/Runtime/Core/Private/Serialization/Archive.cpp b/Engine/Source/Runtime/Core/Private/Serialization/Archive.cpp
--- a/Engine/Source/Runtime/Core/Private/Serialization/Archive.cpp
+++ b/Engine/Source/Runtime/Core/Private/Serialization/Archive.cpp
@@ -1327,6 +1327,48 @@ void FArchive::SerializeIntPacked(uint32& Value)
 	}
 }
 
+void FArchive::SerializeIntPacked64(uint64& Value)
+{
+	if (IsLoading())
+	{
+		Value = 0;
+		uint8 cnt = 0;
+		uint8 more = 1;
+		while (more)
+		{
+			uint8 NextByte;
+			Serialize(&NextByte, 1);			// Read next byte
+
+			more = NextByte & 1;				// Check 1 bit to see if theres more after this
+			NextByte = NextByte >> 1;			// Shift to get actual 7 bit value
+			Value += (uint64)NextByte << (7 * cnt++);	// Add to total value
+		}
+	}
+	else
+	{
+		uint8 PackedBytes[9];
+		int32 PackedByteCount = 0;
+		uint64 Remaining = Value;
+		while (true)
+		{
+			uint8 nextByte = Remaining & 0x7f;		// Get next 7 bits to write
+			Remaining = Remaining >> 7;				// Update remaining
+			nextByte = nextByte << 1;				// Make room for 'more' bit
+			if (Remaining > 0)
+			{
+				nextByte |= 1;						// set more bit
+				PackedBytes[PackedByteCount++] = nextByte;
+			}
+			else
+			{
+				PackedBytes[PackedByteCount++] = nextByte;
+				break;
+			}
+		}
+		Serialize(PackedBytes, PackedByteCount); // Actually serialize the bytes we made
+	}
+}
+
 void FArchive::LogfImpl(const TCHAR* Fmt, ...)
 {
 	// We need to use malloc here directly as GMalloc might not be safe, e.g. if called from within GMalloc!
diff --git a/Engine/Source/Runtime/Core/Public/Misc/NetworkGuid.h b/Engine/Source/Runtime/Core/Public/Misc/NetworkGuid.h
--- a/Engine/Source/Runtime/Core/Public/Misc/NetworkGuid.h
+++ b/Engine/Source/Runtime/Core/Public/Misc/NetworkGuid.h
@@ -7,107 +7,114 @@
 
 /**
  * Implements a globally unique identifier for network related use.
- *
- * For now, this is just a uint32 with some bits meaning special things.
- * This may be expanded (beyond a uint32) eventually while keeping the API the same.
  */
 class FNetworkGUID
 {
 public:
 
-	uint32 Value;
-
-public:
+	union
+	{
+		UE_DEPRECATED(5.3, "Deprecated in favor of ObjectId")
+		uint32 Value;
+		uint64 ObjectId;
+	};
 
 	FNetworkGUID()
-		: Value(0)
+		: ObjectId(0)
 	{ }
 
-	FNetworkGUID( uint32 V )
-		: Value(V)
+	UE_DEPRECATED(5.3, "No longer used")
+	FNetworkGUID(uint32 V)
+		: ObjectId(V)
 	{ }
 		
-public:
-
-	friend bool operator==( const FNetworkGUID& X, const FNetworkGUID& Y )
+	friend bool operator==(const FNetworkGUID& X, const FNetworkGUID& Y)
 	{
-		return (X.Value == Y.Value);
+		return (X.ObjectId == Y.ObjectId);
 	}
 
-	friend bool operator!=( const FNetworkGUID& X, const FNetworkGUID& Y )
+	friend bool operator!=(const FNetworkGUID& X, const FNetworkGUID& Y)
 	{
-		return (X.Value != Y.Value);
+		return (X.ObjectId != Y.ObjectId);
 	}
 	
-	friend FArchive& operator<<( FArchive& Ar, FNetworkGUID& G )
+	friend FArchive& operator<<(FArchive& Ar, FNetworkGUID& G)
 	{
-		Ar.SerializeIntPacked(G.Value);
+		Ar.SerializeIntPacked64(G.ObjectId);
 		return Ar;
 	}
 
-public:
-
-	void BuildFromNetIndex( int32 StaticNetIndex )
+	UE_DEPRECATED(5.3, "No longer used")
+	void BuildFromNetIndex(int32 StaticNetIndex)
 	{
-		Value = (StaticNetIndex << 1 | 1);
+		ObjectId = ((uint64)StaticNetIndex << 1 | 1);
 	}
 
+	UE_DEPRECATED(5.3, "No longer used")
 	int32 ExtractNetIndex()
 	{
-		if (Value & 1)
-		{
-			return Value >> 1;
-		}
-		return 0;
+		return ((int32)ObjectId & 1) ? (int32)ObjectId >> 1 : 0;
 	}
 
-	friend uint32 GetTypeHash( const FNetworkGUID& Guid )
+	friend uint32 GetTypeHash(const FNetworkGUID& Guid)
 	{
-		return Guid.Value;
+		return ::GetTypeHash(Guid.ObjectId);
 	}
 
 	bool IsDynamic() const
 	{
-		return Value > 0 && !(Value & 1);
+		return IsValid() && !IsStatic();
 	}
 
 	bool IsStatic() const
 	{
-		return Value & 1;
+		return ObjectId & 1;
 	}
 
 	bool IsValid() const
 	{
-		return Value > 0;
+		return ObjectId > 0;
 	}
 
-	CORE_API bool NetSerialize( FArchive& Ar, class UPackageMap* Map, bool& bOutSuccess );
-
 	/** A Valid but unassigned NetGUID */
 	bool IsDefault() const
 	{
-		return (Value == 1);
+		return (ObjectId == 1);
 	}
 
-	CORE_API static FNetworkGUID GetDefault()
+	static FNetworkGUID GetDefault()
 	{
-		return FNetworkGUID(1);
+		return CreateFromIndex(0, true);
 	}
 
 	void Reset()
 	{
-		Value = 0;
+		ObjectId = 0;
 	}
 
 	FString ToString() const
 	{
-		return FString::Printf(TEXT("%d"), Value);
+		return FString::Printf(TEXT("%" UINT64_FMT), ObjectId);
 	}
 
-public:
+	bool operator<(const FNetworkGUID& Other) const
+	{
+		return ObjectId < Other.ObjectId;
+	}
+
+	UE_DEPRECATED(5.3, "Deprecated in favor of CreateFromIndex")
+	static FNetworkGUID Make(int32 Seed, bool bIsStatic)
+	{
+		return CreateFromIndex((uint64)Seed, bIsStatic);
+	}
 
-	CORE_API static FNetworkGUID Make(int32 seed, bool bIsStatic)
+	static FNetworkGUID CreateFromIndex(uint64 NetIndex, bool bIsStatic)
 	{
-		return FNetworkGUID(seed << 1 | (bIsStatic ? 1 : 0));
+		check(NetIndex <= (MAX_uint64 >> 1));
+
+		FNetworkGUID NewGuid;
+		NewGuid.ObjectId = NetIndex << 1 | (bIsStatic ? 1 : 0);
+
+		return NewGuid;
 	}
 };
diff --git a/Engine/Source/Runtime/Core/Public/Serialization/Archive.h b/Engine/Source/Runtime/Core/Public/Serialization/Archive.h
--- a/Engine/Source/Runtime/Core/Public/Serialization/Archive.h
+++ b/Engine/Source/Runtime/Core/Public/Serialization/Archive.h
@@ -1592,6 +1592,7 @@ public:
 
 	/** Packs int value into bytes of 7 bits with 8th bit for 'more' */
 	virtual void SerializeIntPacked(uint32& Value);
+	virtual void SerializeIntPacked64(uint64& Value);
 
 	/** Tells the archive to attempt to preload the specified object so data can be loaded out of it. */
 	virtual void Preload(UObject* Object) { }
diff --git a/Engine/Source/Runtime/Engine/Classes/Engine/PackageMapClient.h b/Engine/Source/Runtime/Engine/Classes/Engine/PackageMapClient.h
--- a/Engine/Source/Runtime/Engine/Classes/Engine/PackageMapClient.h
+++ b/Engine/Source/Runtime/Engine/Classes/Engine/PackageMapClient.h
@@ -221,6 +221,8 @@ public:
 
 	TMap< FNetworkGUID, FNetGuidCacheObject >		ObjectLookup;
 	TMap< TWeakObjectPtr< UObject >, FNetworkGUID >	NetGUIDLookup;
+
+	UE_DEPRECATED(5.3, "No longer used")
 	int32											UniqueNetIDs[2];
 
 	TSet< FNetworkGUID >							ImportedNetGuids;
@@ -237,6 +239,8 @@ private:
 
 	friend class UPackageMapClient;
 
+	uint64 NetworkGuidIndex[2];
+
 	TMap<FName, FNetworkGUID> PendingAsyncPackages;
 
 	/** Maps net field export group name to the respective FNetFieldExportGroup */
diff --git a/Engine/Source/Runtime/Engine/Private/DataChannel.cpp b/Engine/Source/Runtime/Engine/Private/DataChannel.cpp
--- a/Engine/Source/Runtime/Engine/Private/DataChannel.cpp
+++ b/Engine/Source/Runtime/Engine/Private/DataChannel.cpp
@@ -5482,9 +5482,14 @@ static void	FindNetGUID( const TArray<FString>& Args, UWorld* InWorld )
 		}
 		else
 		{
-			uint32 GUIDValue=0;
-			TTypeFromString<uint32>::FromString(GUIDValue, *Args[0]);
-			FNetworkGUID NetGUID(GUIDValue);
+			uint64 GUIDValue=0;
+			TTypeFromString<uint64>::FromString(GUIDValue, *Args[0]);
+
+			uint64 NetIndex = GUIDValue >> 1;
+			bool bStatic = GUIDValue & 1;
+
+			//@todo: add a from string method in dev only
+			FNetworkGUID NetGUID = FNetworkGUID::CreateFromIndex(NetIndex, bStatic);
 
 			// Search
 			FString Str = Driver->GuidCache->History.FindRef(NetGUID);
diff --git a/Engine/Source/Runtime/Engine/Private/DemoNetDriver.cpp b/Engine/Source/Runtime/Engine/Private/DemoNetDriver.cpp
--- a/Engine/Source/Runtime/Engine/Private/DemoNetDriver.cpp
+++ b/Engine/Source/Runtime/Engine/Private/DemoNetDriver.cpp
@@ -5189,19 +5189,19 @@ void UDemoNetDriver::CleanupOutstandingRewindActors()
 					// Destroy the actor before removing entries from the GuidCache so its entries are still valid in NotifyActorDestroyed.
 					World->DestroyActor(Actor, true);
 
-					ensureMsgf(GuidCache->NetGUIDLookup.Remove(CacheObject->Object) > 0, TEXT("CleanupOutstandingRewindActors: No entry found for %d in NetGUIDLookup"), NetGUID.Value);
+					ensureMsgf(GuidCache->NetGUIDLookup.Remove(CacheObject->Object) > 0, TEXT("CleanupOutstandingRewindActors: No entry found for %s in NetGUIDLookup"), *NetGUID.ToString());
 					GuidCache->ObjectLookup.Remove(NetGUID);
 					CacheObject->bNoLoad = false;
 				}
 				else
 				{
-					UE_LOG(LogDemo, Warning, TEXT("CleanupOutstandingRewindActors - Invalid object for %d, skipping."), NetGUID.Value);
+					UE_LOG(LogDemo, Warning, TEXT("CleanupOutstandingRewindActors - Invalid object for %s, skipping."), *NetGUID.ToString());
 					continue;
 				}
 			}	
 			else
 			{
-				UE_LOG(LogDemo, Warning, TEXT("CleanupOutstandingRewindActors - CacheObject not found for %s"), NetGUID.Value);
+				UE_LOG(LogDemo, Warning, TEXT("CleanupOutstandingRewindActors - CacheObject not found for %s"), *NetGUID.ToString());
 			}
 		}
 	}
diff --git a/Engine/Source/Runtime/Engine/Private/NetConnection.cpp b/Engine/Source/Runtime/Engine/Private/NetConnection.cpp
--- a/Engine/Source/Runtime/Engine/Private/NetConnection.cpp
+++ b/Engine/Source/Runtime/Engine/Private/NetConnection.cpp
@@ -3344,7 +3344,7 @@ void UNetConnection::ReceivedPacket( FBitReader& Reader, bool bIsReinjectedPacke
 
 											UE_LOG(LogNetTraffic, Log,
 													TEXT("UNetConnection::ReceivedPacket: Unable to find static actor to cleanup for ignored bunch. ")
-													TEXT("(Channel %d NetGUID %lu)"), Bunch.ChIndex, ActorGUID.Value);
+													TEXT("(Channel %d NetGUID %s)"), Bunch.ChIndex, *ActorGUID.ToString());
 										}
 									}
 								}
diff --git a/Engine/Source/Runtime/Engine/Private/PackageMapClient.cpp b/Engine/Source/Runtime/Engine/Private/PackageMapClient.cpp
--- a/Engine/Source/Runtime/Engine/Private/PackageMapClient.cpp
+++ b/Engine/Source/Runtime/Engine/Private/PackageMapClient.cpp
@@ -18,6 +18,7 @@
 #include "Net/Core/Trace/NetTrace.h"
 #include "Serialization/MemoryReader.h"
 #include "Net/NetworkGranularMemoryLogging.h"
+#include "Misc/CommandLine.h"
 
 #include UE_INLINE_GENERATED_CPP_BY_NAME(PackageMapClient)
 
@@ -1244,7 +1245,7 @@ bool UPackageMapClient::ExportNetGUIDForReplay(FNetworkGUID& NetGUID, UObject* O
 		// It's possible InternalWriteObject has modified the NetGUIDAckStatus, so
 		// do a quick sanity check to make sure the ID wasn't removed before updating the status.
 		int32* NewPacketIdPtr = OverrideAckState->NetGUIDAckStatus.Find(NetGUID);
-		if (ensureMsgf(NewPacketIdPtr != nullptr, TEXT("ExportNetGUIDForReplay PacketID was removed for %u %s"), NetGUID.Value, *GetPathNameSafe(Object)))
+		if (ensureMsgf(NewPacketIdPtr != nullptr, TEXT("ExportNetGUIDForReplay PacketID was removed for %s %s"), *NetGUID.ToString(), *GetPathNameSafe(Object)))
 		{
 			*NewPacketIdPtr = GUID_PACKET_ACKED;
 		}
@@ -2579,8 +2580,18 @@ FNetGUIDCache::FNetGUIDCache(UNetDriver* InDriver)
 	, IsExportingNetGUIDBunch(false)
 	, DelinquentAsyncLoads(GDelinquencyNumberOfTopOffendersToTrack > 0 ? GDelinquencyNumberOfTopOffendersToTrack : 0)
 {
-	UniqueNetIDs[0] = UniqueNetIDs[1] = 0;
 	UniqueNetFieldExportGroupPathIndex = 0;
+
+	uint64 NetworkGuidSeed = 0;
+
+#if !UE_BUILD_SHIPPING
+	FParse::Value(FCommandLine::Get(), TEXT("NetworkGuidSeed="), NetworkGuidSeed);
+#endif
+
+	PRAGMA_DISABLE_DEPRECATION_WARNINGS
+	UniqueNetIDs[0] = UniqueNetIDs[1] = (int32)NetworkGuidSeed;
+	PRAGMA_ENABLE_DEPRECATION_WARNINGS
+	NetworkGuidIndex[0] = NetworkGuidIndex[1] = NetworkGuidSeed;
 }
 
 class FArchiveCountMemGUID : public FArchive
@@ -2639,7 +2650,7 @@ void FNetGUIDCache::CleanReferences()
 
 			// We've seen this static object before, but we're seeing it again with a higher guid.
 			// That means this is our newly assigned GUID and we can safely time out the old one.
-			else if (FoundGuid.Value < Guid.Value)
+			else if (FoundGuid < Guid)
 			{
 				ObjectLookup[FoundGuid].ReadOnlyTimestamp = Time;
 				FoundGuid = Guid;
@@ -2834,9 +2845,6 @@ FNetworkGUID FNetGUIDCache::GetNetGUID(const UObject* Object) const
 	return NetGUID;
 }
 
-#define COMPOSE_NET_GUID( Index, IsStatic )	( ( ( Index ) << 1 ) | ( IsStatic ) )
-#define ALLOC_NEW_NET_GUID( IsStatic )		( COMPOSE_NET_GUID( ++UniqueNetIDs[ IsStatic ], IsStatic ) )
-
 /**
  *	Generate a new NetGUID for this object and assign it.
  */
@@ -2847,7 +2855,7 @@ FNetworkGUID FNetGUIDCache::AssignNewNetGUID_Server( UObject* Object )
 	// Generate new NetGUID and assign it
 	const int32 IsStatic = IsDynamicObject( Object ) ? 0 : 1;
 
-	const FNetworkGUID NewNetGuid( ALLOC_NEW_NET_GUID( IsStatic ) );
+	const FNetworkGUID NewNetGuid = FNetworkGUID::CreateFromIndex(++NetworkGuidIndex[IsStatic], IsStatic != 0);
 
 	RegisterNetGUID_Server( NewNetGuid, Object );
 
@@ -2866,7 +2874,7 @@ FNetworkGUID FNetGUIDCache::AssignNewNetGUIDFromPath_Server( const FString& Path
 	FNetworkGUID OuterGUID = GetOrAssignNetGUID( ObjOuter );
 
 	// Generate new NetGUID and assign it
-	const FNetworkGUID NewNetGuid( ALLOC_NEW_NET_GUID( 1 ) );
+	const FNetworkGUID NewNetGuid = FNetworkGUID::CreateFromIndex(++NetworkGuidIndex[1], true);
 
 	uint32 NetworkChecksum = GetClassNetworkChecksum( ObjClass );
 
@@ -2875,9 +2883,6 @@ FNetworkGUID FNetGUIDCache::AssignNewNetGUIDFromPath_Server( const FString& Path
 	return NewNetGuid;
 }
 
-#undef COMPOSE_NET_GUID
-#undef ALLOC_NEW_NET_GUID
-
 void FNetGUIDCache::RegisterNetGUID_Internal( const FNetworkGUID& NetGUID, const FNetGuidCacheObject& CacheObject )
 {
 	LLM_SCOPE_BYTAG(GuidCache);
diff --git a/Engine/Source/Runtime/Net/Core/Public/Net/Core/Trace/Private/NetTraceInternal.h b/Engine/Source/Runtime/Net/Core/Public/Net/Core/Trace/Private/NetTraceInternal.h
--- a/Engine/Source/Runtime/Net/Core/Public/Net/Core/Trace/Private/NetTraceInternal.h
+++ b/Engine/Source/Runtime/Net/Core/Public/Net/Core/Trace/Private/NetTraceInternal.h
@@ -286,7 +286,7 @@ uint32 GetObjectIdForNetTrace(const T&)
 
 inline uint32 GetObjectIdForNetTrace(const FNetworkGUID& NetGUID)
 {
-	return NetGUID.Value;
+	return (uint32)NetGUID.ObjectId;
 }
 
 template<typename T>
@@ -578,7 +578,7 @@ void FNetTraceCollector::Reset()
 	} } while (0)
 
 #define UE_NET_TRACE_INTERNAL_DYNAMIC_NAME(Name, Collector, StartPos, EndPos, Verbosity) UE_NET_TRACE_DO_IF(FNetTrace::GetCollectorAtVerbosity(Collector, Verbosity), FNetTrace::TracePacketContentEvent(*Collector, FNetTrace::TraceName(Name), StartPos, EndPos, Verbosity))
-#define UE_NET_TRACE_INTERNAL_ASSIGNED_GUID(GameInstanceId, NetGUID, PathName, OwnerId) UE_NET_TRACE_DO_IF(GNetTraceRuntimeVerbosity, FNetTrace::TraceObjectCreated(GameInstanceId, NetGUID.Value, PathName, 0U, OwnerId))
+#define UE_NET_TRACE_INTERNAL_ASSIGNED_GUID(GameInstanceId, NetGUID, PathName, OwnerId) UE_NET_TRACE_DO_IF(GNetTraceRuntimeVerbosity, FNetTrace::TraceObjectCreated(GameInstanceId, GetObjectIdForNetTrace(NetGUID), PathName, 0U, OwnerId))
 #define UE_NET_TRACE_INTERNAL_NETHANDLE_CREATED(Handle, DebugName, ProtocolId, OwnerId) FNetTrace::TraceObjectCreated(Handle.GetReplicationSystemId(), Handle.GetId(), DebugName, ProtocolId, OwnerId)
 #define UE_NET_TRACE_INTERNAL_NETHANDLE_DESTROYED(Handle) FNetTrace::TraceObjectDestroyed(Handle.GetReplicationSystemId(), Handle.GetId())
 #define UE_NET_TRACE_INTERNAL_CONNECTION_CREATED(...) FNetTrace::TraceConnectionCreated(__VA_ARGS__)
