Suppress log error when testing WritePackedVector.

#rnx

diff --git a/Engine/Plugins/Runtime/ReplicationSystemTestPlugin/Source/Private/Tests/Engine/TestNetSerialization.cpp b/Engine/Plugins/Runtime/ReplicationSystemTestPlugin/Source/Private/Tests/Engine/TestNetSerialization.cpp
--- a/Engine/Plugins/Runtime/ReplicationSystemTestPlugin/Source/Private/Tests/Engine/TestNetSerialization.cpp
+++ b/Engine/Plugins/Runtime/ReplicationSystemTestPlugin/Source/Private/Tests/Engine/TestNetSerialization.cpp
@@ -3,6 +3,7 @@
 #include "NetworkAutomationTest.h"
 #include "NetworkAutomationTestMacros.h"
 #include "Engine/NetSerialization.h"
+#include "Logging/LogScopedVerbosityOverride.h"
 #include "Serialization/BitReader.h"
 #include "Serialization/BitWriter.h"
 #include <limits>
@@ -23,6 +24,9 @@ namespace UE::Net::Private
 
 			const int32 ValueCount = UE_ARRAY_COUNT(Quantize10_Values);
 
+			// We are intentionally testing at least one value that would cause an error to be logged during serialization. We need to suppress it in order for other software that thinks that is a test fail.
+			LOG_SCOPE_VERBOSITY_OVERRIDE(LogCore, ELogVerbosity::Fatal);
+
 			constexpr bool bAllowResize = false;
 			FBitWriter Writer(128, bAllowResize);
 
@@ -75,6 +79,9 @@ namespace UE::Net::Private
 
 			const int32 ValueCount = UE_ARRAY_COUNT(Quantize100_Values);
 
+			// We are intentionally testing at least one value that would cause an error to be logged during serialization. We need to suppress it in order for other software that thinks that is a test fail.
+			LOG_SCOPE_VERBOSITY_OVERRIDE(LogCore, ELogVerbosity::Fatal);
+
 			constexpr bool bAllowResize = false;
 			FBitWriter Writer(128, bAllowResize);
 
@@ -127,6 +134,9 @@ namespace UE::Net::Private
 
 			const int32 ValueCount = UE_ARRAY_COUNT(Quantize10_Values);
 
+			// We are intentionally testing at least one value that would cause an error to be logged during serialization. We need to suppress it in order for other software that thinks that is a test fail.
+			LOG_SCOPE_VERBOSITY_OVERRIDE(LogCore, ELogVerbosity::Fatal);
+
 			constexpr bool bAllowResize = false;
 			FBitWriter Writer(512, bAllowResize);
 
@@ -180,6 +190,9 @@ namespace UE::Net::Private
 
 			const int32 ValueCount = UE_ARRAY_COUNT(Quantize100_Values);
 
+			// We are intentionally testing at least one value that would cause an error to be logged during serialization. We need to suppress it in order for other software that thinks that is a test fail.
+			LOG_SCOPE_VERBOSITY_OVERRIDE(LogCore, ELogVerbosity::Fatal);
+
 			constexpr bool bAllowResize = false;
 			FBitWriter Writer(512, bAllowResize);
 
@@ -233,6 +246,9 @@ namespace UE::Net::Private
 
 			const int32 ValueCount = UE_ARRAY_COUNT(Quantize100_Values);
 
+			// We are intentionally testing at least one value that would cause an error to be logged during serialization. We need to suppress it in order for other software that thinks that is a test fail.
+			LOG_SCOPE_VERBOSITY_OVERRIDE(LogCore, ELogVerbosity::Fatal);
+
 			constexpr bool bAllowResize = false;
 			FBitWriter Writer(512, bAllowResize);
 
