Remove log spam. Use more appropriate UE_NET_LOG in test.

#rnx

diff --git a/Engine/Plugins/Runtime/ReplicationSystemTestPlugin/Source/Private/Tests/Serialization/TestCompressedChangeMask.cpp b/Engine/Plugins/Runtime/ReplicationSystemTestPlugin/Source/Private/Tests/Serialization/TestCompressedChangeMask.cpp
--- a/Engine/Plugins/Runtime/ReplicationSystemTestPlugin/Source/Private/Tests/Serialization/TestCompressedChangeMask.cpp
+++ b/Engine/Plugins/Runtime/ReplicationSystemTestPlugin/Source/Private/Tests/Serialization/TestCompressedChangeMask.cpp
@@ -193,7 +193,7 @@ UE_NET_TEST_FIXTURE(FCompressedChangeMaskTest, WriteEmptyChangeMask)
 		// Debug data
 		const FNetBitArrayView ChangeMask = MakeNetBitArrayView(Data, BitCount);
 		const uint32 NumSetBits = ChangeMask.CountSetBits();
-		UE_LOG(LogTemp, Log, TEXT("WroteSparseArray with SetBits %u / %u using %u Bits"), NumSetBits, BitCount, Writer.GetPosBits());
+		//UE_LOG(LogTemp, Log, TEXT("WroteSparseArray with SetBits %u / %u using %u Bits"), NumSetBits, BitCount, Writer.GetPosBits());
 	};
 
 	ForAllEntries(TestWriteFunc, TestData);
@@ -219,7 +219,7 @@ UE_NET_TEST_FIXTURE(FCompressedChangeMaskTest, WriteChangeMasks)
 		// Debug data
 		const FNetBitArrayView ChangeMask = MakeNetBitArrayView(Data, BitCount);
 		const uint32 NumSetBits = ChangeMask.CountSetBits();
-		UE_LOG(LogTemp, Log, TEXT("WroteSparseArray with SetBits %u / %u using %u Bits"), NumSetBits, BitCount, Writer.GetPosBits());
+		//UE_LOG(LogTemp, Log, TEXT("WroteSparseArray with SetBits %u / %u using %u Bits"), NumSetBits, BitCount, Writer.GetPosBits());
 
 		ChangeMaskStats.TotalChangeMaskBits += BitCount;
 		ChangeMaskStats.TotalSetBits += NumSetBits;
@@ -235,7 +235,7 @@ UE_NET_TEST_FIXTURE(FCompressedChangeMaskTest, WriteChangeMasks)
 		float Ratio = ChangeMaskStats.TotalWrittenBits / (float)ChangeMaskStats.TotalChangeMaskBits;
 		float WrittenBitsPerDirtyBit = ChangeMaskStats.TotalWrittenBits / (float)ChangeMaskStats.TotalSetBits;
 
-		UE_LOG(LogTemp, Log, TEXT("WriteChangeMasks::Summary SetBits %u / TotalBits %u WrittenBits %u Ratio %f BitsPerDirtyBit %f"), ChangeMaskStats.TotalSetBits, ChangeMaskStats.TotalChangeMaskBits, ChangeMaskStats.TotalWrittenBits, Ratio, WrittenBitsPerDirtyBit);
+		UE_NET_LOG(TEXT("")) << TStringBuilder<128>().Appendf(TEXT("WriteChangeMasks::Summary SetBits %u / TotalBits %u WrittenBits %u Ratio %f BitsPerDirtyBit %f"), ChangeMaskStats.TotalSetBits, ChangeMaskStats.TotalChangeMaskBits, ChangeMaskStats.TotalWrittenBits, Ratio, WrittenBitsPerDirtyBit);
 	}
 }
 
