Iris - Offset poll frame when adding new objects to a specific interval to spread them out across frames, leading to more stable CPU times for polling.
Note that there's no rebalancing code so as objects get out of scope and new objects are created an inbalance may occur eventually.

#rnx

diff --git a/Engine/Plugins/Runtime/ReplicationSystemTestPlugin/Source/Private/Tests/ReplicationSystem/TestDependentObjects.cpp b/Engine/Plugins/Runtime/ReplicationSystemTestPlugin/Source/Private/Tests/ReplicationSystem/TestDependentObjects.cpp
--- a/Engine/Plugins/Runtime/ReplicationSystemTestPlugin/Source/Private/Tests/ReplicationSystem/TestDependentObjects.cpp
+++ b/Engine/Plugins/Runtime/ReplicationSystemTestPlugin/Source/Private/Tests/ReplicationSystem/TestDependentObjects.cpp
@@ -206,23 +206,35 @@ UE_NET_TEST_FIXTURE(FReplicationSystemServerClientTestFixture, TestDependentObje
 	UE_NET_ASSERT_EQ(ClientDependentObject->IntA, ServerDependentObject->IntA);
 
 	// Trigger replication
-	ServerObject->IntA = 1;
-	ServerDependentObject->IntA = 1;
 
-	// Send and deliver packet
-	Server->PreSendUpdate();
-	Server->SendAndDeliverTo(Client, true);
-	Server->PostSendUpdate();
+	constexpr uint32 MaxIterationCount = 256;
+	uint32 It = 0;
+	for (const uint32 EndIt = MaxIterationCount; It < EndIt; ++It)
+	{
+		ServerObject->IntA += 1;
+		ServerDependentObject->IntA += 1;
 
-	// Verify that only the server object has been updated
-	UE_NET_ASSERT_EQ(ClientObject->IntA, ServerObject->IntA);
-	UE_NET_ASSERT_NE(ClientDependentObject->IntA, ServerDependentObject->IntA);
+		// Send and deliver packet
+		Server->PreSendUpdate();
+		Server->SendAndDeliverTo(Client, true);
+		Server->PostSendUpdate();
+
+		UE_NET_ASSERT_EQ(ServerObject->IntA, ClientObject->IntA);
+
+		// Verify that only the server object has been updated
+		if (ClientDependentObject->IntA != ServerDependentObject->IntA)
+		{
+			break;
+		}
+	}
+	// At some point the object is expected not to be polled
+	UE_NET_ASSERT_LT(It, MaxIterationCount);
 
 	// Add dependency
 	Server->ReplicationBridge->AddDependentObject(ServerObject->NetRefHandle, ServerDependentObject->NetRefHandle);
 
 	// Change a value on owner
-	ServerObject->IntA = 2;
+	ServerObject->IntA += 1;
 
 	// Send and deliver packet
 	Server->PreSendUpdate();
@@ -318,11 +330,11 @@ UE_NET_TEST_FIXTURE(FReplicationSystemServerClientTestFixture, TestDependentObje
 
 	// Spawn object on server
 	UTestReplicatedIrisObject* ServerObject = Server->CreateObject(0, 0);
-	Server->ReplicationBridge->SetPollFramePeriod(ServerObject, 4U);
+	Server->ReplicationBridge->SetPollFramePeriod(ServerObject, 14U);
 
 	// Spawn second object add it as a dependency and bump poll period
 	UTestReplicatedIrisObject* ServerDependentObject = Server->CreateObject(0, 0);
-	Server->ReplicationBridge->SetPollFramePeriod(ServerDependentObject, 10U);
+	Server->ReplicationBridge->SetPollFramePeriod(ServerDependentObject, 30U);
 	Server->ReplicationBridge->AddDependentObject(ServerObject->NetRefHandle, ServerDependentObject->NetRefHandle);
 
 	// Send and deliver packet
@@ -334,18 +346,27 @@ UE_NET_TEST_FIXTURE(FReplicationSystemServerClientTestFixture, TestDependentObje
 	UTestReplicatedIrisObject* ClientObject = Cast<UTestReplicatedIrisObject>(Client->GetReplicationBridge()->GetReplicatedObject(ServerObject->NetRefHandle));
 	UTestReplicatedIrisObject* ClientDependentObject = Cast<UTestReplicatedIrisObject>(Client->GetReplicationBridge()->GetReplicatedObject(ServerDependentObject->NetRefHandle));
 
-	// Modify data
-	ServerObject->IntA = 1;
-	ServerDependentObject->IntA = 1;
+	// Modify data until none of the objects are replicated, due to the set polling frequencies.
+	constexpr uint32 MaxIterationCount = 32;
+	uint32 It = 0;
+	for (const uint32 EndIt = MaxIterationCount; It < EndIt; ++It)
+	{
+		ServerObject->IntA += 1;
+		ServerDependentObject->IntA += 1;
 
-	// Send and deliver packet, we expect nothing to replicate
-	Server->PreSendUpdate();
-	Server->SendAndDeliverTo(Client, true);
-	Server->PostSendUpdate();
+		// Send and deliver packet
+		Server->PreSendUpdate();
+		Server->SendAndDeliverTo(Client, true);
+		Server->PostSendUpdate();
 
-	// Verify that nothing replicated
-	UE_NET_ASSERT_NE(ClientObject->IntA, ServerObject->IntA);
-	UE_NET_ASSERT_NE(ClientDependentObject->IntA, ServerDependentObject->IntA);
+		// Verify that nothing replicated
+		if (ClientObject->IntA != ServerObject->IntA && ClientDependentObject->IntA != ServerDependentObject->IntA)
+		{
+			break;
+		}
+	}
+	// If we hit this assert then it's likely the poll frequency limiter is broken.
+	UE_NET_ASSERT_LT(It, MaxIterationCount);
 
 	// Mark dependent dirty
 	ReplicationSystem->MarkDirty(ServerDependentObject->NetRefHandle);
@@ -355,18 +376,18 @@ UE_NET_TEST_FIXTURE(FReplicationSystemServerClientTestFixture, TestDependentObje
 	Server->SendAndDeliverTo(Client, true);
 	Server->PostSendUpdate();
 
+	UE_NET_EXPECT_NE(ClientObject->IntA, ServerObject->IntA);
 	// Verify that dependent object has replicated
-	UE_NET_ASSERT_NE(ClientObject->IntA, ServerObject->IntA);
 	UE_NET_ASSERT_EQ(ClientDependentObject->IntA, ServerDependentObject->IntA);
 
 	// Modify data
-	ServerObject->IntA = 2;
-	ServerDependentObject->IntA = 2;
+	ServerObject->IntA += 2;
+	ServerDependentObject->IntA += 2;
 
 	// Mark parent dirty
 	ReplicationSystem->MarkDirty(ServerObject->NetRefHandle);
 
-	// Send and deliver packet, we expect both objects to have replicated
+	// Send and deliver packet
 	Server->PreSendUpdate();
 	Server->SendAndDeliverTo(Client, true);
 	Server->PostSendUpdate();
diff --git a/Engine/Plugins/Runtime/ReplicationSystemTestPlugin/Source/Private/Tests/ReplicationSystem/TestObjectPollFrequencyLimiter.cpp b/Engine/Plugins/Runtime/ReplicationSystemTestPlugin/Source/Private/Tests/ReplicationSystem/TestObjectPollFrequencyLimiter.cpp
new file mode 100644
--- /dev/null
+++ b/Engine/Plugins/Runtime/ReplicationSystemTestPlugin/Source/Private/Tests/ReplicationSystem/TestObjectPollFrequencyLimiter.cpp
@@ -0,0 +1,108 @@
+// Copyright Epic Games, Inc. All Rights Reserved.
+
+#include "NetworkAutomationTest.h"
+#include "NetworkAutomationTestMacros.h"
+#include "Net/Core/NetBitArray.h"
+#include "Iris/ReplicationSystem/ObjectPollFrequencyLimiter.h"
+#include "Math/UnrealMathUtility.h"
+
+namespace UE::Net::Private
+{
+
+UE_NET_TEST(ObjectPollFrequencyLimiter, ObjectsToPollAreSpreadOutWhenCreatingObjectAndUpdatingOncePerFrame)
+{
+	constexpr uint32 PollFrequency = 17U;
+	constexpr uint32 MaxObjectCount = PollFrequency + 1U;
+
+	FObjectPollFrequencyLimiter FrequencyLimiter;
+	FrequencyLimiter.Init(MaxObjectCount);
+
+	FNetBitArray EmptyBitArray(MaxObjectCount + 1);
+	FNetBitArray ScopeBitArray(MaxObjectCount + 1);
+
+	// Create one object at a time and update.
+	FNetBitArray OutputArray(MaxObjectCount + 1);
+	FNetBitArrayView OutputBitArrayView = MakeNetBitArrayView(OutputArray);
+
+	for (uint32 It = 0, EndIt = MaxObjectCount; It < EndIt; ++It)
+	{
+		FInternalNetRefIndex ObjectIndex = It + 1U;
+		FrequencyLimiter.SetPollFramePeriod(ObjectIndex, PollFrequency);
+
+		ScopeBitArray.SetBits(1, ObjectIndex);
+		FrequencyLimiter.Update(MakeNetBitArrayView(ScopeBitArray), MakeNetBitArrayView(EmptyBitArray), OutputBitArrayView);
+	}
+
+	// All objects are created. Verify we get one object polled per frame
+	for (uint32 It = 0, EndIt = MaxObjectCount; It < EndIt; ++It)
+	{
+		OutputBitArrayView.Reset();
+		FrequencyLimiter.Update(MakeNetBitArrayView(ScopeBitArray), MakeNetBitArrayView(EmptyBitArray), OutputBitArrayView);
+
+		const uint32 PollCount = OutputArray.CountSetBits();
+		const uint32 ExpectedPollCount = 1U;
+		UE_NET_ASSERT_EQ(PollCount, ExpectedPollCount);
+	}
+}
+
+UE_NET_TEST(ObjectPollFrequencyLimiter, ObjectsToPollAreSpreadOutRegardlessOfUpdateCountAndCreatedObjectCount)
+{
+	constexpr int32 RandSeed = 1192314641;
+	FMath::RandInit(RandSeed);
+
+	constexpr uint32 MaxObjectCount = 128;
+	FObjectPollFrequencyLimiter FrequencyLimiter;
+	FrequencyLimiter.Init(MaxObjectCount);
+
+	constexpr uint32 PollFrequency = 8U;
+	// With three times as many objects as the poll frequency we expect three objects to be polled per frame.
+	constexpr uint32 PollFrequencyObjectCount = 4U*(PollFrequency + 1U);
+	constexpr uint32 MaxUpdateCountPerIteration = 5U;
+
+	FNetBitArray EmptyBitArray(MaxObjectCount + 1);
+	FNetBitArray ScopeBitArray(MaxObjectCount + 1);
+
+	// Create an arbitrary number of objects with arbitrary number of frames in between.
+	FNetBitArray OutputArray(MaxObjectCount + 1);
+	{
+		FNetBitArrayView OutputBitArrayView = MakeNetBitArrayView(OutputArray);
+
+		uint32 CreatedObjectCount = 0;
+		do
+		{
+			int32 UpdateCount = FMath::RandHelper(MaxUpdateCountPerIteration + 1);
+			while (UpdateCount-- > 0)
+			{
+				FrequencyLimiter.Update(MakeNetBitArrayView(ScopeBitArray), MakeNetBitArrayView(EmptyBitArray), OutputBitArrayView);
+			}
+
+			int32 CreateCount = FMath::RandHelper(FMath::Min(3U, PollFrequencyObjectCount - CreatedObjectCount) + 1);
+			while (CreateCount-- > 0)
+			{
+				++CreatedObjectCount;
+				FInternalNetRefIndex ObjectIndex = CreatedObjectCount;
+				FrequencyLimiter.SetPollFramePeriod(ObjectIndex, PollFrequency);
+			}
+
+			if (CreatedObjectCount > 0)
+			{
+				ScopeBitArray.SetBits(1, CreatedObjectCount);
+			}
+
+		} while (CreatedObjectCount < PollFrequencyObjectCount);
+	}
+
+	// Verify the objects get polled evenly.
+	for (uint32 PollIt = 0; PollIt <= PollFrequency; ++PollIt)
+	{
+		OutputArray.Reset();
+		FNetBitArrayView OutputBitArrayView = MakeNetBitArrayView(OutputArray);
+		FrequencyLimiter.Update(MakeNetBitArrayView(ScopeBitArray), MakeNetBitArrayView(EmptyBitArray), OutputBitArrayView);
+
+		uint32 PollCount = OutputArray.CountSetBits();
+		const uint32 ExpectedPollCount = PollFrequencyObjectCount/(PollFrequency + 1U);
+		UE_NET_ASSERT_EQ(PollCount, ExpectedPollCount);
+	}
+}
+
+}
diff --git a/Engine/Plugins/Runtime/ReplicationSystemTestPlugin/Source/Private/Tests/ReplicationSystem/TestReplicationSystem.cpp b/Engine/Plugins/Runtime/ReplicationSystemTestPlugin/Source/Private/Tests/ReplicationSystem/TestReplicationSystem.cpp
--- a/Engine/Plugins/Runtime/ReplicationSystemTestPlugin/Source/Private/Tests/ReplicationSystem/TestReplicationSystem.cpp
+++ b/Engine/Plugins/Runtime/ReplicationSystemTestPlugin/Source/Private/Tests/ReplicationSystem/TestReplicationSystem.cpp
@@ -1676,18 +1676,26 @@ UE_NET_TEST_FIXTURE(FReplicationSystemServerClientTestFixture, TestObjectPollFra
 	UE_NET_ASSERT_EQ(ClientObject->IntA, ServerObject->IntA);
 	UE_NET_ASSERT_EQ(ClientObjectPolledEveryOtherFrame->IntA, ServerObjectPolledEveryOtherFrame->IntA);
 
-	// Trigger replication
-	ServerObject->IntA = 1;
-	ServerObjectPolledEveryOtherFrame->IntA = 1;
+	// After two value updates it's expected that the polling occurs exactly one time for the object with poll frame period 1 (meaning every other frame).
+	bool SlowPollObjectHasBeenEqual = false;
+	bool SlowPollObjectHasBeenInequal = false;
+
+	// Update values
+	ServerObject->IntA += 1;
+	ServerObjectPolledEveryOtherFrame->IntA += 1;
 
 	// Send and deliver packet
 	Server->PreSendUpdate();
 	Server->SendAndDeliverTo(Client, true);
 	Server->PostSendUpdate();
 
-	// Verify that only the server object has been updated
 	UE_NET_ASSERT_EQ(ClientObject->IntA, ServerObject->IntA);
-	UE_NET_ASSERT_NE(ClientObjectPolledEveryOtherFrame->IntA, ServerObjectPolledEveryOtherFrame->IntA);
+	SlowPollObjectHasBeenEqual |= (ClientObjectPolledEveryOtherFrame->IntA == ServerObjectPolledEveryOtherFrame->IntA);
+	SlowPollObjectHasBeenInequal |= (ClientObjectPolledEveryOtherFrame->IntA != ServerObjectPolledEveryOtherFrame->IntA);
+
+	// Update values
+	ServerObject->IntA += 1;
+	ServerObjectPolledEveryOtherFrame->IntA += 1;
 
 	// Send and deliver packet
 	Server->PreSendUpdate();
@@ -1696,7 +1704,11 @@ UE_NET_TEST_FIXTURE(FReplicationSystemServerClientTestFixture, TestObjectPollFra
 
 	// Verify that both objects now are in sync
 	UE_NET_ASSERT_EQ(ClientObject->IntA, ServerObject->IntA);
-	UE_NET_ASSERT_EQ(ClientObjectPolledEveryOtherFrame->IntA, ServerObjectPolledEveryOtherFrame->IntA);
+	SlowPollObjectHasBeenEqual |= (ClientObjectPolledEveryOtherFrame->IntA == ServerObjectPolledEveryOtherFrame->IntA);
+	SlowPollObjectHasBeenInequal |= (ClientObjectPolledEveryOtherFrame->IntA != ServerObjectPolledEveryOtherFrame->IntA);
+
+	UE_NET_ASSERT_TRUE(SlowPollObjectHasBeenEqual);
+	UE_NET_ASSERT_TRUE(SlowPollObjectHasBeenInequal);
 }
 
 }
diff --git a/Engine/Source/Runtime/Experimental/Iris/Core/Private/Iris/ReplicationSystem/ObjectPollFrequencyLimiter.h b/Engine/Source/Runtime/Experimental/Iris/Core/Private/Iris/ReplicationSystem/ObjectPollFrequencyLimiter.h
--- a/Engine/Source/Runtime/Experimental/Iris/Core/Private/Iris/ReplicationSystem/ObjectPollFrequencyLimiter.h
+++ b/Engine/Source/Runtime/Experimental/Iris/Core/Private/Iris/ReplicationSystem/ObjectPollFrequencyLimiter.h
@@ -49,9 +49,9 @@ inline void FObjectPollFrequencyLimiter::SetPollFramePeriod(FInternalNetRefIndex
 	MaxInternalHandle = FPlatformMath::Max(MaxInternalHandle, InternalIndex);
 
 	FramesBetweenUpdates[InternalIndex] = PollFramePeriod;
-	// Spread the polling of objects with the same frequency so that if you add lots of objects the same frame they won't be polled at the same time.
-	const uint32 FrameOffset = FrameIndexOffsets[PollFramePeriod]++;
-	FrameCounters[InternalIndex] = FrameOffset % (uint32(PollFramePeriod) + 1U);
+	// Spread the polling of objects with the same frequency so that if you add lots of objects the same frame they won't be polled at the same time. The update loop decrements counters so we need to be careful with how we offset things.
+	const uint8 FrameOffset = --FrameIndexOffsets[PollFramePeriod];
+	FrameCounters[InternalIndex] = static_cast<uint8>(uint32(~(FrameIndex + FrameOffset)) % uint32(PollFramePeriod + 1U));
 }
 
 inline void FObjectPollFrequencyLimiter::SetPollWithObject(FInternalNetRefIndex ObjectToPollWithInternalIndex, FInternalNetRefIndex InternalIndex)
