Iris Filter UnitTest:
* Split Mock filter into a seperate class when it needs to read fragment data of an object

diff --git a/Engine/Plugins/Runtime/ReplicationSystemTestPlugin/Source/Private/Tests/ReplicationSystem/Filtering/MockNetObjectFilter.cpp b/Engine/Plugins/Runtime/ReplicationSystemTestPlugin/Source/Private/Tests/ReplicationSystem/Filtering/MockNetObjectFilter.cpp
--- a/Engine/Plugins/Runtime/ReplicationSystemTestPlugin/Source/Private/Tests/ReplicationSystem/Filtering/MockNetObjectFilter.cpp
+++ b/Engine/Plugins/Runtime/ReplicationSystemTestPlugin/Source/Private/Tests/ReplicationSystem/Filtering/MockNetObjectFilter.cpp
@@ -59,14 +59,6 @@ bool UMockNetObjectFilter::AddObject(uint32 ObjectIndex, FNetObjectFilterAddObje
 
 		AddedObjectIndices.SetBit(ObjectIndex);
 		++AddedCount;
-
-		// Check if object has the NetTest_FilterOut RepTag which says whether the object should be filtered out or not.
-		UE::Net::FRepTagFindInfo RepTagInfo;
-		if (UE::Net::FindRepTag(Params.Protocol, RepTag_NetTest_FilterOut, RepTagInfo))
-		{
-			// Warning: Using the internal state requires the proper NetSerializer and dequantization to get the value.
-			ObjectToFilterOutOffset.Add(ObjectIndex, RepTagInfo.InternalStateAbsoluteOffset);
-		}
 	}
 
 	return CallSetup.AddObject.bReturnValue;
@@ -91,9 +83,6 @@ void UMockNetObjectFilter::RemoveObject(uint32 ObjectIndex, const FNetObjectFilt
 
 		AddedObjectIndices.ClearBit(ObjectIndex);
 		--AddedCount;
-
-		ObjectToFilterOut.Remove(ObjectIndex);
-		ObjectToFilterOutOffset.Remove(ObjectIndex);
 	}
 }
 
@@ -105,14 +94,8 @@ void UMockNetObjectFilter::UpdateObjects(FNetObjectFilterUpdateParams& Params)
 	for (uint32 ObjectIndex : MakeArrayView(Params.ObjectIndices, Params.ObjectCount))
 	{
 		bIsProperCall = bIsProperCall && AddedObjectIndices.GetBit(ObjectIndex);
-
-		if (UPTRINT* InternalStateOffset = ObjectToFilterOutOffset.Find(ObjectIndex))
-		{
-			// Warning: This is generally not safe, not even for primitive types.
-			const bool bFilterOut = *reinterpret_cast<const bool*>(Params.StateBuffers[ObjectIndex] + *InternalStateOffset);
-			ObjectToFilterOut.Emplace(ObjectIndex, bFilterOut);
-		}
 	}
+
 	CallStatus.SuccessfulCallCounts.UpdateObjects += bIsProperCall;
 }
 
@@ -140,12 +123,6 @@ void UMockNetObjectFilter::Filter(FNetObjectFilteringParams& Params)
 	{
 		Params.OutAllowedObjects.Copy(MakeNetBitArrayView(AddedObjectIndices));
 	}
-
-	// Go through objects with FilterOut wishes
-	for (const auto& Pair : ObjectToFilterOut)
-	{
-		Params.OutAllowedObjects.SetBitValue(Pair.Key, !Pair.Value);
-	}
 }
 
 void UMockNetObjectFilter::PostFilter(FNetObjectPostFilteringParams& Params)
@@ -155,3 +132,75 @@ void UMockNetObjectFilter::PostFilter(FNetObjectPostFilteringParams& Params)
 	constexpr bool bIsProperCall = true;
 	CallStatus.SuccessfulCallCounts.PostFilter += bIsProperCall;
 }
+
+
+//**************************************************************************************************
+// UMockNetObjectFilterUsingFragmentData
+//**************************************************************************************************
+
+void UMockNetObjectFilterUsingFragmentData::Init(FNetObjectFilterInitParams& Params)
+{
+	Super::Init(Params);
+
+	//TODO: Set to PostFilter pass when added
+}
+
+bool UMockNetObjectFilterUsingFragmentData::AddObject(uint32 ObjectIndex, FNetObjectFilterAddObjectParams& Params)
+{
+	bool bResult = Super::AddObject(ObjectIndex, Params);
+
+	if (bResult)
+	{
+		// Check if object has the NetTest_FilterOut RepTag which says whether the object should be filtered out or not.
+		UE::Net::FRepTagFindInfo RepTagInfo;
+		if (UE::Net::FindRepTag(Params.Protocol, RepTag_NetTest_FilterOut, RepTagInfo))
+		{
+			// Warning: Using the internal state requires the proper NetSerializer and dequantization to get the value.
+			ObjectToFilterOutOffset.Add(ObjectIndex, RepTagInfo.InternalStateAbsoluteOffset);
+		}
+		else
+		{
+			ensureAlwaysMsgf(false, TEXT("Could not find FilterOut reptag in object index %u"), ObjectIndex);
+		}
+	}
+
+	return bResult;
+}
+
+void UMockNetObjectFilterUsingFragmentData::RemoveObject(uint32 ObjectIndex, const FNetObjectFilteringInfo& Params)
+{
+	Super::RemoveObject(ObjectIndex, Params);
+
+	ObjectToFilterOut.Remove(ObjectIndex);
+	ObjectToFilterOutOffset.Remove(ObjectIndex);
+}
+
+void UMockNetObjectFilterUsingFragmentData::UpdateObjects(FNetObjectFilterUpdateParams& Params)
+{
+	++CallStatus.CallCounts.UpdateObjects;
+
+	bool bIsProperCall = Params.ObjectCount <= AddedCount;
+	for (uint32 ObjectIndex : MakeArrayView(Params.ObjectIndices, Params.ObjectCount))
+	{
+		bIsProperCall = bIsProperCall && AddedObjectIndices.GetBit(ObjectIndex);
+
+		if (UPTRINT* InternalStateOffset = ObjectToFilterOutOffset.Find(ObjectIndex))
+		{
+			// Warning: This is generally not safe, not even for primitive types.
+			const bool bFilterOut = *reinterpret_cast<const bool*>(Params.StateBuffers[ObjectIndex] + *InternalStateOffset);
+			ObjectToFilterOut.Emplace(ObjectIndex, bFilterOut);
+		}
+	}
+	CallStatus.SuccessfulCallCounts.UpdateObjects += bIsProperCall;
+}
+
+void UMockNetObjectFilterUsingFragmentData::Filter(FNetObjectFilteringParams& Params)
+{
+	Super::Filter(Params);
+
+	// Go through objects with FilterOut wishes
+	for (const auto& Pair : ObjectToFilterOut)
+	{
+		Params.OutAllowedObjects.SetBitValue(Pair.Key, !Pair.Value);
+	}
+}
diff --git a/Engine/Plugins/Runtime/ReplicationSystemTestPlugin/Source/Private/Tests/ReplicationSystem/Filtering/MockNetObjectFilter.h b/Engine/Plugins/Runtime/ReplicationSystemTestPlugin/Source/Private/Tests/ReplicationSystem/Filtering/MockNetObjectFilter.h
--- a/Engine/Plugins/Runtime/ReplicationSystemTestPlugin/Source/Private/Tests/ReplicationSystem/Filtering/MockNetObjectFilter.h
+++ b/Engine/Plugins/Runtime/ReplicationSystemTestPlugin/Source/Private/Tests/ReplicationSystem/Filtering/MockNetObjectFilter.h
@@ -76,7 +76,7 @@ protected:
 	virtual void Filter(FNetObjectFilteringParams&) override;
 	virtual void PostFilter(FNetObjectPostFilteringParams&) override;
 
-private:
+protected:
 	UMockNetObjectFilter();
 
 	FFunctionCallStatus CallStatus;
@@ -85,8 +85,27 @@ private:
 	UE::Net::FNetBitArray AddedConnectionIndices;
 	SIZE_T AddedCount;
 
-	TMap<uint32, UPTRINT> ObjectToFilterOutOffset;
-	TMap<uint32, bool> ObjectToFilterOut;
-
 	static constexpr float DefaultPriority = 1.0f;
 };
+
+/**
+ * Filter that reads fragment data to decide if an object is filtered out or not
+ */
+UCLASS()
+class UMockNetObjectFilterUsingFragmentData : public UMockNetObjectFilter
+{
+	GENERATED_BODY()
+
+protected:
+
+	virtual void Init(FNetObjectFilterInitParams&) override;
+	virtual bool AddObject(uint32 ObjectIndex, FNetObjectFilterAddObjectParams&) override;
+	virtual void RemoveObject(uint32 ObjectIndex, const FNetObjectFilteringInfo&) override;
+	virtual void UpdateObjects(FNetObjectFilterUpdateParams&) override;
+	virtual void Filter(FNetObjectFilteringParams&) override;
+
+private:
+
+	TMap<uint32, UPTRINT> ObjectToFilterOutOffset;
+	TMap<uint32, bool> ObjectToFilterOut;
+};
\ No newline at end of file
diff --git a/Engine/Plugins/Runtime/ReplicationSystemTestPlugin/Source/Private/Tests/ReplicationSystem/Filtering/TestFiltering.cpp b/Engine/Plugins/Runtime/ReplicationSystemTestPlugin/Source/Private/Tests/ReplicationSystem/Filtering/TestFiltering.cpp
--- a/Engine/Plugins/Runtime/ReplicationSystemTestPlugin/Source/Private/Tests/ReplicationSystem/Filtering/TestFiltering.cpp
+++ b/Engine/Plugins/Runtime/ReplicationSystemTestPlugin/Source/Private/Tests/ReplicationSystem/Filtering/TestFiltering.cpp
@@ -41,10 +41,20 @@ private:
 
 		// Modify definitions to only include our mock prioritizer. Ugly... 
 		TArray<FNetObjectFilterDefinition> NewFilterDefinitions;
-		FNetObjectFilterDefinition& MockDefinition = NewFilterDefinitions.Emplace_GetRef();
-		MockDefinition.FilterName = "MockFilter";
-		MockDefinition.ClassName = "/Script/ReplicationSystemTestPlugin.MockNetObjectFilter";
-		MockDefinition.ConfigClassName = "/Script/ReplicationSystemTestPlugin.MockNetObjectFilterConfig";
+		{
+			FNetObjectFilterDefinition& MockDefinition = NewFilterDefinitions.Emplace_GetRef();
+			MockDefinition.FilterName = "MockFilter";
+			MockDefinition.ClassName = "/Script/ReplicationSystemTestPlugin.MockNetObjectFilter";
+			MockDefinition.ConfigClassName = "/Script/ReplicationSystemTestPlugin.MockNetObjectFilterConfig";
+		}
+
+		{
+			FNetObjectFilterDefinition& MockDefinition = NewFilterDefinitions.Emplace_GetRef();
+			MockDefinition.FilterName = "MockFilterWithFragments";
+			MockDefinition.ClassName = "/Script/ReplicationSystemTestPlugin.MockNetObjectFilterUsingFragmentData";
+			MockDefinition.ConfigClassName = "/Script/ReplicationSystemTestPlugin.MockNetObjectFilterConfig";
+		}
+
 		DefinitionsProperty->CopyCompleteValue((void*)(UPTRINT(FilterDefinitions) + DefinitionsProperty->GetOffset_ForInternal()), &NewFilterDefinitions);
 	}
 
@@ -65,12 +75,18 @@ private:
 	{
 		MockNetObjectFilter = ExactCast<UMockNetObjectFilter>(Server->GetReplicationSystem()->GetFilter("MockFilter"));
 		MockFilterHandle = Server->GetReplicationSystem()->GetFilterHandle("MockFilter");
+
+		MockNetObjectFilterWithFragments = ExactCast<UMockNetObjectFilterUsingFragmentData>(Server->GetReplicationSystem()->GetFilter("MockFilterWithFragments"));
+		MockFilterWithFragmentsHandle = Server->GetReplicationSystem()->GetFilterHandle("MockFilterWithFragments");
 	}
 
 protected:
 	UMockNetObjectFilter* MockNetObjectFilter;
 	FNetObjectFilterHandle MockFilterHandle;
 
+	UMockNetObjectFilterUsingFragmentData* MockNetObjectFilterWithFragments;
+	FNetObjectFilterHandle MockFilterWithFragmentsHandle;
+
 private:
 	TArray<FNetObjectFilterDefinition> OriginalFilterDefinitions;
 };
@@ -1139,8 +1155,8 @@ UE_NET_TEST_FIXTURE(FTestFilteringFixture, ObjectGetsFilterOutSettingFromStart)
 		UMockNetObjectFilter::FFunctionCallSetup CallSetup;
 		CallSetup.AddObject.bReturnValue = true;
 		CallSetup.Filter.bFilterOutByDefault = false;
-		MockNetObjectFilter->SetFunctionCallSetup(CallSetup);
-		MockNetObjectFilter->ResetFunctionCallStatus();
+		MockNetObjectFilterWithFragments->SetFunctionCallSetup(CallSetup);
+		MockNetObjectFilterWithFragments->ResetFunctionCallStatus();
 	}
 
 	// Add client
@@ -1148,7 +1164,7 @@ UE_NET_TEST_FIXTURE(FTestFilteringFixture, ObjectGetsFilterOutSettingFromStart)
 
 	// Create object and set filter
 	UTestFilteringObject* ServerObject = Server->CreateObject<UTestFilteringObject>();
-	Server->ReplicationSystem->SetFilter(ServerObject->NetRefHandle, MockFilterHandle);
+	Server->ReplicationSystem->SetFilter(ServerObject->NetRefHandle, MockFilterWithFragmentsHandle);
 
 	// We want the object to be filtered out
 	constexpr bool bFilterOut = true;
@@ -1170,17 +1186,8 @@ UE_NET_TEST_FIXTURE(FTestFilteringFixture, ObjectGetsUpdatedFilterOutSetting)
 		UMockNetObjectFilter::FFunctionCallSetup CallSetup;
 		CallSetup.AddObject.bReturnValue = true;
 		CallSetup.Filter.bFilterOutByDefault = false;
-		MockNetObjectFilter->SetFunctionCallSetup(CallSetup);
-		MockNetObjectFilter->ResetFunctionCallStatus();
-	}
-
-	// Setup dynamic filter for the test. For this test we want the value of the object's filter property to rule.
-	{
-		UMockNetObjectFilter::FFunctionCallSetup CallSetup;
-		CallSetup.AddObject.bReturnValue = true;
-		CallSetup.Filter.bFilterOutByDefault = false;
-		MockNetObjectFilter->SetFunctionCallSetup(CallSetup);
-		MockNetObjectFilter->ResetFunctionCallStatus();
+		MockNetObjectFilterWithFragments->SetFunctionCallSetup(CallSetup);
+		MockNetObjectFilterWithFragments->ResetFunctionCallStatus();
 	}
 
 	// Add client
@@ -1188,7 +1195,7 @@ UE_NET_TEST_FIXTURE(FTestFilteringFixture, ObjectGetsUpdatedFilterOutSetting)
 
 	// Create object and set filter
 	UTestFilteringObject* ServerObject = Server->CreateObject<UTestFilteringObject>();
-	Server->ReplicationSystem->SetFilter(ServerObject->NetRefHandle, MockFilterHandle);
+	Server->ReplicationSystem->SetFilter(ServerObject->NetRefHandle, MockFilterWithFragmentsHandle);
 
 	// We don't want the object to be filtered out
 	ServerObject->SetFilterOut(false);
