Iris - Huge object going out and in to scope now correctly end up flushing state.

diff --git a/Engine/Plugins/Runtime/ReplicationSystemTestPlugin/Source/Private/Tests/ReplicationSystem/TestObjectSplitting.cpp b/Engine/Plugins/Runtime/ReplicationSystemTestPlugin/Source/Private/Tests/ReplicationSystem/TestObjectSplitting.cpp
--- a/Engine/Plugins/Runtime/ReplicationSystemTestPlugin/Source/Private/Tests/ReplicationSystem/TestObjectSplitting.cpp
+++ b/Engine/Plugins/Runtime/ReplicationSystemTestPlugin/Source/Private/Tests/ReplicationSystem/TestObjectSplitting.cpp
@@ -864,4 +864,203 @@ UE_NET_TEST_FIXTURE(FSplitObjectTestFixture, TestCancelPendingDestroyDuringHugeO
 	UE_NET_ASSERT_EQ(ClientObject->IntA, ServerObject->IntA);
 }
 
+UE_NET_TEST_FIXTURE(FSplitObjectTestFixture, TestReliableAttachmentIsDeliveredDespiteHugeObjectBeingDestroyed)
+{
+	// Add a client
+	FReplicationSystemTestClient* Client = CreateClient();
+	RegisterNetBlobHandlers(Client);
+
+	// Spawn object on server
+	UTestReplicatedIrisObject* ServerObject = CreateObject(Server);
+
+	// Send and deliver packet
+	Server->PreSendUpdate();
+	Server->SendAndDeliverTo(Client, DeliverPacket);
+	Server->PostSendUpdate();
+
+	// Force huge object state
+	SetObjectPayloadByteCount(ServerObject, HugeObjectPayloadByteCount);
+
+	// Add reliable attachment
+	{
+		TRefCountPtr<FNetObjectAttachment> Attachment = ServerMockNetObjectAttachmentHandler->CreateReliableNetObjectAttachment(1U);
+
+		FNetObjectReference AttachmentTarget = FObjectReferenceCache::MakeNetObjectReference(ServerObject->NetRefHandle);
+		Server->GetReplicationSystem()->QueueNetObjectAttachment(Client->ConnectionIdOnServer, AttachmentTarget, Attachment);
+	}
+
+	// Write packets
+	for (uint32 RetryIt = 0; RetryIt != HugeObjectMaxNetTickCountToArrive; ++RetryIt)
+	{
+		Server->PreSendUpdate();
+		Server->SendUpdate(Client->ConnectionIdOnServer);
+		Server->PostSendUpdate();
+	}
+
+	// Filter out object to cause object to be set in state WaitOnFlush
+	Server->GetReplicationSystem()->AddToGroup(NotReplicatedNetObjectGroupHandle, ServerObject->NetRefHandle);
+	Server->PreSendUpdate();
+	Server->SendUpdate(Client->ConnectionIdOnServer);
+	Server->PostSendUpdate();
+
+	// Deliver huge state
+	{
+		SIZE_T PacketCount = 0;
+		const auto& ConnectionInfo = Server->GetConnectionInfo(Client->ConnectionIdOnServer);
+		PacketCount = ConnectionInfo.WrittenPackets.Count();
+		for (SIZE_T PacketIt = 0; PacketIt != PacketCount; ++PacketIt)
+		{
+			Server->DeliverTo(Client, DeliverPacket);
+		}
+	}
+
+	// Deliver latest state
+	Server->PreSendUpdate();
+	Server->SendAndDeliverTo(Client, DeliverPacket);
+	Server->PostSendUpdate();
+
+	// Verify the attachment made it through, despite the wish to destroy the object.
+	const UMockNetObjectAttachmentHandler::FCallCounts AttachmentCallCounts = ClientMockNetObjectAttachmentHandler->GetFunctionCallCounts();
+	UE_NET_ASSERT_EQ(AttachmentCallCounts.OnNetBlobReceived, 1U);
+
+	// The object should not exist
+	UTestReplicatedIrisObject* ClientObject = Cast<UTestReplicatedIrisObject>(Client->GetReplicationBridge()->GetReplicatedObject(ServerObject->NetRefHandle));
+	UE_NET_ASSERT_EQ(ClientObject, nullptr);
+}
+
+UE_NET_TEST_FIXTURE(FSplitObjectTestFixture, TestHugeObjectIsFlushedAndNotDestroyedWhenFilteredOutAndThenIn)
+{
+	// Add a client
+	FReplicationSystemTestClient* Client = CreateClient();
+	RegisterNetBlobHandlers(Client);
+
+	// Spawn object on server
+	UTestReplicatedIrisObject* ServerObject = CreateObject(Server);
+
+	// Send and deliver packet
+	Server->PreSendUpdate();
+	Server->SendAndDeliverTo(Client, DeliverPacket);
+	Server->PostSendUpdate();
+
+	// Force huge object state
+	SetObjectPayloadByteCount(ServerObject, HugeObjectPayloadByteCount);
+
+	// Add reliable attachment
+	{
+		TRefCountPtr<FNetObjectAttachment> Attachment = ServerMockNetObjectAttachmentHandler->CreateReliableNetObjectAttachment(1U);
+
+		FNetObjectReference AttachmentTarget = FObjectReferenceCache::MakeNetObjectReference(ServerObject->NetRefHandle);
+		Server->GetReplicationSystem()->QueueNetObjectAttachment(Client->ConnectionIdOnServer, AttachmentTarget, Attachment);
+	}
+
+	// Write packets
+	for (uint32 RetryIt = 0; RetryIt != HugeObjectMaxNetTickCountToArrive; ++RetryIt)
+	{
+		Server->PreSendUpdate();
+		Server->SendUpdate(Client->ConnectionIdOnServer);
+		Server->PostSendUpdate();
+	}
+
+	// Filter out object to cause object to be set in state WaitOnFlush
+	Server->GetReplicationSystem()->AddToGroup(NotReplicatedNetObjectGroupHandle, ServerObject->NetRefHandle);
+	Server->PreSendUpdate();
+	Server->SendUpdate(Client->ConnectionIdOnServer);
+	Server->PostSendUpdate();
+
+	// Remove object from filter to cause object to be set in state Created
+	Server->GetReplicationSystem()->RemoveFromGroup(NotReplicatedNetObjectGroupHandle, ServerObject->NetRefHandle);
+	Server->PreSendUpdate();
+	Server->SendUpdate(Client->ConnectionIdOnServer);
+	Server->PostSendUpdate();
+
+	// Deliver huge state
+	{
+		SIZE_T PacketCount = 0;
+		const auto& ConnectionInfo = Server->GetConnectionInfo(Client->ConnectionIdOnServer);
+		PacketCount = ConnectionInfo.WrittenPackets.Count();
+		for (SIZE_T PacketIt = 0; PacketIt != PacketCount; ++PacketIt)
+		{
+			Server->DeliverTo(Client, DeliverPacket);
+		}
+	}
+
+	// Modify a property on the object and make sure it's replicated as the object should still be created.
+	ServerObject->IntA += 1;
+
+	// Deliver latest state
+	Server->PreSendUpdate();
+	Server->SendAndDeliverTo(Client, DeliverPacket);
+	Server->PostSendUpdate();
+
+	// Verify the attachment made it through
+	const UMockNetObjectAttachmentHandler::FCallCounts AttachmentCallCounts = ClientMockNetObjectAttachmentHandler->GetFunctionCallCounts();
+	UE_NET_ASSERT_EQ(AttachmentCallCounts.OnNetBlobReceived, 1U);
+
+	UTestReplicatedIrisObject* ClientObject = Cast<UTestReplicatedIrisObject>(Client->GetReplicationBridge()->GetReplicatedObject(ServerObject->NetRefHandle));
+	UE_NET_ASSERT_NE(ClientObject, nullptr);
+	UE_NET_ASSERT_EQ(ClientObject->IntA, ServerObject->IntA);
+}
+
+// Below test will fail as we only have a special path for reliable attachments for objects that stopped replicating, not for being filtered out.
+#if 0
+UE_NET_TEST_FIXTURE(FSplitObjectTestFixture, TestReliableAttachmentAddedAfterSplittingHugeObjectIsDeliveredBeforeObjectIsFilteredOut)
+{
+	// Add a client
+	FReplicationSystemTestClient* Client = CreateClient();
+	RegisterNetBlobHandlers(Client);
+
+	// Spawn object on server
+	UTestReplicatedIrisObject* ServerObject = CreateObject(Server);
+
+	// Send and deliver packet
+	Server->PreSendUpdate();
+	Server->SendAndDeliverTo(Client, DeliverPacket);
+	Server->PostSendUpdate();
+
+	// Force huge object state
+	SetObjectPayloadByteCount(ServerObject, HugeObjectPayloadByteCount);
+
+	// Write packets
+	for (uint32 RetryIt = 0; RetryIt != HugeObjectMaxNetTickCountToArrive; ++RetryIt)
+	{
+		Server->PreSendUpdate();
+		Server->SendUpdate(Client->ConnectionIdOnServer);
+		Server->PostSendUpdate();
+	}
+
+	// Add reliable attachment
+	{
+		TRefCountPtr<FNetObjectAttachment> Attachment = ServerMockNetObjectAttachmentHandler->CreateReliableNetObjectAttachment(1U);
+
+		FNetObjectReference AttachmentTarget = FObjectReferenceCache::MakeNetObjectReference(ServerObject->NetRefHandle);
+		Server->GetReplicationSystem()->QueueNetObjectAttachment(Client->ConnectionIdOnServer, AttachmentTarget, Attachment);
+	}
+
+	// Filter out object to cause object to be set in state WaitOnFlush
+	Server->GetReplicationSystem()->AddToGroup(NotReplicatedNetObjectGroupHandle, ServerObject->NetRefHandle);
+	Server->PreSendUpdate();
+	Server->SendUpdate(Client->ConnectionIdOnServer);
+	Server->PostSendUpdate();
+
+	// Deliver huge state
+	{
+		SIZE_T PacketCount = 0;
+		const auto& ConnectionInfo = Server->GetConnectionInfo(Client->ConnectionIdOnServer);
+		PacketCount = ConnectionInfo.WrittenPackets.Count();
+		for (SIZE_T PacketIt = 0; PacketIt != PacketCount; ++PacketIt)
+		{
+			Server->DeliverTo(Client, DeliverPacket);
+		}
+	}
+
+	// Verify the attachment made it through, despite the wish to destroy the object.
+	const UMockNetObjectAttachmentHandler::FCallCounts AttachmentCallCounts = ClientMockNetObjectAttachmentHandler->GetFunctionCallCounts();
+	UE_NET_ASSERT_EQ(AttachmentCallCounts.OnNetBlobReceived, 1U);
+
+	// The object should not exist
+	UTestReplicatedIrisObject* ClientObject = Cast<UTestReplicatedIrisObject>(Client->GetReplicationBridge()->GetReplicatedObject(ServerObject->NetRefHandle));
+	UE_NET_ASSERT_EQ(ClientObject, nullptr);
+}
+#endif
+
 }
diff --git a/Engine/Source/Runtime/Experimental/Iris/Core/Private/Iris/ReplicationSystem/ReplicationWriter.cpp b/Engine/Source/Runtime/Experimental/Iris/Core/Private/Iris/ReplicationSystem/ReplicationWriter.cpp
--- a/Engine/Source/Runtime/Experimental/Iris/Core/Private/Iris/ReplicationSystem/ReplicationWriter.cpp
+++ b/Engine/Source/Runtime/Experimental/Iris/Core/Private/Iris/ReplicationSystem/ReplicationWriter.cpp
@@ -110,7 +110,7 @@ void FReplicationWriter::FReplicationInfo::SetState(EReplicatedObjectState NewSt
 		break;
 		case EReplicatedObjectState::Created:
 		{
-			check(CurrentState == EReplicatedObjectState::PendingCreate || CurrentState == EReplicatedObjectState::WaitOnCreateConfirmation || CurrentState == EReplicatedObjectState::CancelPendingDestroy);
+			check(CurrentState == EReplicatedObjectState::PendingCreate || CurrentState == EReplicatedObjectState::WaitOnCreateConfirmation || CurrentState == EReplicatedObjectState::CancelPendingDestroy || CurrentState == EReplicatedObjectState::WaitOnFlush);
 		}
 		break;
 		case EReplicatedObjectState::PendingTearOff:
@@ -628,9 +628,10 @@ void FReplicationWriter::UpdateScope(const FNetBitArrayView& UpdatedScope)
 		}
 		else if (State == EReplicatedObjectState::WaitOnFlush)
 		{			
-			if (ensureAlwaysMsgf(!!Info.TearOff, TEXT("We cannot cancel flush for an object pending tearoff.")))
+			if (ensureAlwaysMsgf(!Info.TearOff, TEXT("We cannot cancel flush for an object pending tearoff ( InternalIndex: %u )"), Index))
 			{
 				// If we are waiting on flush but are re-added to scope we reset flush flags to default.
+				ObjectsPendingDestroy.ClearBit(Index);
 				Info.FlushFlags = GetDefaultFlushFlags();
 				SetState(Index, EReplicatedObjectState::Created);
 			}
