Iris - TestMessage improvements. Use %hs  instead of character conversion when printing ascii strings. Added NetRefHandle support.

#rnx

diff --git a/Engine/Plugins/Runtime/ReplicationSystemTestPlugin/Source/Private/TestMessage.cpp b/Engine/Plugins/Runtime/ReplicationSystemTestPlugin/Source/Private/TestMessage.cpp
--- a/Engine/Plugins/Runtime/ReplicationSystemTestPlugin/Source/Private/TestMessage.cpp
+++ b/Engine/Plugins/Runtime/ReplicationSystemTestPlugin/Source/Private/TestMessage.cpp
@@ -1,6 +1,7 @@
 // Copyright Epic Games, Inc. All Rights Reserved.
 
 #include "TestMessage.h"
+#include "Iris/ReplicationSystem/NetRefHandle.h"
 #include "Misc/StringBuilder.h"
 #include "Net/Core/NetHandle/NetHandle.h"
 
@@ -14,4 +15,11 @@ FTestMessage& operator<<(FTestMessage& Message, const FNetHandle& NetHandle)
 	return Message << StringBuilder;
 }
 
+FTestMessage& operator<<(FTestMessage& Message, const FNetRefHandle& NetRefHandle)
+{
+	TStringBuilder<64> StringBuilder;
+	StringBuilder << NetRefHandle;
+	return Message << StringBuilder;
+}
+
 }
diff --git a/Engine/Plugins/Runtime/ReplicationSystemTestPlugin/Source/Private/TestMessage.h b/Engine/Plugins/Runtime/ReplicationSystemTestPlugin/Source/Private/TestMessage.h
--- a/Engine/Plugins/Runtime/ReplicationSystemTestPlugin/Source/Private/TestMessage.h
+++ b/Engine/Plugins/Runtime/ReplicationSystemTestPlugin/Source/Private/TestMessage.h
@@ -2,7 +2,7 @@
 
 #pragma once
 
-#include "CoreTypes.h"
+#include "HAL/Platform.h"
 #include "UObject/NameTypes.h"
 #include "Math/Vector.h"
 #include <cinttypes>
@@ -10,6 +10,7 @@
 namespace UE::Net
 {
 	class FNetHandle;
+	class FNetRefHandle;
 }
 
 namespace UE::Net
@@ -118,7 +119,7 @@ inline FTestMessage& operator<<(FTestMessage& Message, nullptr_t)
 
 inline FTestMessage& operator<<(FTestMessage& Message, const char* String)
 {
-	return Message << (String != nullptr ? StringCast<TCHAR>(String).Get() : TEXT("(null)"));
+	return Message << TStringBuilder<256>().Appendf(TEXT("%hs"), (String ? String : "(nullptr)"));
 }
 
 inline FTestMessage& operator<<(FTestMessage& Message, const FName& Name)
@@ -138,6 +139,8 @@ inline FTestMessage& operator<<(FTestMessage& Message, const FVector3f& Vector)
 
 FTestMessage& operator<<(FTestMessage& Message, const FNetHandle& NetHandle);
 
+FTestMessage& operator<<(FTestMessage& Message, const FNetRefHandle& NetRefHandle);
+
 #if defined(_WIN32) && !defined(_WIN64)
 inline FTestMessage& operator<<(FTestMessage& Message, SIZE_T Value)
 {
