Iris - Changemask for TArray. Allows TArrays to only replicate a subset of the entire array when few elements have changed. Feature can be disabled with cvar net.Iris.UseChangeMaskForTArray 0

diff --git a/Engine/Plugins/Runtime/ReplicationSystemTestPlugin/Source/Private/Tests/ReplicationSystem/TestObjectReferences.cpp b/Engine/Plugins/Runtime/ReplicationSystemTestPlugin/Source/Private/Tests/ReplicationSystem/TestObjectReferences.cpp
--- a/Engine/Plugins/Runtime/ReplicationSystemTestPlugin/Source/Private/Tests/ReplicationSystem/TestObjectReferences.cpp
+++ b/Engine/Plugins/Runtime/ReplicationSystemTestPlugin/Source/Private/Tests/ReplicationSystem/TestObjectReferences.cpp
@@ -61,8 +61,9 @@ void UTestObjectReferences_TestClassWithDefaultSubObject::GetLifetimeReplicatedP
 namespace UE::Net::Private
 {
 
-struct FTestObjectReferencesFixture : public FReplicationSystemServerClientTestFixture
+class FTestObjectReferencesFixture : public FReplicationSystemServerClientTestFixture
 {
+public:
 	struct FReferenceCollector
 	{
 		void Reset() { References.Reset(); }
@@ -127,6 +128,13 @@ struct FTestObjectReferencesFixture : public FReplicationSystemServerClientTestF
 		}
 	}
 
+	void SetMemberDirty(FNetBitArrayView ChangeMask, uint32 MemberIndex)
+	{
+		const FReplicationStateDescriptor* Descriptor = TestObjectProtocol->ReplicationStateDescriptors[1];
+		const FReplicationStateMemberChangeMaskDescriptor& ChangeMaskDesciptor = Descriptor->MemberChangeMaskDescriptors[MemberIndex];
+		ChangeMask.SetBits(ChangeMaskDesciptor.BitOffset, ChangeMaskDesciptor.BitCount);
+	}
+
 	FReferenceCollector CollectAllCollector;
 	FReferenceCollector CollectOnlyValidCollector;
 	
@@ -684,9 +692,8 @@ UE_NET_TEST_FIXTURE(FTestObjectReferencesFixture, TestVisitAllDirty)
 	TestObject->TestStructWithNestedRefCArray_CArray[1].StructWithRef_CArray[1].Ref_CArray[1] = TestReferences[CurrentReferenceIndex++]; // 13
 	TestObject->TestStructWithNestedRefCArray_CArray[2].StructWithRef_CArray[2].Ref_CArray[2] = TestReferences[CurrentReferenceIndex++]; // 14
 
-	uint32 ChangeMaskData[8];
-
-	FNetBitArrayView ChangeMask(ChangeMaskData, TestObjectProtocol->ChangeMaskBitCount, FNetBitArrayView::ResetOnInit);
+	FNetBitArray ChangeMaskStorage(TestObjectProtocol->ChangeMaskBitCount);
+	FNetBitArrayView ChangeMask = MakeNetBitArrayView(ChangeMaskStorage);
 
 	// To trigger copy of data
 	Server->PreSendUpdate();
@@ -706,20 +713,23 @@ UE_NET_TEST_FIXTURE(FTestObjectReferencesFixture, TestVisitAllDirty)
 
 		VisitReferences(TestObjectStateBuffer, TestObjectProtocol, Collector);
 		UE_NET_ASSERT_EQ((uint32)Collector.References.Num(), 1U);
-		UE_NET_ASSERT_TRUE(Collector.References[0].GetRefHandle() == TestReferences[0]->NetRefHandle);
+		UE_NET_ASSERT_EQ(Collector.References[0].GetRefHandle(), TestReferences[0]->NetRefHandle);
 	}
 
 	// Partial
 	{
 		FCollectValidReferenceCollector Collector(false, ChangeMask);
 		ChangeMask.Reset();
-		ChangeMask.SetBit(11);
+
+		// Mark member 11 as dirty. This corresponds to TestStructWithNestedRefTArray_TArray.
+		SetMemberDirty(ChangeMask, 11);
+
 		VisitReferences(TestObjectStateBuffer, TestObjectProtocol, Collector);
 
 		UE_NET_ASSERT_EQ(Collector.References.Num(), 3);
 		for (int32 It = 0; It < Collector.References.Num(); ++It)
 		{
-			UE_NET_ASSERT_TRUE(Collector.References[It].GetRefHandle() == TestReferences[It + 20]->NetRefHandle);
+			UE_NET_ASSERT_EQ(Collector.References[It].GetRefHandle(), TestReferences[It + 20]->NetRefHandle);
 		}
 	}
 
@@ -732,7 +742,7 @@ UE_NET_TEST_FIXTURE(FTestObjectReferencesFixture, TestVisitAllDirty)
 		UE_NET_ASSERT_EQ((uint32)Collector.References.Num(), CurrentReferenceIndex);
 		for (int32 It = 0; It < Collector.References.Num(); ++It)
 		{
-			UE_NET_ASSERT_TRUE(Collector.References[It].GetRefHandle() == TestReferences[It]->NetRefHandle);
+			UE_NET_ASSERT_EQ(Collector.References[It].GetRefHandle(), TestReferences[It]->NetRefHandle);
 		}
 	}
 
@@ -745,7 +755,7 @@ UE_NET_TEST_FIXTURE(FTestObjectReferencesFixture, TestVisitAllDirty)
 		UE_NET_ASSERT_EQ((uint32)Collector.References.Num(), CurrentReferenceIndex - 1);
 		for (int32 It = 0; It < Collector.References.Num(); ++It)
 		{
-			UE_NET_ASSERT_TRUE(Collector.References[It].GetRefHandle() == TestReferences[It + 1]->NetRefHandle);
+			UE_NET_ASSERT_EQ(Collector.References[It].GetRefHandle(), TestReferences[It + 1]->NetRefHandle);
 		}
 	}
 }
diff --git a/Engine/Plugins/Runtime/ReplicationSystemTestPlugin/Source/Private/Tests/ReplicationSystem/TestReplicationOperations.cpp b/Engine/Plugins/Runtime/ReplicationSystemTestPlugin/Source/Private/Tests/ReplicationSystem/TestReplicationOperations.cpp
--- a/Engine/Plugins/Runtime/ReplicationSystemTestPlugin/Source/Private/Tests/ReplicationSystem/TestReplicationOperations.cpp
+++ b/Engine/Plugins/Runtime/ReplicationSystemTestPlugin/Source/Private/Tests/ReplicationSystem/TestReplicationOperations.cpp
@@ -44,7 +44,7 @@ struct FTestReplicationOperationsFixture : public FReplicationSystemTestFixture
 
 	FMemStackBase& GetTempAllocator() { return *TempLinearAllocator; }
 
-	uint8 ChangeMaskBuffer[16];
+	alignas(16) uint8 ChangeMaskBuffer[32];
 	FNetBitStreamWriter ChangeMaskWriter;
 	FMemStackBase* TempLinearAllocator;
 
@@ -54,7 +54,7 @@ struct FTestReplicationOperationsFixture : public FReplicationSystemTestFixture
 	{ 
 		check((BitCount >> 3) < sizeof(ChangeMaskBuffer));
 		FMemory::Memset(ChangeMaskBuffer, 0);
-		ChangeMaskWriter.InitBytes(ChangeMaskBuffer, 16);
+		ChangeMaskWriter.InitBytes(ChangeMaskBuffer, sizeof(ChangeMaskBuffer));
 
 		return &ChangeMaskWriter;
 	}
diff --git a/Engine/Plugins/Runtime/ReplicationSystemTestPlugin/Source/Private/Tests/Serialization/TestArrayPropertyNetSerializer.cpp b/Engine/Plugins/Runtime/ReplicationSystemTestPlugin/Source/Private/Tests/Serialization/TestArrayPropertyNetSerializer.cpp
--- a/Engine/Plugins/Runtime/ReplicationSystemTestPlugin/Source/Private/Tests/Serialization/TestArrayPropertyNetSerializer.cpp
+++ b/Engine/Plugins/Runtime/ReplicationSystemTestPlugin/Source/Private/Tests/Serialization/TestArrayPropertyNetSerializer.cpp
@@ -4,11 +4,15 @@
 #include "TestNetSerializerFixture.h"
 #include "Iris/ReplicationState/InternalReplicationStateDescriptorUtils.h"
 #include "Iris/ReplicationState/PropertyNetSerializerInfoRegistry.h"
+#include "Iris/ReplicationState/PropertyReplicationState.h"
 #include "Iris/ReplicationState/ReplicationStateDescriptorBuilder.h"
+#include "Iris/ReplicationState/ReplicationStateUtil.h"
 #include "Iris/Serialization/InternalNetSerializers.h"
 #include "Iris/Serialization/NetSerializers.h"
 #include "Iris/Serialization/InternalNetSerializationContext.h"
-#include "MockNetSerializer.h"
+#include "Net/Core/NetBitArray.h"
+#include "Net/UnrealNetwork.h"
+#include "Tests/Serialization/MockNetSerializer.h"
 
 namespace UE::Net::Private
 {
@@ -122,6 +126,13 @@ protected:
 	FStructWithDynamicArrayOfComplexTypeForArrayPropertyNetSerializerTest NonEmptyArrayInstance1;
 };
 
+class FTestElementChangeMaskForArrayPropertyNetSerializer : public FTestArrayPropertyNetSerializerBase
+{
+protected:
+	virtual void SetUp() override;
+	virtual void TearDown() override;
+};
+
 // Basic tests
 UE_NET_TEST_FIXTURE(FTestArrayPropertyNetSerializerBase, TestHasIsForwardingSerializerTrait)
 {
@@ -210,6 +221,231 @@ UE_NET_TEST_FIXTURE(FTestComplexArrayPropertyNetSerializer, TestFreeDynamicState
 	TestFreeDynamicState();
 }
 
+// Tests for array element changemask
+UE_NET_TEST_FIXTURE(FTestElementChangeMaskForArrayPropertyNetSerializer, TestChangeMaskIsModifiedWhenAddingElements)
+{
+	FPropertyReplicationState ReplicationState(ReplicationStateDescriptor);
+	TObjectPtr<UClassWithDynamicArrayOfPrimitiveTypeForArrayPropertyNetSerializerTest> Source = NewObject<UClassWithDynamicArrayOfPrimitiveTypeForArrayPropertyNetSerializerTest>();
+
+	FNetBitArrayView ArrayChangeMask = UE::Net::Private::GetMemberChangeMask(ReplicationState.GetStateBuffer(), ReplicationStateDescriptor);
+	ArrayChangeMask.Reset();
+
+	// Add one element at a time until the last bit in the changemask is set and make sure the array is always marked as dirty as well as the corresponding element
+	for (unsigned ElementIt = 0, ElementEndIt = FPropertyReplicationState::TArrayElementChangeMaskBits; ElementIt != ElementEndIt; ++ElementIt)
+	{
+		Source->ArrayOfUint.Add(0);
+		ReplicationState.PollPropertyReplicationState(Source.Get());
+		
+		UE_NET_ASSERT_TRUE(ArrayChangeMask.GetBit(FPropertyReplicationState::TArrayPropertyChangeMaskBitIndex));
+		UE_NET_ASSERT_TRUE(ArrayChangeMask.GetBit(FPropertyReplicationState::TArrayElementChangeMaskBitOffset + ElementIt));
+		UE_NET_ASSERT_EQ(ArrayChangeMask.CountSetBits(), 2U);
+
+		ArrayChangeMask.Reset();
+	}
+}
+
+UE_NET_TEST_FIXTURE(FTestElementChangeMaskForArrayPropertyNetSerializer, TestChangeMaskIsModifiedWhenAddingElementsBeyondChangeMaskCapacity)
+{
+	FPropertyReplicationState ReplicationState(ReplicationStateDescriptor);
+	TObjectPtr<UClassWithDynamicArrayOfPrimitiveTypeForArrayPropertyNetSerializerTest> Source = NewObject<UClassWithDynamicArrayOfPrimitiveTypeForArrayPropertyNetSerializerTest>();
+
+	FNetBitArrayView ArrayChangeMask = UE::Net::Private::GetMemberChangeMask(ReplicationState.GetStateBuffer(), ReplicationStateDescriptor);
+	ArrayChangeMask.Reset();
+
+	// First add elements up to the max element changemask capacity.
+	{
+		Source->ArrayOfUint.SetNum(FPropertyReplicationState::TArrayElementChangeMaskBits);
+		ReplicationState.PollPropertyReplicationState(Source.Get());
+
+		UE_NET_ASSERT_EQ(ArrayChangeMask.CountSetBits(), ArrayChangeMask.GetNumBits());
+		ArrayChangeMask.Reset();
+	}
+
+	// Double the array size by adding one element at a time until the last bit in the changemask is set and make sure the array is always marked as dirty as well as the corresponding element.
+	for (unsigned ElementIt = 0, ElementEndIt = FPropertyReplicationState::TArrayElementChangeMaskBits; ElementIt != ElementEndIt; ++ElementIt)
+	{
+		Source->ArrayOfUint.Add(0);
+		ReplicationState.PollPropertyReplicationState(Source.Get());
+		
+		UE_NET_ASSERT_TRUE(ArrayChangeMask.GetBit(FPropertyReplicationState::TArrayPropertyChangeMaskBitIndex));
+		UE_NET_ASSERT_TRUE(ArrayChangeMask.GetBit(FPropertyReplicationState::TArrayElementChangeMaskBitOffset + ElementIt));
+		UE_NET_ASSERT_EQ(ArrayChangeMask.CountSetBits(), 2U);
+
+		ArrayChangeMask.Reset();
+	}
+}
+
+UE_NET_TEST_FIXTURE(FTestElementChangeMaskForArrayPropertyNetSerializer, TestChangeMaskIsModifiedWhenModifyingElements)
+{
+	FPropertyReplicationState ReplicationState(ReplicationStateDescriptor);
+	TObjectPtr<UClassWithDynamicArrayOfPrimitiveTypeForArrayPropertyNetSerializerTest> Source = NewObject<UClassWithDynamicArrayOfPrimitiveTypeForArrayPropertyNetSerializerTest>();
+
+	FNetBitArrayView ArrayChangeMask = UE::Net::Private::GetMemberChangeMask(ReplicationState.GetStateBuffer(), ReplicationStateDescriptor);
+
+	// First add elements up to the max element changemask capacity.
+	{
+		Source->ArrayOfUint.SetNum(FPropertyReplicationState::TArrayElementChangeMaskBits);
+		ReplicationState.PollPropertyReplicationState(Source.Get());
+
+		ArrayChangeMask.Reset();
+	}
+
+	// Modify one element at a time until the last bit in the changemask is set and make sure the array is always marked as dirty as well as the corresponding element
+	for (unsigned ElementIt = 0, ElementEndIt = FPropertyReplicationState::TArrayElementChangeMaskBits; ElementIt != ElementEndIt; ++ElementIt)
+	{
+		Source->ArrayOfUint[ElementIt] ^= 7U;
+		ReplicationState.PollPropertyReplicationState(Source.Get());
+		
+		UE_NET_ASSERT_TRUE(ArrayChangeMask.GetBit(FPropertyReplicationState::TArrayPropertyChangeMaskBitIndex));
+		UE_NET_ASSERT_TRUE(ArrayChangeMask.GetBit(FPropertyReplicationState::TArrayElementChangeMaskBitOffset + ElementIt));
+		UE_NET_ASSERT_EQ(ArrayChangeMask.CountSetBits(), 2U);
+
+		ArrayChangeMask.Reset();
+	}
+}
+
+UE_NET_TEST_FIXTURE(FTestElementChangeMaskForArrayPropertyNetSerializer, TestChangeMaskIsModifiedWhenModifyingElementsBeyondChangeMaskCapacity)
+{
+	FPropertyReplicationState ReplicationState(ReplicationStateDescriptor);
+	TObjectPtr<UClassWithDynamicArrayOfPrimitiveTypeForArrayPropertyNetSerializerTest> Source = NewObject<UClassWithDynamicArrayOfPrimitiveTypeForArrayPropertyNetSerializerTest>();
+
+	FNetBitArrayView ArrayChangeMask = UE::Net::Private::GetMemberChangeMask(ReplicationState.GetStateBuffer(), ReplicationStateDescriptor);
+	ArrayChangeMask.Reset();
+
+	// First add elements up to double the max element changemask capacity.
+	{
+		Source->ArrayOfUint.SetNum(2*FPropertyReplicationState::TArrayElementChangeMaskBits);
+		ReplicationState.PollPropertyReplicationState(Source.Get());
+
+		ArrayChangeMask.Reset();
+	}
+
+	// Double the array size by adding one element at a time until the last bit in the changemask is set and make sure the array is always marked as dirty as well as the corresponding element.
+	for (unsigned ElementIt = 0, ElementEndIt = FPropertyReplicationState::TArrayElementChangeMaskBits; ElementIt != ElementEndIt; ++ElementIt)
+	{
+		Source->ArrayOfUint[FPropertyReplicationState::TArrayElementChangeMaskBits + ElementIt] ^= 7U;
+		ReplicationState.PollPropertyReplicationState(Source.Get());
+		
+		UE_NET_ASSERT_TRUE(ArrayChangeMask.GetBit(FPropertyReplicationState::TArrayPropertyChangeMaskBitIndex));
+		UE_NET_ASSERT_TRUE(ArrayChangeMask.GetBit(FPropertyReplicationState::TArrayElementChangeMaskBitOffset + ElementIt));
+		UE_NET_ASSERT_EQ(ArrayChangeMask.CountSetBits(), 2U);
+
+		ArrayChangeMask.Reset();
+	}
+}
+
+UE_NET_TEST_FIXTURE(FTestElementChangeMaskForArrayPropertyNetSerializer, TestChangeMaskIsModifiedWhenModifyingMultipleArbitraryElements)
+{
+	FPropertyReplicationState ReplicationState(ReplicationStateDescriptor);
+	TObjectPtr<UClassWithDynamicArrayOfPrimitiveTypeForArrayPropertyNetSerializerTest> Source = NewObject<UClassWithDynamicArrayOfPrimitiveTypeForArrayPropertyNetSerializerTest>();
+
+	FNetBitArrayView ArrayChangeMask = UE::Net::Private::GetMemberChangeMask(ReplicationState.GetStateBuffer(), ReplicationStateDescriptor);
+	ArrayChangeMask.Reset();
+
+	constexpr int32 ElementCount = 17;
+	// First add some elements
+	{
+		Source->ArrayOfUint.SetNum(ElementCount);
+		ReplicationState.PollPropertyReplicationState(Source.Get());
+
+		ArrayChangeMask.Reset();
+	}
+
+	// Modify some arbitrary elements
+	const int32 ElementIndicesToModify[] = {3, 7, 11};
+	for (int32 ElementIndex : ElementIndicesToModify)
+	{
+		Source->ArrayOfUint[ElementIndex] ^= 471147114711U;
+	}
+
+	ReplicationState.PollPropertyReplicationState(Source.Get());
+
+	// Verify the elements were marked as dirty
+	UE_NET_ASSERT_TRUE(ArrayChangeMask.GetBit(FPropertyReplicationState::TArrayPropertyChangeMaskBitIndex));
+	UE_NET_ASSERT_EQ(ArrayChangeMask.CountSetBits(), uint32(UE_ARRAY_COUNT(ElementIndicesToModify) + 1U));
+
+	for (int32 ElementIndex : ElementIndicesToModify)
+	{
+		UE_NET_ASSERT_TRUE(ArrayChangeMask.GetBit(FPropertyReplicationState::TArrayElementChangeMaskBitOffset + ElementIndex));
+	}
+}
+
+UE_NET_TEST_FIXTURE(FTestElementChangeMaskForArrayPropertyNetSerializer, TestSerializingFewElementsUsesLessBandwidth)
+{
+	FPropertyReplicationState ReplicationState(ReplicationStateDescriptor);
+	TObjectPtr<UClassWithDynamicArrayOfPrimitiveTypeForArrayPropertyNetSerializerTest> Source = NewObject<UClassWithDynamicArrayOfPrimitiveTypeForArrayPropertyNetSerializerTest>();
+
+	FNetBitArrayView ArrayChangeMask = UE::Net::Private::GetMemberChangeMask(ReplicationState.GetStateBuffer(), ReplicationStateDescriptor);
+	ArrayChangeMask.Reset();
+
+	constexpr int32 ElementCount = 17;
+	// First add some elements
+	{
+		Source->ArrayOfUint.SetNum(ElementCount);
+
+		// Modify all elements
+		for (uint32& Element : Source->ArrayOfUint)
+		{
+			Element ^= 101U;
+		}
+
+		ReplicationState.PollPropertyReplicationState(Source.Get());
+	}
+
+	NetSerializationContext.SetChangeMask(&ArrayChangeMask);
+
+	FNetSerializeArgs SerializeArgs = {};
+	{
+		SerializeArgs.Version = ArrayPropertyNetSerializer->Version;
+		SerializeArgs.NetSerializerConfig = NetSerializerConfigParam(ArrayPropertyNetSerializerConfig);
+		SerializeArgs.Source = NetSerializerValuePointer(StateBuffer0);
+		SerializeArgs.ChangeMaskInfo = FNetSerializerChangeMaskParam{0, static_cast<uint16>(ArrayChangeMask.GetNumBits())};
+	}
+
+	uint32 BitsWrittenForFullArray = 0;
+	{
+		Quantize(Source->ArrayOfUint, StateBuffer0);
+
+		Writer.InitBytes(BitStreamBuffer0, sizeof(BitStreamBuffer0));
+		ArrayPropertyNetSerializer->Serialize(NetSerializationContext, SerializeArgs);
+		Writer.CommitWrites();
+
+		BitsWrittenForFullArray = Writer.GetPosBits();
+		UE_NET_ASSERT_GT(BitsWrittenForFullArray, 0U);
+
+		FreeDynamicState(StateBuffer0);
+	}
+
+	// Modify a few elements
+	{
+		ArrayChangeMask.Reset();
+
+		const int32 ElementIndicesToModify[] = {3, 7, 11};
+		for (int32 ElementIndex : ElementIndicesToModify)
+		{
+			Source->ArrayOfUint[ElementIndex] += 1U;
+		}
+
+		ReplicationState.PollPropertyReplicationState(Source.Get());
+	}
+
+	uint32 BitsWrittenForPartialArray = 0;
+	{
+		Quantize(Source->ArrayOfUint, StateBuffer0);
+
+		Writer.InitBytes(BitStreamBuffer0, sizeof(BitStreamBuffer0));
+		ArrayPropertyNetSerializer->Serialize(NetSerializationContext, SerializeArgs);
+		Writer.CommitWrites();
+
+		BitsWrittenForPartialArray = Writer.GetPosBits();
+		UE_NET_ASSERT_GT(BitsWrittenForPartialArray, 0U);
+
+		FreeDynamicState(StateBuffer0);
+	}
+
+	UE_NET_ASSERT_LT(BitsWrittenForPartialArray, BitsWrittenForFullArray);
+}
+
 // FTestArrayPropertyNetSerializerBase implementation
 void FTestArrayPropertyNetSerializerBase::SetUp()
 {
@@ -782,7 +1018,7 @@ void FTestSimpleArrayPropertyNetSerializer::TestIsEqual()
 	// Check non-empty non-identical arrays aren't considered equal when using proper element serializer
 	for (const bool bIsUsingMockSerializer : {false, true})
 	{
-		bIsUsingMockSerializer? SetupForMockSerializer() : SetupForOriginalSerializer();
+		bIsUsingMockSerializer ? SetupForMockSerializer() : SetupForOriginalSerializer();
 
 		FStructWithDynamicArrayOfPrimitiveTypeForArrayPropertyNetSerializerTest NonEmptyArrayInstance = NonEmptyArrayInstance0;
 		NonEmptyArrayInstance.ArrayOfUint[1] = NonEmptyArrayInstance0.ArrayOfUint[1] ^ 1U;
@@ -1049,4 +1285,31 @@ void FTestComplexArrayPropertyNetSerializer::TestFreeDynamicState()
 	UE_NET_ASSERT_EQ(MockNetSerializerCallCounter.FreeDynamicState, uint32(NonEmptyArrayInstance0.ArrayOfStructWithArray.Num()));
 }
 
+// FTestElementChangeMaskForArrayPropertyNetSerializer
+void FTestElementChangeMaskForArrayPropertyNetSerializer::SetUp()
+{
+	FTestArrayPropertyNetSerializerBase::SetUp();
+
+	FReplicationStateDescriptorBuilder::FResult Descriptors;
+	FReplicationStateDescriptorBuilder::CreateDescriptorsForClass(Descriptors, UClassWithDynamicArrayOfPrimitiveTypeForArrayPropertyNetSerializerTest::StaticClass());
+	UE_NET_ASSERT_GT(Descriptors.Num(), 0);
+	ReplicationStateDescriptor = Descriptors[0];
+	ArrayPropertyNetSerializerConfig = const_cast<FArrayPropertyNetSerializerConfig*>(static_cast<const FArrayPropertyNetSerializerConfig*>(ReplicationStateDescriptor->MemberSerializerDescriptors[0].SerializerConfig));
+
+	NetSerializationContext = FNetSerializationContext(&Reader, &Writer);
+	NetSerializationContext.SetInternalContext(&InternalContext);
+}
+
+void FTestElementChangeMaskForArrayPropertyNetSerializer::TearDown()
+{
+	ReplicationStateDescriptor.SafeRelease();
+
+	FTestArrayPropertyNetSerializerBase::TearDown();
+}
+
+}
+
+void UClassWithDynamicArrayOfPrimitiveTypeForArrayPropertyNetSerializerTest::GetLifetimeReplicatedProps(TArray<FLifetimeProperty>& OutLifetimeProps) const
+{
+	DOREPLIFETIME(ThisClass, ArrayOfUint);
 }
diff --git a/Engine/Plugins/Runtime/ReplicationSystemTestPlugin/Source/Private/Tests/Serialization/TestArrayPropertyNetSerializer.h b/Engine/Plugins/Runtime/ReplicationSystemTestPlugin/Source/Private/Tests/Serialization/TestArrayPropertyNetSerializer.h
--- a/Engine/Plugins/Runtime/ReplicationSystemTestPlugin/Source/Private/Tests/Serialization/TestArrayPropertyNetSerializer.h
+++ b/Engine/Plugins/Runtime/ReplicationSystemTestPlugin/Source/Private/Tests/Serialization/TestArrayPropertyNetSerializer.h
@@ -21,3 +21,16 @@ struct FStructWithDynamicArrayOfComplexTypeForArrayPropertyNetSerializerTest
 	UPROPERTY()
 	TArray<FStructWithDynamicArrayOfPrimitiveTypeForArrayPropertyNetSerializerTest> ArrayOfStructWithArray;
 };
+
+UCLASS()
+class UClassWithDynamicArrayOfPrimitiveTypeForArrayPropertyNetSerializerTest : public UObject
+{
+	GENERATED_BODY()
+
+public:
+	UPROPERTY(Replicated)
+	TArray<uint32> ArrayOfUint;
+
+protected:
+	virtual void GetLifetimeReplicatedProps(TArray<FLifetimeProperty>& OutLifetimeProps) const override;
+};
diff --git a/Engine/Source/Runtime/Experimental/Iris/Core/Private/Iris/ReplicationState/InternalPropertyReplicationState.cpp b/Engine/Source/Runtime/Experimental/Iris/Core/Private/Iris/ReplicationState/InternalPropertyReplicationState.cpp
--- a/Engine/Source/Runtime/Experimental/Iris/Core/Private/Iris/ReplicationState/InternalPropertyReplicationState.cpp
+++ b/Engine/Source/Runtime/Experimental/Iris/Core/Private/Iris/ReplicationState/InternalPropertyReplicationState.cpp
@@ -1,11 +1,11 @@
 // Copyright Epic Games, Inc. All Rights Reserved.
 
 #include "Iris/ReplicationState/InternalPropertyReplicationState.h"
-#include "CoreTypes.h"
-#include "Net/Core/NetBitArray.h"
+#include "Iris/ReplicationState/PropertyReplicationState.h"
 #include "Iris/ReplicationState/ReplicationStateUtil.h"
 #include "Iris/ReplicationState/InternalReplicationStateDescriptorUtils.h"
 #include "Iris/Serialization/NetSerializers.h"
+#include "Net/Core/NetBitArray.h"
 #include "UObject/UnrealType.h"
 
 namespace UE::Net::Private
@@ -295,4 +295,58 @@ void InternalCopyPropertyValue(const FReplicationStateDescriptor* Descriptor, ui
 	Property->CopySingleValue(Dst, Src);
 }
 
+bool InternalCompareAndCopyArrayWithElementChangeMask(const FReplicationStateDescriptor* Descriptor, uint32 MemberIndex, const void* RESTRICT DstArray, const void* RESTRICT SrcArray, UE::Net::FNetBitArrayView& ChangeMask)
+{
+	bool bArrayIsEqual = true;
+
+	const FReplicationStateMemberChangeMaskDescriptor& ChangeMaskInfo = Descriptor->MemberChangeMaskDescriptors[MemberIndex];
+
+	const FReplicationStateMemberSerializerDescriptor& ArrayDescriptor = Descriptor->MemberSerializerDescriptors[MemberIndex];
+	const FArrayPropertyNetSerializerConfig* ArrayConfig = static_cast<const FArrayPropertyNetSerializerConfig*>(ArrayDescriptor.SerializerConfig);
+
+	FScriptArrayHelper SrcScriptArray(ArrayConfig->Property.Get(), SrcArray);
+	FScriptArrayHelper DstScriptArray(ArrayConfig->Property.Get(), DstArray);
+
+	// Detect size change and adjust the destination array size as needed and clear bits that don't relate to any elements.
+	const uint32 SrcElementCount = SrcScriptArray.Num();
+	const uint32 DstElementCount = DstScriptArray.Num();
+	if (SrcElementCount != DstElementCount)
+	{
+		bArrayIsEqual = false;
+
+		DstScriptArray.Resize(SrcElementCount);
+
+		// If array has shrunk then mask off bits in the changemask pertaining to elements that no longer exist. For a growing array we skip compare and set it to dirty in the element loop.
+		if (SrcElementCount < DstElementCount)
+		{
+			if (SrcElementCount + 1U < ChangeMaskInfo.BitCount)
+			{
+				const uint32 BitOffsetToClear = ChangeMaskInfo.BitOffset + FPropertyReplicationState::TArrayElementChangeMaskBitOffset + SrcElementCount;
+				const uint32 BitCountToClear = ChangeMaskInfo.BitCount - 1U - SrcElementCount;
+				ChangeMask.ClearBits(BitOffsetToClear, BitCountToClear);
+			}
+		}
+	}
+
+	const FReplicationStateDescriptor* ElementDescriptor = ArrayConfig->StateDescriptor;
+	const uint32 ElementChangeMaskBitOffset = ChangeMaskInfo.BitOffset + FPropertyReplicationState::TArrayElementChangeMaskBitOffset;
+	for (uint32 ElementIt = 0, ElementEndIt = SrcElementCount; ElementIt < ElementEndIt; ++ElementIt)
+	{
+		const uint8* SrcElement = SrcScriptArray.GetRawPtr(ElementIt);
+		uint8* DstElement = DstScriptArray.GetRawPtr(ElementIt);
+
+		// Compare elements up to the previous element count, i.e. at most DstElementCount. New elements will be considered different and always copied.
+		const bool bIsNewElement = ElementIt >= DstElementCount;
+		if (bIsNewElement || !InternalCompareStructProperty(ElementDescriptor, DstElement, SrcElement))
+		{
+			bArrayIsEqual = false;
+			const uint32 ElementBitOffset = ElementChangeMaskBitOffset + (ElementIt % FPropertyReplicationState::TArrayElementChangeMaskBits);
+			ChangeMask.SetBit(ElementBitOffset);
+			InternalCopyStructProperty(ElementDescriptor, DstElement, SrcElement);
+		}
+	}
+
+	return bArrayIsEqual;
+}
+
 }
diff --git a/Engine/Source/Runtime/Experimental/Iris/Core/Private/Iris/ReplicationState/InternalPropertyReplicationState.h b/Engine/Source/Runtime/Experimental/Iris/Core/Private/Iris/ReplicationState/InternalPropertyReplicationState.h
--- a/Engine/Source/Runtime/Experimental/Iris/Core/Private/Iris/ReplicationState/InternalPropertyReplicationState.h
+++ b/Engine/Source/Runtime/Experimental/Iris/Core/Private/Iris/ReplicationState/InternalPropertyReplicationState.h
@@ -5,6 +5,11 @@
 #include "HAL/Platform.h"
 #include "Iris/ReplicationState/ReplicationStateDescriptor.h"
 
+namespace UE::Net
+{
+	class FNetBitArrayView;
+}
+
 namespace UE::Net::Private
 {
 
@@ -17,7 +22,6 @@ IRISCORE_API void DestructPropertyReplicationState(uint8* StateBuffer, const FRe
 /** Copy external state in SrcStateBuffer to external state DstStateBuffer including changemask */
 IRISCORE_API void CopyPropertyReplicationState(uint8* RESTRICT DstStateBuffer, uint8* RESTRICT SrcStateBuffer, const FReplicationStateDescriptor* Descriptor);
 
-
 /** Copy property value, if the property is fully replicated we use the properties copy function, otherwise we do a per member copy of all replicated data */
 IRISCORE_API void InternalCopyPropertyValue(const FReplicationStateDescriptor* Descriptor, uint32 MemberIndex, void* RESTRICT Dst, const void* RESTRICT Src);
 
@@ -30,4 +34,11 @@ IRISCORE_API bool InternalCompareStructProperty(const FReplicationStateDescripto
 /** CompareMembers, if data is not fully replicated we will use per member compare for structs and arrays */
 IRISCORE_API bool InternalCompareMember(const FReplicationStateDescriptor* Descriptor, uint32 MemberIndex, const void* RESTRICT ValueA, const void* RESTRICT ValueB);
 
+/**
+ * Compare array elements and mark dirty indices in the changemask. The resulting DstArray will be identical to the SrcArray after calling this function.
+ * @param ChangeMask needs to be big enough for the entire descriptor. The member's change mask descriptor will be modified for dirty array indices. The bit which indicates the array is dirty is left unmodified.
+ * @return false if the DstArray was modified in anyway, including size changes, and true if the arrays were already identical.
+ */
+IRISCORE_API bool InternalCompareAndCopyArrayWithElementChangeMask(const FReplicationStateDescriptor* Descriptor, uint32 MemberIndex, const void* RESTRICT DstArray, const void* RESTRICT SrcArray, UE::Net::FNetBitArrayView& ChangeMask);
+
 }
diff --git a/Engine/Source/Runtime/Experimental/Iris/Core/Private/Iris/ReplicationState/PropertyReplicationState.cpp b/Engine/Source/Runtime/Experimental/Iris/Core/Private/Iris/ReplicationState/PropertyReplicationState.cpp
--- a/Engine/Source/Runtime/Experimental/Iris/Core/Private/Iris/ReplicationState/PropertyReplicationState.cpp
+++ b/Engine/Source/Runtime/Experimental/Iris/Core/Private/Iris/ReplicationState/PropertyReplicationState.cpp
@@ -2,6 +2,7 @@
 
 #include "Iris/ReplicationState/PropertyReplicationState.h"
 #include "Iris/ReplicationState/InternalPropertyReplicationState.h"
+#include "Iris/ReplicationState/InternalReplicationStateDescriptorUtils.h"
 #include "Iris/ReplicationState/ReplicationStateUtil.h"
 #include "Iris/ReplicationSystem/ReplicationSystem.h"
 #include "Net/Core/NetBitArray.h"
@@ -128,32 +129,17 @@ void FPropertyReplicationState::InjectState(const FReplicationStateDescriptor* D
 
 void FPropertyReplicationState::Set(const FPropertyReplicationState& Other)
 {
-	if (IsValid())
+	if (IsValid() && this != &Other)
 	{
 		const FReplicationStateDescriptor* Descriptor = ReplicationStateDescriptor;
-		uint8* DstStateBuffer = StateBuffer;
-		uint8* SrcStateBuffer = Other.StateBuffer;
-
 		const FReplicationStateMemberDescriptor* MemberDescriptors = Descriptor->MemberDescriptors;
-		const FProperty*const* Properties = Descriptor->MemberProperties;
-		const uint32 MemberCount = Descriptor->MemberCount;
 
-		for (uint32 MemberIt = 0; MemberIt < MemberCount; ++MemberIt)
+		const uint8* SrcStateBuffer = Other.StateBuffer;
+		for (uint32 MemberIt = 0, MemberEndIt = Descriptor->MemberCount; MemberIt < MemberEndIt; ++MemberIt)
 		{
 			const FReplicationStateMemberDescriptor& MemberDescriptor = MemberDescriptors[MemberIt];
-			const FProperty* Property = Properties[MemberIt];
-
-			uint8* DstValue = DstStateBuffer + MemberDescriptor.ExternalMemberOffset;
-			uint8* SrcValue = SrcStateBuffer + MemberDescriptor.ExternalMemberOffset;
-	
-			// $IRIS TODO: Could use serializer IsEqual.
-			if (!IsDirty(MemberIt) && IsCustomConditionEnabled(MemberIt) && !Property->Identical(SrcValue, DstValue))
-			{
-				MarkDirty(MemberIt);
-			}
-
-			// We need to use the internal copy as we do not want to run custom assign/copy constructors as we want a true copy of all replicated data
-			Private::InternalCopyPropertyValue(Descriptor, MemberIt, DstValue, SrcValue);
+			const uint8* SrcValue = SrcStateBuffer + MemberDescriptor.ExternalMemberOffset;
+			SetPropertyValue(MemberIt, SrcValue);
 		}
 	}
 }
@@ -163,12 +149,31 @@ void FPropertyReplicationState::SetPropertyValue(uint32 Index, const void* SrcVa
 	const FReplicationStateDescriptor* Descriptor = ReplicationStateDescriptor;
 	void* DstValue = StateBuffer + Descriptor->MemberDescriptors[Index].ExternalMemberOffset;
 
+	// Special handling for top level arrays with changemask bits for elements
+	const FReplicationStateMemberSerializerDescriptor& MemberSerializerDescriptor = Descriptor->MemberSerializerDescriptors[Index];
+	if (IsUsingArrayPropertyNetSerializer(MemberSerializerDescriptor))
+	{
+		const FReplicationStateMemberChangeMaskDescriptor& ChangeMaskInfo = Descriptor->MemberChangeMaskDescriptors[Index];
+		if (!IsInitState() && IsCustomConditionEnabled(Index) && ChangeMaskInfo.BitCount > 1U)
+		{
+			FNetBitArrayView MemberChangeMask = Private::GetMemberChangeMask(StateBuffer, Descriptor);
+
+			const bool bArraysAreEqual = Private::InternalCompareAndCopyArrayWithElementChangeMask(Descriptor, Index, DstValue, SrcValue, MemberChangeMask);
+			if (!bArraysAreEqual && !MemberChangeMask.GetBit(ChangeMaskInfo.BitOffset + TArrayPropertyChangeMaskBitIndex))
+			{
+				MarkArrayDirty(Index);
+			}
+
+			return;
+		}
+	}
+
 	if (!IsDirty(Index) && IsCustomConditionEnabled(Index) && !Private::InternalCompareMember(Descriptor, Index, SrcValue, DstValue))
 	{
 		MarkDirty(Index);
 	}
 
-	Private::InternalCopyPropertyValue(Descriptor, Index, DstValue, SrcValue);		
+	Private::InternalCopyPropertyValue(Descriptor, Index, DstValue, SrcValue);
 }
 
 void FPropertyReplicationState::GetPropertyValue(uint32 Index, void* DstValue) const
@@ -218,6 +223,32 @@ bool FPropertyReplicationState::IsDirty(uint32 Index) const
 	}
 }
 
+void FPropertyReplicationState::MarkArrayDirty(uint32 Index)
+{
+	const FReplicationStateDescriptor* Descriptor = ReplicationStateDescriptor;
+	FReplicationStateHeader& Header = Private::GetReplicationStateHeader(StateBuffer, Descriptor);
+
+	if (IsInitState())
+	{
+		if (!Private::FReplicationStateHeaderAccessor::GetIsInitStateDirty(Header))
+		{
+			Private::FReplicationStateHeaderAccessor::MarkInitStateDirty(Header);
+			if (Header.IsBound())
+			{
+				MarkNetObjectStateDirty(Header);
+			}
+		}
+	}
+	else
+	{
+		FNetBitArrayView MemberChangeMask = Private::GetMemberChangeMask(StateBuffer, Descriptor);
+		FReplicationStateMemberChangeMaskDescriptor ChangeMaskInfo = Descriptor->MemberChangeMaskDescriptors[Index];
+		ChangeMaskInfo.BitOffset += TArrayPropertyChangeMaskBitIndex;
+		ChangeMaskInfo.BitCount = 1;
+		Private::MarkDirty(Header, MemberChangeMask, ChangeMaskInfo);
+	}
+}
+
 bool FPropertyReplicationState::PollPropertyReplicationState(const void* RESTRICT SrcStateData)
 {
 	if (IsValid())
diff --git a/Engine/Source/Runtime/Experimental/Iris/Core/Private/Iris/ReplicationState/ReplicationStateDescriptorBuilder.cpp b/Engine/Source/Runtime/Experimental/Iris/Core/Private/Iris/ReplicationState/ReplicationStateDescriptorBuilder.cpp
--- a/Engine/Source/Runtime/Experimental/Iris/Core/Private/Iris/ReplicationState/ReplicationStateDescriptorBuilder.cpp
+++ b/Engine/Source/Runtime/Experimental/Iris/Core/Private/Iris/ReplicationState/ReplicationStateDescriptorBuilder.cpp
@@ -11,6 +11,7 @@
 #include "Iris/ReplicationState/PropertyNetSerializerInfoRegistry.h"
 #include "Iris/ReplicationState/InternalPropertyReplicationState.h"
 #include "Iris/ReplicationState/IrisFastArraySerializer.h"
+#include "Iris/ReplicationState/PropertyReplicationState.h"
 #include "Iris/ReplicationState/ReplicationStateDescriptorConfig.h"
 #include "Iris/ReplicationState/ReplicationStateDescriptorRegistry.h"
 #include "Iris/ReplicationSystem/ReplicationOperations.h"
@@ -49,7 +50,11 @@ static const FName PropertyNetSerializerRegistry_NAME_ChangeMaskStorage("ChangeM
 static const FName ReplicationStateDescriptorBuilder_NAME_NetCullDistanceSquared("NetCullDistanceSquared");
 static const FName ReplicationStateDescriptorBuilder_NAME_RoleGroup("RoleGroup");
 
-static TAutoConsoleVariable<int32> CVarIrisUseNativeFastArray(TEXT("net.Iris.UseNativeFastArray"), 1, TEXT("Enable or disable IrisNativeFastArray."));
+static bool bIrisUseNativeFastArray = true;
+static FAutoConsoleVariableRef CVarIrisUseNativeFastArray(TEXT("net.Iris.UseNativeFastArray"), bIrisUseNativeFastArray, TEXT("Enable or disable IrisNativeFastArray."));
+
+static bool bIrisUseChangeMaskForTArray = true;
+static FAutoConsoleVariableRef CVarIrisUseChangeMaskForTArray(TEXT("net.Iris.UseChangeMaskForTArray"), bIrisUseChangeMaskForTArray, TEXT("Enable or disable the use of a changemask for individual elements in TArrays. When enabled and packet loss occurs the received array may not reflect a state which was ever present on the sending side since the array won't be replicated atomically. Enabled by default."));
 
 static bool bWarnAboutStructsWithCustomSerialization = true;
 static FAutoConsoleVariableRef CVarIrisWarnAboutStructsWithCustomSerialization(TEXT("net.Iris.WarnAboutStructsWithCustomSerialization"), bWarnAboutStructsWithCustomSerialization, TEXT("Warn when generating descriptors for structs with custom serialization."));
@@ -73,7 +78,8 @@ enum class EMemberPropertyTraits : uint32
 	HasCustomObjectReference			= HasObjectReference << 1U,
 	IsSourceTriviallyConstructible		= HasCustomObjectReference << 1U,
 	IsSourceTriviallyDestructible		= IsSourceTriviallyConstructible << 1U,
-	IsFastArray							= IsSourceTriviallyDestructible << 1U,
+	IsTArray							= IsSourceTriviallyDestructible << 1U,
+	IsFastArray							= IsTArray << 1U,
 	IsNativeFastArray					= IsFastArray << 1U,
 	IsFastArrayWithExtraProperties		= IsNativeFastArray << 1U,
 	IsFastArrayItem						= IsFastArrayWithExtraProperties << 1U,
@@ -1865,6 +1871,7 @@ void FPropertyReplicationStateDescriptorBuilder::GetIrisPropertyTraits(FMemberPr
 
 	Traits |= ((Property->PropertyFlags & (CPF_NoDestructor | CPF_IsPlainOldData)) != 0 ? EMemberPropertyTraits::IsSourceTriviallyDestructible : EMemberPropertyTraits::None);
 	Traits |= ((Property->PropertyFlags & (CPF_ZeroConstructor | CPF_IsPlainOldData)) != 0 ? EMemberPropertyTraits::IsSourceTriviallyConstructible : EMemberPropertyTraits::None);
+	Traits |= (Property->IsA(FArrayProperty::StaticClass()) ? EMemberPropertyTraits::IsTArray : EMemberPropertyTraits::None);
 
 	if (Data)
 	{
@@ -1934,6 +1941,13 @@ bool FPropertyReplicationStateDescriptorBuilder::IsSupportedProperty(FMemberProp
 	{
 		OutMemberProperty.ChangeMaskBits = 1U + FIrisFastArraySerializer::IrisFastArrayChangeMaskBits;
 	}
+	else if (EnumHasAnyFlags(OutMemberProperty.Traits, EMemberPropertyTraits::IsTArray))
+	{
+		if (bIrisUseChangeMaskForTArray)
+		{
+			OutMemberProperty.ChangeMaskBits = 1U + FPropertyReplicationState::TArrayElementChangeMaskBits;
+		}
+	}
 
 	// Structs with custom serializers may have a default custom replication fragment.
 	if (!OutMemberProperty.CreateAndRegisterReplicationFragmentFunction)
@@ -2190,7 +2204,7 @@ EMemberPropertyTraits FPropertyReplicationStateDescriptorBuilder::GetFastArrayPr
 				EMemberPropertyTraits Traits = EMemberPropertyTraits::IsFastArray;
 
 				// If the struct is derived from FIrisFastArraySerializer and has a single property we can treat it as a native FastArray
-				const bool bUseNativeFastArray = CVarIrisUseNativeFastArray.GetValueOnAnyThread() != 0 && ReplicatedPropertyCount == 1U;
+				const bool bUseNativeFastArray = bIrisUseNativeFastArray && ReplicatedPropertyCount == 1U;
 				if (bUseNativeFastArray && Struct->IsChildOf(FIrisFastArraySerializer::StaticStruct()))
 				{
 					Traits |= EMemberPropertyTraits::HasPushBasedDirtiness | EMemberPropertyTraits::IsNativeFastArray;
diff --git a/Engine/Source/Runtime/Experimental/Iris/Core/Public/Iris/ReplicationState/PropertyReplicationState.h b/Engine/Source/Runtime/Experimental/Iris/Core/Public/Iris/ReplicationState/PropertyReplicationState.h
--- a/Engine/Source/Runtime/Experimental/Iris/Core/Public/Iris/ReplicationState/PropertyReplicationState.h
+++ b/Engine/Source/Runtime/Experimental/Iris/Core/Public/Iris/ReplicationState/PropertyReplicationState.h
@@ -2,7 +2,7 @@
 
 #pragma once
 
-#include "CoreTypes.h"
+#include "HAL/Platform.h"
 #include "Templates/RefCounting.h"
 
 namespace UE::Net
@@ -71,14 +71,10 @@ public:
 	*/
 	IRISCORE_API void GetPropertyValue(uint32 Index, void* DstValue) const;
 
-	/** Need array editors?, to edit arrays with statemasks */
-	/** PropertyNetArrayEditor EditPropertyAsArray(uint32 Index); */
-	/** const PropertyArrayReader GetPropertyAsArray(uint32 Index) */
-
 	/** Is the property at the given index dirty, for properties with multiple bits in the statemask this will return true if any of those bits are set */
 	IRISCORE_API bool IsDirty(uint32 Index) const;
 
-	/** Explicitly mark the the property at the given Index as dirty, for properties with multiple bits in the statemask this will mark all bits dirty */
+	/** Explicitly mark the property at the given Index as dirty, for properties with multiple bits in the statemask this will mark all bits dirty */
 	IRISCORE_API void MarkDirty(uint32 Index);
 
 	/** $IRIS TODO: Move Poll/Push/CallRepNotifies out of PropertyReplicationState as loose functions */
@@ -131,9 +127,23 @@ public:
 	const FReplicationStateDescriptor* GetReplicationStateDescriptor() const { return ReplicationStateDescriptor; }
 	uint8* GetStateBuffer() const { return StateBuffer; }
 
+protected:
+	/** Explicitly mark an array at the given Index as dirty. It will not mark element changemask bits as dirty. */
+	IRISCORE_API void MarkArrayDirty(uint32 Index);
+
 public:
 	// These methods are for internal use only.
 
+	enum : unsigned
+	{
+		// How many bits to use to track dirtiness for individual elements.
+		TArrayElementChangeMaskBits = 63U,
+		// The changemask bit into the TArray changemask info that indicates the TArray is dirty.
+		TArrayPropertyChangeMaskBitIndex = 0U,
+		// At which offset in the TArray changemask element dirtiness begins.
+		TArrayElementChangeMaskBitOffset = 1U,
+	};
+
 	// Polls source data for a single property. If the property has changed it will be marked as dirty.
 	void PollProperty(const void* SrcData, uint32 Index);
 
