Iris
* Added unit test for filtering that runs both pre and post filter types

#rnx

diff --git a/Engine/Plugins/Runtime/ReplicationSystemTestPlugin/Source/Private/Tests/ReplicationSystem/Filtering/TestFiltering.cpp b/Engine/Plugins/Runtime/ReplicationSystemTestPlugin/Source/Private/Tests/ReplicationSystem/Filtering/TestFiltering.cpp
--- a/Engine/Plugins/Runtime/ReplicationSystemTestPlugin/Source/Private/Tests/ReplicationSystem/Filtering/TestFiltering.cpp
+++ b/Engine/Plugins/Runtime/ReplicationSystemTestPlugin/Source/Private/Tests/ReplicationSystem/Filtering/TestFiltering.cpp
@@ -25,7 +25,7 @@ protected:
 	virtual void TearDown() override
 	{
 		FReplicationSystemServerClientTestFixture::TearDown();
-		RestoreNetObjectPrioritizerDefinitions();
+		RestoreFilterDefinitions();
 	}
 
 private:
@@ -58,7 +58,7 @@ private:
 		DefinitionsProperty->CopyCompleteValue((void*)(UPTRINT(FilterDefinitions) + DefinitionsProperty->GetOffset_ForInternal()), &NewFilterDefinitions);
 	}
 
-	void RestoreNetObjectPrioritizerDefinitions()
+	void RestoreFilterDefinitions()
 	{
 		// Restore NetObjectPrioritizerDefinitions CDO state from the saved state.
 		const UClass* NetObjectFilterDefinitionsClass = UNetObjectFilterDefinitions::StaticClass();
@@ -69,6 +69,9 @@ private:
 
 		MockFilterHandle = InvalidNetObjectFilterHandle;
 		MockNetObjectFilter = nullptr;
+
+		MockFilterWithFragmentsHandle = InvalidNetObjectFilterHandle;
+		MockNetObjectFilterWithFragments = nullptr;
 	}
 
 	void InitMockNetObjectFilter()
@@ -1220,6 +1223,50 @@ UE_NET_TEST_FIXTURE(FTestFilteringFixture, ObjectGetsUpdatedFilterOutSetting)
 	UE_NET_ASSERT_EQ(Client->GetReplicationBridge()->GetReplicatedObject(ServerObject->NetRefHandle), nullptr);
 }
 
+UE_NET_TEST_FIXTURE(FTestFilteringFixture, MixPreAndPostFilters)
+{
+	// Setup dynamic filter for the test.
+	{
+		UMockNetObjectFilter::FFunctionCallSetup CallSetup;
+		CallSetup.AddObject.bReturnValue = true;
+		CallSetup.Filter.bFilterOutByDefault = true;
+
+		MockNetObjectFilter->SetFunctionCallSetup(CallSetup);
+		MockNetObjectFilter->ResetFunctionCallStatus();
+
+		MockNetObjectFilterWithFragments->SetFunctionCallSetup(CallSetup);
+		MockNetObjectFilterWithFragments->ResetFunctionCallStatus();
+	}
+
+	// Add client
+	FReplicationSystemTestClient* Client = CreateClient();
+
+	// Create the objects and set filter
+	UTestFilteringObject* ServerObjectPreFilter = Server->CreateObject<UTestFilteringObject>();
+	Server->ReplicationSystem->SetFilter(ServerObjectPreFilter->NetRefHandle, MockFilterHandle);
+
+	UTestFilteringObject* ServerObjectPostFilter = Server->CreateObject<UTestFilteringObject>();
+	Server->ReplicationSystem->SetFilter(ServerObjectPostFilter->NetRefHandle, MockFilterWithFragmentsHandle);
+
+	// Create a non-filtered object
+	UTestFilteringObject* ServerObjectNoFilter = Server->CreateObject<UTestFilteringObject>();
+
+	// We want the objects to be filtered out
+	constexpr bool bFilterOut = true;
+	ServerObjectPreFilter->SetFilterOut(bFilterOut);
+	ServerObjectPostFilter->SetFilterOut(bFilterOut);
+
+	// Send and deliver packets
+	Server->PreSendUpdate();
+	Server->SendAndDeliverTo(Client, DeliverPacket);
+	Server->PostSendUpdate();
+
+	// Check that the object does not exist on the client.
+	UE_NET_ASSERT_EQ(Client->GetReplicationBridge()->GetReplicatedObject(ServerObjectPreFilter->NetRefHandle), nullptr);
+	UE_NET_ASSERT_EQ(Client->GetReplicationBridge()->GetReplicatedObject(ServerObjectPostFilter->NetRefHandle), nullptr);
+	UE_NET_ASSERT_NE(Client->GetReplicationBridge()->GetReplicatedObject(ServerObjectNoFilter->NetRefHandle), nullptr);
+}
+
 UE_NET_TEST_FIXTURE(FTestFilteringFixture, DynamicFilteredOutSubObjectsAreResetWhenIndexIsReused)
 {
 	// Setup dynamic filter for the test
