ReplicationSystemTest
* Break on failed line instead of inside the FMessageLog

#rnx

diff --git a/Engine/Plugins/Runtime/ReplicationSystemTestPlugin/Source/Private/NetworkAutomationTestMacros.cpp b/Engine/Plugins/Runtime/ReplicationSystemTestPlugin/Source/Private/NetworkAutomationTestMacros.cpp
--- a/Engine/Plugins/Runtime/ReplicationSystemTestPlugin/Source/Private/NetworkAutomationTestMacros.cpp
+++ b/Engine/Plugins/Runtime/ReplicationSystemTestPlugin/Source/Private/NetworkAutomationTestMacros.cpp
@@ -57,7 +57,6 @@ FTestMessageLog::~FTestMessageLog()
 		break;
 	case ELogVerbosity::Error:
 		UE_TEST_MESSAGE_LOG_(Error);
-		UE_DEBUG_BREAK();
 		break;
 	case ELogVerbosity::Warning:
 		UE_TEST_MESSAGE_LOG_(Warning);
diff --git a/Engine/Plugins/Runtime/ReplicationSystemTestPlugin/Source/Private/NetworkAutomationTestMacros.h b/Engine/Plugins/Runtime/ReplicationSystemTestPlugin/Source/Private/NetworkAutomationTestMacros.h
--- a/Engine/Plugins/Runtime/ReplicationSystemTestPlugin/Source/Private/NetworkAutomationTestMacros.h
+++ b/Engine/Plugins/Runtime/ReplicationSystemTestPlugin/Source/Private/NetworkAutomationTestMacros.h
@@ -167,10 +167,10 @@ UE_NET_COMPARE_RETURN_TESTRESULT(TCmpLessOrEqual, <=)
 UE_NET_COMPARE_RETURN_TESTRESULT(TCmpGreater, >)
 UE_NET_COMPARE_RETURN_TESTRESULT(TCmpGreaterOrEqual, >=)
 
-#define UE_NET_FAIL_MSG(Message) return this->AddTestFailure(), UE::Net::FTestMessageLog(*this, ELogVerbosity::Error) = UE::Net::FTestMessage() << Message
+#define UE_NET_FAIL_MSG(Message) return UE_DEBUG_BREAK(), this->AddTestFailure(), UE::Net::FTestMessageLog(*this, ELogVerbosity::Error) = UE::Net::FTestMessage() << Message
 #define UE_NET_WARN_MSG(Message) this->AddTestWarning(), UE::Net::FTestMessageLog(*this, ELogVerbosity::Warning) = UE::Net::FTestMessage() << Message
 
-#define UE_NET_FAIL(Message) return this->AddTestFailure(), void(UE::Net::FTestMessageLog(*this, ELogVerbosity::Error, Message))
+#define UE_NET_FAIL(Message) return UE_DEBUG_BREAK(), this->AddTestFailure(), void(UE::Net::FTestMessageLog(*this, ELogVerbosity::Error, Message))
 #define UE_NET_WARN(Message) this->AddTestWarning(), void(UE::Net::FTestMessageLog(*this, ELogVerbosity::Warning, Message))
 
 #define UE_NET_TEST_(CompareWithTestResult, V1, V2, FailFunction) if (UE::Net::FTestResult Result_ = CompareWithTestResult(V1, V2, #V1, #V2)) {} else FailFunction(Result_.GetMessage())
