Iris
* Added GridFilter unit tests
* Encapsulated RepTag creation hack for unit tests
* Fix FragmentGridFilter reading world location fragment using wrong FVector unit type

#rnx

diff --git a/Engine/Plugins/Runtime/ReplicationSystemTestPlugin/Source/Private/Tests/ReplicationSystem/Filtering/TestFiltering.cpp b/Engine/Plugins/Runtime/ReplicationSystemTestPlugin/Source/Private/Tests/ReplicationSystem/Filtering/TestFiltering.cpp
--- a/Engine/Plugins/Runtime/ReplicationSystemTestPlugin/Source/Private/Tests/ReplicationSystem/Filtering/TestFiltering.cpp
+++ b/Engine/Plugins/Runtime/ReplicationSystemTestPlugin/Source/Private/Tests/ReplicationSystem/Filtering/TestFiltering.cpp
@@ -35,11 +35,11 @@ private:
 		const FProperty* DefinitionsProperty = NetObjectFilterDefinitionsClass->FindPropertyByName("NetObjectFilterDefinitions");
 		check(DefinitionsProperty != nullptr);
 
-		// Save NetObjectFilterDefinitions CDO state.
+		// Save CDO state.
 		UNetObjectFilterDefinitions* FilterDefinitions = GetMutableDefault<UNetObjectFilterDefinitions>();
 		DefinitionsProperty->CopyCompleteValue(&OriginalFilterDefinitions, (void*)(UPTRINT(FilterDefinitions) + DefinitionsProperty->GetOffset_ForInternal()));
-
-		// Modify definitions to only include our mock prioritizer. Ugly... 
+		
+		// Modify definitions to only include our filters. Ugly... 
 		TArray<FNetObjectFilterDefinition> NewFilterDefinitions;
 		{
 			FNetObjectFilterDefinition& MockDefinition = NewFilterDefinitions.Emplace_GetRef();
@@ -60,7 +60,7 @@ private:
 
 	void RestoreFilterDefinitions()
 	{
-		// Restore NetObjectPrioritizerDefinitions CDO state from the saved state.
+		// Restore CDO state from the saved state.
 		const UClass* NetObjectFilterDefinitionsClass = UNetObjectFilterDefinitions::StaticClass();
 		const FProperty* DefinitionsProperty = NetObjectFilterDefinitionsClass->FindPropertyByName("NetObjectFilterDefinitions");
 		UNetObjectFilterDefinitions* FilterDefinitions = GetMutableDefault<UNetObjectFilterDefinitions>();
diff --git a/Engine/Plugins/Runtime/ReplicationSystemTestPlugin/Source/Private/Tests/ReplicationSystem/Filtering/TestFilteringObject.cpp b/Engine/Plugins/Runtime/ReplicationSystemTestPlugin/Source/Private/Tests/ReplicationSystem/Filtering/TestFilteringObject.cpp
--- a/Engine/Plugins/Runtime/ReplicationSystemTestPlugin/Source/Private/Tests/ReplicationSystem/Filtering/TestFilteringObject.cpp
+++ b/Engine/Plugins/Runtime/ReplicationSystemTestPlugin/Source/Private/Tests/ReplicationSystem/Filtering/TestFilteringObject.cpp
@@ -20,3 +20,17 @@ void UTestFilteringObject::GetLifetimeReplicatedProps( TArray< class FLifetimePr
 	Super::GetLifetimeReplicatedProps(OutLifetimeProps);
 	DOREPLIFETIME(ThisClass, NetTest_FilterOut);
 }
+
+
+void UTestLocationFragmentFilteringObject::RegisterReplicationFragments(UE::Net::FFragmentRegistrationContext& Context, UE::Net::EFragmentRegistrationFlags RegistrationFlags)
+{
+	// Build descriptors and allocate PropertyReplicationFragments for this object
+	UE::Net::FReplicationFragmentUtil::CreateAndRegisterFragmentsForObject(this, Context, RegistrationFlags);
+}
+
+void UTestLocationFragmentFilteringObject::GetLifetimeReplicatedProps(TArray< class FLifetimeProperty >& OutLifetimeProps) const
+{
+	Super::GetLifetimeReplicatedProps(OutLifetimeProps);
+	DOREPLIFETIME(ThisClass, WorldLocation);
+	DOREPLIFETIME(ThisClass, NetCullDistanceSquared);
+}
diff --git a/Engine/Plugins/Runtime/ReplicationSystemTestPlugin/Source/Private/Tests/ReplicationSystem/Filtering/TestFilteringObject.h b/Engine/Plugins/Runtime/ReplicationSystemTestPlugin/Source/Private/Tests/ReplicationSystem/Filtering/TestFilteringObject.h
--- a/Engine/Plugins/Runtime/ReplicationSystemTestPlugin/Source/Private/Tests/ReplicationSystem/Filtering/TestFilteringObject.h
+++ b/Engine/Plugins/Runtime/ReplicationSystemTestPlugin/Source/Private/Tests/ReplicationSystem/Filtering/TestFilteringObject.h
@@ -21,3 +21,25 @@ private:
 	UPROPERTY(Replicated)
 	bool NetTest_FilterOut = false;
 };
+
+
+
+UCLASS()
+class UTestLocationFragmentFilteringObject : public UReplicatedTestObject
+{
+	GENERATED_BODY()
+
+public:
+	
+protected:
+	virtual void RegisterReplicationFragments(UE::Net::FFragmentRegistrationContext& Fragments, UE::Net::EFragmentRegistrationFlags RegistrationFlags) override;
+
+public:
+	
+	UPROPERTY(Replicated)
+	FVector WorldLocation;
+
+	UPROPERTY(Replicated)
+	float NetCullDistanceSquared = 0.0f;
+};
+
diff --git a/Engine/Plugins/Runtime/ReplicationSystemTestPlugin/Source/Private/Tests/ReplicationSystem/Filtering/TestGridFilter.cpp b/Engine/Plugins/Runtime/ReplicationSystemTestPlugin/Source/Private/Tests/ReplicationSystem/Filtering/TestGridFilter.cpp
new file mode 100644
--- /dev/null
+++ b/Engine/Plugins/Runtime/ReplicationSystemTestPlugin/Source/Private/Tests/ReplicationSystem/Filtering/TestGridFilter.cpp
@@ -0,0 +1,217 @@
+// Copyright Epic Games, Inc. All Rights Reserved.
+
+#include "NetworkAutomationTest.h"
+#include "NetworkAutomationTestMacros.h"
+#include "Iris/ReplicationSystem/ReplicationSystem.h"
+#include "Iris/ReplicationSystem/Filtering/NetObjectFilter.h"
+#include "Iris/ReplicationSystem/Filtering/NetObjectFilterDefinitions.h"
+#include "Iris/ReplicationSystem/Filtering/NetObjectGridFilter.h"
+
+#include "Tests/ReplicationSystem/Filtering/MockNetObjectFilter.h"
+#include "Tests/ReplicationSystem/Filtering/TestFilteringObject.h"
+#include "Tests/ReplicationSystem/ReplicationSystemServerClientTestFixture.h"
+
+
+namespace UE::Net::Private
+{
+
+class FTestGridFilterFixture : public FReplicationSystemServerClientTestFixture
+{
+protected:
+	virtual void SetUp() override
+	{
+		InitFilterDefinitions();
+		FReplicationSystemServerClientTestFixture::SetUp();
+		InitFilterHandles();
+	}
+
+	virtual void TearDown() override
+	{
+		FReplicationSystemServerClientTestFixture::TearDown();
+		RestoreFilterDefinitions();
+	}
+
+private:
+	void InitFilterDefinitions()
+	{
+		const UClass* NetObjectFilterDefinitionsClass = UNetObjectFilterDefinitions::StaticClass();
+		const FProperty* DefinitionsProperty = NetObjectFilterDefinitionsClass->FindPropertyByName("NetObjectFilterDefinitions");
+		check(DefinitionsProperty != nullptr);
+
+		// Save CDO state.
+		UNetObjectFilterDefinitions* FilterDefinitions = GetMutableDefault<UNetObjectFilterDefinitions>();
+		DefinitionsProperty->CopyCompleteValue(&OriginalFilterDefinitions, (void*)(UPTRINT(FilterDefinitions) + DefinitionsProperty->GetOffset_ForInternal()));
+
+		// Modify definitions to only include our filters. Ugly... 
+		TArray<FNetObjectFilterDefinition> NewFilterDefinitions;
+		{
+			FNetObjectFilterDefinition& GridWorldLocDef = NewFilterDefinitions.Emplace_GetRef();
+			GridWorldLocDef.FilterName = "NetObjectGridWorldLocFilter";
+			GridWorldLocDef.ClassName = "/Script/IrisCore.NetObjectGridWorldLocFilter";
+			GridWorldLocDef.ConfigClassName = "/Script/IrisCore.NetObjectGridFilterConfig";
+		}
+
+		{
+			FNetObjectFilterDefinition& GridFragmentDef = NewFilterDefinitions.Emplace_GetRef();
+			GridFragmentDef.FilterName = "NetObjectGridFragmentLocFilter";
+			GridFragmentDef.ClassName = "/Script/IrisCore.NetObjectGridFragmentLocFilter";
+			GridFragmentDef.ConfigClassName = "/Script/IrisCore.NetObjectGridFilterConfig";
+		}
+
+		DefinitionsProperty->CopyCompleteValue((void*)(UPTRINT(FilterDefinitions) + DefinitionsProperty->GetOffset_ForInternal()), &NewFilterDefinitions);
+	}
+
+	void RestoreFilterDefinitions()
+	{
+		// Restore CDO state from the saved state.
+		const UClass* NetObjectFilterDefinitionsClass = UNetObjectFilterDefinitions::StaticClass();
+		const FProperty* DefinitionsProperty = NetObjectFilterDefinitionsClass->FindPropertyByName("NetObjectFilterDefinitions");
+		UNetObjectFilterDefinitions* FilterDefinitions = GetMutableDefault<UNetObjectFilterDefinitions>();
+		DefinitionsProperty->CopyCompleteValue((void*)(UPTRINT(FilterDefinitions) + DefinitionsProperty->GetOffset_ForInternal()), &OriginalFilterDefinitions);
+		OriginalFilterDefinitions.Empty();
+
+		WorldLocFilterHandle = InvalidNetObjectFilterHandle;
+		WorldLocFilter = nullptr;
+
+		FragmentLocFilterHandle = InvalidNetObjectFilterHandle;
+		FragmentLocFilter = nullptr;
+	}
+
+	void InitFilterHandles()
+	{
+		WorldLocFilter = ExactCast<UNetObjectGridWorldLocFilter>(Server->GetReplicationSystem()->GetFilter("NetObjectGridWorldLocFilter"));
+		WorldLocFilterHandle = Server->GetReplicationSystem()->GetFilterHandle("NetObjectGridWorldLocFilter");
+
+		FragmentLocFilter = ExactCast<UNetObjectGridFragmentLocFilter>(Server->GetReplicationSystem()->GetFilter("NetObjectGridFragmentLocFilter"));
+		FragmentLocFilterHandle = Server->GetReplicationSystem()->GetFilterHandle("NetObjectGridFragmentLocFilter");
+	}
+
+protected:
+	UNetObjectGridWorldLocFilter* WorldLocFilter;
+	FNetObjectFilterHandle WorldLocFilterHandle;
+
+	UNetObjectGridFragmentLocFilter* FragmentLocFilter;
+	FNetObjectFilterHandle FragmentLocFilterHandle;
+
+private:
+	TArray<FNetObjectFilterDefinition> OriginalFilterDefinitions;
+};
+
+// Test world loc filter
+UE_NET_TEST_FIXTURE(FTestGridFilterFixture, TestWorldLocGridFilter)
+{
+	// Spawn object on server
+	UReplicatedTestObject* ServerObjectZero = Server->CreateObjectWithWorldLocation();
+	UReplicatedTestObject* ServerObjectLimit = Server->CreateObjectWithWorldLocation();
+	UReplicatedTestObject* ServerObjectCulled = Server->CreateObjectWithWorldLocation();
+	UReplicatedTestObject* ServerObjectVeryFar = Server->CreateObjectWithWorldLocation();
+
+
+	const UNetObjectGridFilterConfig* DefaultGridConfig = Cast<UNetObjectGridFilterConfig>(UNetObjectGridFilterConfig::StaticClass()->GetDefaultObject());
+
+	struct FObjectLoc 
+	{
+		FVector Loc;
+		float CullDistance;
+	};
+	TMap<const UObject*, FObjectLoc> ObjectLocs;
+
+	// Visible objects
+	ObjectLocs.Add(ServerObjectZero,	FObjectLoc{ FVector::ZeroVector, 1500.f} );
+	ObjectLocs.Add(ServerObjectLimit,	FObjectLoc{ FVector(DefaultGridConfig->CellSizeX, DefaultGridConfig->CellSizeY, 1500.f), DefaultGridConfig->MaxCullDistance});
+
+	// Culled objects
+	ObjectLocs.Add(ServerObjectCulled,	FObjectLoc{ FVector(DefaultGridConfig->CellSizeX + 100.f, DefaultGridConfig->CellSizeY + 100.f, 100.f), 1500.f });
+	ObjectLocs.Add(ServerObjectVeryFar, FObjectLoc{ FVector(DefaultGridConfig->CellSizeX + 99999.f, DefaultGridConfig->CellSizeY + 99999.f, 99999.f), 1500.f });
+
+	Server->GetReplicationBridge()->SetExternalWorldLocationUpdateFunctor([&](FNetRefHandle NetHandle, const UObject* ReplicatedObject, FVector& OutLocation, float& OutCullDistance)
+	{
+		const FObjectLoc& ObjectLoc = ObjectLocs[ReplicatedObject];
+		OutLocation = ObjectLoc.Loc;
+		OutCullDistance = ObjectLoc.CullDistance;
+	});
+
+	// Apply grid filter
+	Server->ReplicationSystem->SetFilter(ServerObjectZero->NetRefHandle, WorldLocFilterHandle);
+	Server->ReplicationSystem->SetFilter(ServerObjectLimit->NetRefHandle, WorldLocFilterHandle);
+	Server->ReplicationSystem->SetFilter(ServerObjectCulled->NetRefHandle, WorldLocFilterHandle);
+	Server->ReplicationSystem->SetFilter(ServerObjectVeryFar->NetRefHandle, WorldLocFilterHandle);
+
+	// Add client (view location at 0,0,0)
+	FReplicationSystemTestClient* Client = CreateClient();
+
+	// Send and deliver packet
+	Server->PreSendUpdate();
+	Server->SendAndDeliverTo(Client, DeliverPacket);
+	Server->PostSendUpdate();
+
+	// Test visible objects
+	UE_NET_ASSERT_NE(Client->GetReplicationBridge()->GetReplicatedObject(ServerObjectZero->NetRefHandle), nullptr);
+	UE_NET_ASSERT_NE(Client->GetReplicationBridge()->GetReplicatedObject(ServerObjectLimit->NetRefHandle), nullptr);
+
+	// Test culled objects
+	// $IRIS TODO:  This fails since the culldistance makes the actor cross over to the same cell as the client. Need to decide if cull distances should be better respected in Iris
+	//UE_NET_ASSERT_EQ(Client->GetReplicationBridge()->GetReplicatedObject(ServerObjectCulled->NetRefHandle), nullptr);
+	UE_NET_ASSERT_EQ(Client->GetReplicationBridge()->GetReplicatedObject(ServerObjectVeryFar->NetRefHandle), nullptr);
+
+	Server->DestroyObject(ServerObjectZero);
+	Server->DestroyObject(ServerObjectLimit);
+	Server->DestroyObject(ServerObjectCulled);
+	Server->DestroyObject(ServerObjectVeryFar);
+}
+
+// Test fragment data filter
+UE_NET_TEST_FIXTURE(FTestGridFilterFixture, TestFragmentFilter)
+{
+	// Spawn object on server
+	UTestLocationFragmentFilteringObject* ServerObjectZero = Server->CreateObject<UTestLocationFragmentFilteringObject>();
+	UTestLocationFragmentFilteringObject* ServerObjectLimit = Server->CreateObject<UTestLocationFragmentFilteringObject>();
+	UTestLocationFragmentFilteringObject* ServerObjectCulled = Server->CreateObject<UTestLocationFragmentFilteringObject>();
+	UTestLocationFragmentFilteringObject* ServerObjectVeryFar = Server->CreateObject<UTestLocationFragmentFilteringObject>();
+
+	const UNetObjectGridFilterConfig* DefaultGridConfig = Cast<UNetObjectGridFilterConfig>(UNetObjectGridFilterConfig::StaticClass()->GetDefaultObject());
+
+
+	ServerObjectZero->WorldLocation = FVector::ZeroVector;
+	ServerObjectZero->NetCullDistanceSquared = 1500.f * 1500.f;
+	
+	ServerObjectLimit->WorldLocation = FVector(DefaultGridConfig->CellSizeX, DefaultGridConfig->CellSizeY, 1500.f);
+	ServerObjectLimit->NetCullDistanceSquared = DefaultGridConfig->MaxCullDistance * DefaultGridConfig->MaxCullDistance;
+
+	ServerObjectCulled->WorldLocation = FVector(DefaultGridConfig->CellSizeX + 100.f, DefaultGridConfig->CellSizeY + 100.f, 100.f);
+	ServerObjectCulled->NetCullDistanceSquared = 1500.f * 1500.f;
+	
+	ServerObjectVeryFar->WorldLocation = FVector(DefaultGridConfig->CellSizeX + 99999.f, DefaultGridConfig->CellSizeY + 99999.f, 99999.f);
+	ServerObjectVeryFar->NetCullDistanceSquared = 1500.f * 1500.f;
+
+	// Apply grid filter
+	Server->ReplicationSystem->SetFilter(ServerObjectZero->NetRefHandle, FragmentLocFilterHandle);
+	Server->ReplicationSystem->SetFilter(ServerObjectLimit->NetRefHandle, FragmentLocFilterHandle);
+	Server->ReplicationSystem->SetFilter(ServerObjectCulled->NetRefHandle, FragmentLocFilterHandle);
+	Server->ReplicationSystem->SetFilter(ServerObjectVeryFar->NetRefHandle, FragmentLocFilterHandle);
+
+	// Add client (view location at 0,0,0)
+	FReplicationSystemTestClient* Client = CreateClient();
+
+	// Send and deliver packet
+	Server->PreSendUpdate();
+	Server->SendAndDeliverTo(Client, DeliverPacket);
+	Server->PostSendUpdate();
+
+	// Test visible objects
+	UE_NET_ASSERT_NE(Client->GetReplicationBridge()->GetReplicatedObject(ServerObjectZero->NetRefHandle), nullptr);
+	UE_NET_ASSERT_NE(Client->GetReplicationBridge()->GetReplicatedObject(ServerObjectLimit->NetRefHandle), nullptr);
+
+	// Test culled objects
+	// $IRIS TODO:  This fails since the culldistance makes the actor cross over to the same cell as the client. Need to decide if cull distances should be better respected in Iris
+	//UE_NET_ASSERT_EQ(Client->GetReplicationBridge()->GetReplicatedObject(ServerObjectCulled->NetRefHandle), nullptr);
+	UE_NET_ASSERT_EQ(Client->GetReplicationBridge()->GetReplicatedObject(ServerObjectVeryFar->NetRefHandle), nullptr);
+
+	Server->DestroyObject(ServerObjectZero);
+	Server->DestroyObject(ServerObjectLimit);
+	Server->DestroyObject(ServerObjectCulled);
+	Server->DestroyObject(ServerObjectVeryFar);
+}
+
+
+} // end namespace UE::Net::Private
\ No newline at end of file
diff --git a/Engine/Plugins/Runtime/ReplicationSystemTestPlugin/Source/Private/Tests/ReplicationSystem/ReplicatedTestObject.cpp b/Engine/Plugins/Runtime/ReplicationSystemTestPlugin/Source/Private/Tests/ReplicationSystem/ReplicatedTestObject.cpp
--- a/Engine/Plugins/Runtime/ReplicationSystemTestPlugin/Source/Private/Tests/ReplicationSystem/ReplicatedTestObject.cpp
+++ b/Engine/Plugins/Runtime/ReplicationSystemTestPlugin/Source/Private/Tests/ReplicationSystem/ReplicatedTestObject.cpp
@@ -63,6 +63,20 @@ UE::Net::FNetRefHandle UReplicatedTestObjectBridge::BeginReplication(UReplicated
 	return Handle;
 }
 
+UE::Net::FNetRefHandle UReplicatedTestObjectBridge::BeginReplication(UReplicatedTestObject* Instance, const UObjectReplicationBridge::FCreateNetRefHandleParams& Params)
+{
+	// Create NetRefHandle for the registered fragments
+	FNetRefHandle Handle = Super::BeginReplication(Instance, Params);
+
+	// This is optional but typically we want to cache at least the NetRefHandle in the game instance to avoid doing map lookups to find it
+	if (Handle.IsValid())
+	{
+		Instance->NetRefHandle = Handle;
+	}
+
+	return Handle;
+}
+
 UE::Net::FNetRefHandle UReplicatedTestObjectBridge::BeginReplication(FNetRefHandle OwnerHandle, UReplicatedTestObject* Instance, FNetRefHandle InsertRelativeToSubObjectHandle, ESubObjectInsertionOrder InsertionOrder)
 {
 	// Create NetRefHandle for the registered fragments
@@ -221,6 +235,11 @@ void UReplicatedTestObjectBridge::SetPollFramePeriod(UReplicatedTestObject* Inst
 	}
 }
 
+void UReplicatedTestObjectBridge::SetExternalWorldLocationUpdateFunctor(TFunction<void(FNetRefHandle NetHandle, const UObject* ReplicatedObject, FVector& OutLocation, float& OutCullDistance)> LocUpdateFunctor)
+{
+	SetInstanceGetWorldObjectInfoFunction(LocUpdateFunctor);
+}
+
 //////////////////////////////////////////////////////////////////////////
 // Implementation for UTestReplicatedIrisPropertyComponent
 //////////////////////////////////////////////////////////////////////////
diff --git a/Engine/Plugins/Runtime/ReplicationSystemTestPlugin/Source/Private/Tests/ReplicationSystem/ReplicatedTestObject.h b/Engine/Plugins/Runtime/ReplicationSystemTestPlugin/Source/Private/Tests/ReplicationSystem/ReplicatedTestObject.h
--- a/Engine/Plugins/Runtime/ReplicationSystemTestPlugin/Source/Private/Tests/ReplicationSystem/ReplicatedTestObject.h
+++ b/Engine/Plugins/Runtime/ReplicationSystemTestPlugin/Source/Private/Tests/ReplicationSystem/ReplicatedTestObject.h
@@ -440,6 +440,7 @@ public:
 	// This is the local Interface, it is up to each bridge implementation to define the interface for that type
 	// In this example we have methods that directly uses UReplicatedTestObject;
 	FNetRefHandle BeginReplication(UReplicatedTestObject* Instance);
+	FNetRefHandle BeginReplication(UReplicatedTestObject* Instance, const UObjectReplicationBridge::FCreateNetRefHandleParams& Params);
 	FNetRefHandle BeginReplication(FNetRefHandle OwnerHandle, UReplicatedTestObject* Instance, FNetRefHandle InsertRelativeToSubObjectHandle = FNetRefHandle(), ESubObjectInsertionOrder InsertionOrder = UReplicationBridge::ESubObjectInsertionOrder::None);
 
 	// For testing we expose some things that normally are not accessible
@@ -447,6 +448,8 @@ public:
 
 	void SetPollFramePeriod(UReplicatedTestObject* Instance, uint8 FramePeriod);
 
+	void SetExternalWorldLocationUpdateFunctor(TFunction<void(FNetRefHandle NetHandle, const UObject* ReplicatedObject, FVector& OutLocation, float& OutCullDistance)> LocUpdateFunctor);
+
 protected:
 	// Type specifics for serializing creation data this will most likely be made into a separate interface to support different types of header data for different types
 	// But if we can avoid having custom data per type for instantiating remote objects as we would like to be able to fully express the state of a replicated object using the define protocol alone. 
@@ -469,6 +472,8 @@ protected:
 	virtual bool IsAllowedToDestroyInstance(const UObject* Instance) const override;
 
 	TArray<TStrongObjectPtr<UObject>>* CreatedObjectsOnNode;
+
+	TFunction<void(FNetRefHandle NetHandle, const UObject* ReplicatedObject, FVector& OutLocation, float& OutCullDistance)> WorldLocationUpdateFunc;
 };
 
 extern const UE::Net::FRepTag RepTag_FakeGeneratedReplicationState_IntB;
diff --git a/Engine/Plugins/Runtime/ReplicationSystemTestPlugin/Source/Private/Tests/ReplicationSystem/ReplicationSystemServerClientTestFixture.cpp b/Engine/Plugins/Runtime/ReplicationSystemTestPlugin/Source/Private/Tests/ReplicationSystem/ReplicationSystemServerClientTestFixture.cpp
--- a/Engine/Plugins/Runtime/ReplicationSystemTestPlugin/Source/Private/Tests/ReplicationSystem/ReplicationSystemServerClientTestFixture.cpp
+++ b/Engine/Plugins/Runtime/ReplicationSystemTestPlugin/Source/Private/Tests/ReplicationSystem/ReplicationSystemServerClientTestFixture.cpp
@@ -160,6 +160,22 @@ UTestReplicatedIrisObject* FReplicationSystemTestNode::CreateObjectWithDynamicSt
 	return CreatedObject;
 }
 
+UTestReplicatedIrisObject* FReplicationSystemTestNode::CreateObjectWithWorldLocation()
+{
+	UTestReplicatedIrisObject* CreatedObject = NewObject<UTestReplicatedIrisObject>();
+	CreatedObjects.Add(TStrongObjectPtr<UObject>(CreatedObject));
+
+	UObjectReplicationBridge::FCreateNetRefHandleParams Params;
+	Params.bCanReceive = true;
+	Params.bNeedsWorldLocationUpdate = true;
+	Params.bAllowDynamicFilter = true;
+
+	// Add it to the bridge for replication
+	ReplicationBridge->BeginReplication(CreatedObject, Params);
+
+	return CreatedObject;
+}
+
 void FReplicationSystemTestNode::DestroyObject(UReplicatedTestObject* Object, EEndReplicationFlags EndReplicationFlags)
 {
 	// Destroy handle
diff --git a/Engine/Plugins/Runtime/ReplicationSystemTestPlugin/Source/Private/Tests/ReplicationSystem/ReplicationSystemServerClientTestFixture.h b/Engine/Plugins/Runtime/ReplicationSystemTestPlugin/Source/Private/Tests/ReplicationSystem/ReplicationSystemServerClientTestFixture.h
--- a/Engine/Plugins/Runtime/ReplicationSystemTestPlugin/Source/Private/Tests/ReplicationSystem/ReplicationSystemServerClientTestFixture.h
+++ b/Engine/Plugins/Runtime/ReplicationSystemTestPlugin/Source/Private/Tests/ReplicationSystem/ReplicationSystemServerClientTestFixture.h
@@ -119,6 +119,7 @@ public:
 	UTestReplicatedIrisObject* CreateObject(const UTestReplicatedIrisObject::FComponents& Components);
 	UTestReplicatedIrisObject* CreateSubObject(FNetRefHandle Owner, const UTestReplicatedIrisObject::FComponents& Components);
 	UTestReplicatedIrisObject* CreateObjectWithDynamicState(uint32 NumComponents, uint32 NumIrisComponents, uint32 NumDynamicStateComponents);
+	UTestReplicatedIrisObject* CreateObjectWithWorldLocation();
 	void DestroyObject(UReplicatedTestObject*, EEndReplicationFlags EndReplicationFlags = EEndReplicationFlags::Destroy);
 
 	// Connection
diff --git a/Engine/Source/Programs/ReplicationSystemTest/ReplicationSystemTest.Target.cs b/Engine/Source/Programs/ReplicationSystemTest/ReplicationSystemTest.Target.cs
--- a/Engine/Source/Programs/ReplicationSystemTest/ReplicationSystemTest.Target.cs
+++ b/Engine/Source/Programs/ReplicationSystemTest/ReplicationSystemTest.Target.cs
@@ -40,6 +40,7 @@ public class ReplicationSystemTestTarget : TargetRules
 		bUseIris = true;
 
 		GlobalDefinitions.Add("UE_TRACE_ENABLED=1");
+		GlobalDefinitions.Add("UE_NET_TEST_FAKE_REP_TAGS=1");
 	}
 }
 
diff --git a/Engine/Source/Runtime/Experimental/Iris/Core/Private/Iris/ReplicationState/ReplicationStateDescriptorBuilder.cpp b/Engine/Source/Runtime/Experimental/Iris/Core/Private/Iris/ReplicationState/ReplicationStateDescriptorBuilder.cpp
--- a/Engine/Source/Runtime/Experimental/Iris/Core/Private/Iris/ReplicationState/ReplicationStateDescriptorBuilder.cpp
+++ b/Engine/Source/Runtime/Experimental/Iris/Core/Private/Iris/ReplicationState/ReplicationStateDescriptorBuilder.cpp
@@ -28,6 +28,10 @@
 #include "Templates/AlignmentTemplates.h"
 #include "UObject/UnrealType.h"
 
+#ifndef UE_NET_TEST_FAKE_REP_TAGS
+#	define UE_NET_TEST_FAKE_REP_TAGS 0
+#endif
+
 #if (UE_BUILD_SHIPPING || UE_BUILD_TEST)
 #	define UE_NET_ENABLE_DESCRIPTORBUILDER_LOG 0
 #else
@@ -2069,10 +2073,17 @@ FRepTag FPropertyReplicationStateDescriptorBuilder::GetRepTagFromProperty(const
 			UE_LOG_DESCRIPTORBUILDER_WARNING(TEXT("Found NetCullDistanceSquared property that is of type %s instead of a float. This prevents the property from being accessed by networking systems."), ToCStr(Class->GetFName().ToString()));
 		}
 	}
-	// For unit testing purposes. Keep this code last.
-#if IS_PROGRAM
+	//$IRIS TODO: Temp until proper RepTag support is added. 
+	//For unit testing purposes. Keep this code last.
+#if UE_NET_TEST_FAKE_REP_TAGS
 	else
 	{
+		static const FName ReplicationStateDescriptorBuilder_NAME_WorldLocation("WorldLocation");
+		if (PropertyName == ReplicationStateDescriptorBuilder_NAME_WorldLocation)
+		{
+			return RepTag_WorldLocation;
+		}
+
 		ANSICHAR PropertyNameString[NAME_SIZE];
 		PropertyName.GetPlainANSIString(PropertyNameString);
 		if (0 == FCStringAnsi::Strncmp(PropertyNameString, "NetTest_", 8))
@@ -2654,6 +2665,7 @@ SIZE_T FReplicationStateDescriptorBuilder::CreateDescriptorsForClass(FResult& Cr
 	// Add special Actor properties
 	{
 		constexpr bool bExactMatch = true;
+		//$IRIS TODO: Probably should use a cached pointer of the Actor class here.
 		UClass* ActorClass = CastChecked<UClass>(StaticFindObject(UClass::StaticClass(), nullptr, TEXT("/Script/Engine.Actor"), bExactMatch));
 		if (InObjectClass->IsChildOf(ActorClass))
 		{
diff --git a/Engine/Source/Runtime/Experimental/Iris/Core/Private/Iris/ReplicationSystem/Filtering/NetObjectGridFilter.cpp b/Engine/Source/Runtime/Experimental/Iris/Core/Private/Iris/ReplicationSystem/Filtering/NetObjectGridFilter.cpp
--- a/Engine/Source/Runtime/Experimental/Iris/Core/Private/Iris/ReplicationSystem/Filtering/NetObjectGridFilter.cpp
+++ b/Engine/Source/Runtime/Experimental/Iris/Core/Private/Iris/ReplicationSystem/Filtering/NetObjectGridFilter.cpp
@@ -489,8 +489,8 @@ void UNetObjectGridFragmentLocFilter::UpdateObjectInfo(UNetObjectGridFilter::FPe
 	{
 		const UE::Net::FReplicationInstanceProtocol::FFragmentData& FragmentData = FragmentDatas[ObjectLocationInfo.GetLocationStateIndex()];
 		const uint8* LocationAddress = FragmentData.ExternalSrcBuffer + ObjectLocationInfo.GetLocationStateOffset();
-		const float* LocationComponents = reinterpret_cast<const float*>(LocationAddress);
-		PerObjectInfo.Position = FVector(LocationComponents[0], LocationComponents[1], LocationComponents[2]);
+		const FVector* Location = reinterpret_cast<const FVector*>(LocationAddress);
+		PerObjectInfo.Position = *Location;
 	}
 
 	// Update the culldistance
diff --git a/Engine/Source/Runtime/Experimental/Iris/Core/Private/Iris/ReplicationSystem/RepTag.cpp b/Engine/Source/Runtime/Experimental/Iris/Core/Private/Iris/ReplicationSystem/RepTag.cpp
--- a/Engine/Source/Runtime/Experimental/Iris/Core/Private/Iris/ReplicationSystem/RepTag.cpp
+++ b/Engine/Source/Runtime/Experimental/Iris/Core/Private/Iris/ReplicationSystem/RepTag.cpp
@@ -20,6 +20,7 @@ FRepTag MakeRepTag(const char* TagName)
 
 bool HasRepTag(const FReplicationProtocol* Protocol, FRepTag RepTag)
 {
+	check(Protocol);
 	for (const FReplicationStateDescriptor* StateDescriptor : MakeArrayView(Protocol->ReplicationStateDescriptors, Protocol->ReplicationStateCount))
 	{
 		for (const FReplicationStateMemberTagDescriptor& TagDescriptor : MakeArrayView(StateDescriptor->MemberTagDescriptors, StateDescriptor->TagCount))
@@ -36,6 +37,7 @@ bool HasRepTag(const FReplicationProtocol* Protocol, FRepTag RepTag)
 
 bool FindRepTag(const FReplicationStateDescriptor* Descriptor, FRepTag RepTag, FRepTagFindInfo& OutRepTagFindInfo)
 {
+	check(Descriptor);
 	for (const FReplicationStateMemberTagDescriptor& TagDescriptor : MakeArrayView(Descriptor->MemberTagDescriptors, Descriptor->TagCount))
 	{
 		if (TagDescriptor.Tag != RepTag)
@@ -82,6 +84,7 @@ bool FindRepTag(const FReplicationStateDescriptor* Descriptor, FRepTag RepTag, F
 
 bool FindRepTag(const FReplicationProtocol* Protocol, FRepTag RepTag, FRepTagFindInfo& OutRepTagFindInfo)
 {
+	check(Protocol);
 	uint32 InternalOffset = 0;
 	for (const FReplicationStateDescriptor*& StateDescriptor : MakeArrayView(Protocol->ReplicationStateDescriptors, Protocol->ReplicationStateCount))
 	{
