Validate that assertion macros in ReplicationSystemTest will be compatible with ReplicationSystemLowLevelTests.
Fix one in TestGroups.cpp that was incompatible.

diff --git a/Engine/Plugins/Runtime/ReplicationSystemTestPlugin/Source/Private/NetworkAutomationTestMacros.h b/Engine/Plugins/Runtime/ReplicationSystemTestPlugin/Source/Private/NetworkAutomationTestMacros.h
--- a/Engine/Plugins/Runtime/ReplicationSystemTestPlugin/Source/Private/NetworkAutomationTestMacros.h
+++ b/Engine/Plugins/Runtime/ReplicationSystemTestPlugin/Source/Private/NetworkAutomationTestMacros.h
@@ -7,9 +7,7 @@
 #include "TestMessage.h"
 #include "Logging/LogVerbosity.h"
 
-#if UE_NET_WITH_LOW_LEVEL_TESTS
 #include <catch2/catch_test_macros.hpp>
-#endif
 
 namespace UE::Net
 {
@@ -187,41 +185,44 @@ UE_NET_COMPARE_RETURN_TESTRESULT(TCmpGreaterOrEqual, >=)
 #define UE_NET_ASSERT_BOOL_INTERNAL(V1, V2) UE_NET_TEST_BOOL_(V1, V2, UE_NET_FAIL)
 #define UE_NET_EXPECT_BOOL_INTERNAL(V1, V2) UE_NET_TEST_BOOL_(V1, V2, UE_NET_WARN)
 
-#define UE_NET_ASSERT_EQ(V1, V2) UE_NET_ASSERT_INTERNAL(TCmpEqual, V1, V2)
-#define UE_NET_ASSERT_NE(V1, V2) UE_NET_ASSERT_INTERNAL(TCmpNotEqual, V1, V2)
-#define UE_NET_ASSERT_LT(V1, V2) UE_NET_ASSERT_INTERNAL(TCmpLess, V1, V2)
-#define UE_NET_ASSERT_LE(V1, V2) UE_NET_ASSERT_INTERNAL(TCmpLessOrEqual, V1, V2)
-#define UE_NET_ASSERT_GT(V1, V2) UE_NET_ASSERT_INTERNAL(TCmpGreater, V1, V2)
-#define UE_NET_ASSERT_GE(V1, V2) UE_NET_ASSERT_INTERNAL(TCmpGreaterOrEqual, V1, V2)
-#define UE_NET_ASSERT_TRUE(V) UE_NET_ASSERT_BOOL_INTERNAL(V, true)
-#define UE_NET_ASSERT_FALSE(V) UE_NET_ASSERT_BOOL_INTERNAL(V, false)
-
-#define UE_NET_ASSERT_EQ_MSG(V1, V2, Message) UE_NET_ASSERT_INTERNAL_MSG(TCmpEqual, V1, V2) << Message
-#define UE_NET_ASSERT_NE_MSG(V1, V2, Message) UE_NET_ASSERT_INTERNAL_MSG(TCmpNotEqual, V1, V2) << Message
-#define UE_NET_ASSERT_LT_MSG(V1, V2, Message) UE_NET_ASSERT_INTERNAL_MSG(TCmpLess, V1, V2) << Message
-#define UE_NET_ASSERT_LE_MSG(V1, V2, Message) UE_NET_ASSERT_INTERNAL_MSG(TCmpLessOrEqual, V1, V2) << Message
-#define UE_NET_ASSERT_GT_MSG(V1, V2, Message) UE_NET_ASSERT_INTERNAL_MSG(TCmpGreater, V1, V2) << Message
-#define UE_NET_ASSERT_GE_MSG(V1, V2, Message) UE_NET_ASSERT_INTERNAL_MSG(TCmpGreaterOrEqual, V1, V2) << Message
-#define UE_NET_ASSERT_TRUE_MSG(V, Message) UE_NET_ASSERT_BOOL_INTERNAL_MSG(V, true) << Message
-#define UE_NET_ASSERT_FALSE_MSG(V, Message) UE_NET_ASSERT_BOOL_INTERNAL_MSG(V, false) << Message
-
-#define UE_NET_EXPECT_EQ(V1, V2) UE_NET_EXPECT_INTERNAL(TCmpEqual, V1, V2)
-#define UE_NET_EXPECT_NE(V1, V2) UE_NET_EXPECT_INTERNAL(TCmpNotEqual, V1, V2)
-#define UE_NET_EXPECT_LT(V1, V2) UE_NET_EXPECT_INTERNAL(TCmpLess, V1, V2)
-#define UE_NET_EXPECT_LE(V1, V2) UE_NET_EXPECT_INTERNAL(TCmpLessOrEqual, V1, V2)
-#define UE_NET_EXPECT_GT(V1, V2) UE_NET_EXPECT_INTERNAL(TCmpGreater, V1, V2)
-#define UE_NET_EXPECT_GE(V1, V2) UE_NET_EXPECT_INTERNAL(TCmpGreaterOrEqual, V1, V2)
-#define UE_NET_EXPECT_TRUE(V) UE_NET_EXPECT_BOOL_INTERNAL(V, true)
-#define UE_NET_EXPECT_FALSE(V) UE_NET_EXPECT_BOOL_INTERNAL(V, false)
-
-#define UE_NET_EXPECT_EQ_MSG(V1, V2, Message) UE_NET_EXPECT_INTERNAL_MSG(TCmpEqual, V1, V2) << Message
-#define UE_NET_EXPECT_NE_MSG(V1, V2, Message) UE_NET_EXPECT_INTERNAL_MSG(TCmpNotEqual, V1, V2) << Message
-#define UE_NET_EXPECT_LT_MSG(V1, V2, Message) UE_NET_EXPECT_INTERNAL_MSG(TCmpLess, V1, V2) << Message
-#define UE_NET_EXPECT_LE_MSG(V1, V2, Message) UE_NET_EXPECT_INTERNAL_MSG(TCmpLessOrEqual, V1, V2) << Message
-#define UE_NET_EXPECT_GT_MSG(V1, V2, Message) UE_NET_EXPECT_INTERNAL_MSG(TCmpGreater, V1, V2) << Message
-#define UE_NET_EXPECT_GE_MSG(V1, V2, Message) UE_NET_EXPECT_INTERNAL_MSG(TCmpGreaterOrEqual, V1, V2) << Message
-#define UE_NET_EXPECT_TRUE_MSG(V, Message) UE_NET_EXPECT_BOOL_INTERNAL_MSG(V, true) << Message
-#define UE_NET_EXPECT_FALSE_MSG(V, Message) UE_NET_EXPECT_BOOL_INTERNAL_MSG(V, false) << Message
+// Causes a build error if the tested expression isn't compatible with Catch/low-level tests, without evaluating the expression at runtime.
+#define UE_NET_VALIDATE_CATCH_EXPR(V) (false ? (void)(Catch::Decomposer() <= V) : (void)0)
+
+#define UE_NET_ASSERT_EQ(V1, V2) UE_NET_VALIDATE_CATCH_EXPR(V1 == V2); UE_NET_ASSERT_INTERNAL(TCmpEqual, V1, V2)
+#define UE_NET_ASSERT_NE(V1, V2) UE_NET_VALIDATE_CATCH_EXPR(V1 != V2); UE_NET_ASSERT_INTERNAL(TCmpNotEqual, V1, V2)
+#define UE_NET_ASSERT_LT(V1, V2) UE_NET_VALIDATE_CATCH_EXPR(V1 < V2); UE_NET_ASSERT_INTERNAL(TCmpLess, V1, V2)
+#define UE_NET_ASSERT_LE(V1, V2) UE_NET_VALIDATE_CATCH_EXPR(V1 <= V2); UE_NET_ASSERT_INTERNAL(TCmpLessOrEqual, V1, V2)
+#define UE_NET_ASSERT_GT(V1, V2) UE_NET_VALIDATE_CATCH_EXPR(V1 > V2); UE_NET_ASSERT_INTERNAL(TCmpGreater, V1, V2)
+#define UE_NET_ASSERT_GE(V1, V2) UE_NET_VALIDATE_CATCH_EXPR(V1 >= V2); UE_NET_ASSERT_INTERNAL(TCmpGreaterOrEqual, V1, V2)
+#define UE_NET_ASSERT_TRUE(V) UE_NET_VALIDATE_CATCH_EXPR(V); UE_NET_ASSERT_BOOL_INTERNAL(V, true)
+#define UE_NET_ASSERT_FALSE(V) UE_NET_VALIDATE_CATCH_EXPR(V); UE_NET_ASSERT_BOOL_INTERNAL(V, false)
+
+#define UE_NET_ASSERT_EQ_MSG(V1, V2, Message) UE_NET_VALIDATE_CATCH_EXPR(V1 == V2); UE_NET_ASSERT_INTERNAL_MSG(TCmpEqual, V1, V2) << Message
+#define UE_NET_ASSERT_NE_MSG(V1, V2, Message) UE_NET_VALIDATE_CATCH_EXPR(V1 != V2); UE_NET_ASSERT_INTERNAL_MSG(TCmpNotEqual, V1, V2) << Message
+#define UE_NET_ASSERT_LT_MSG(V1, V2, Message) UE_NET_VALIDATE_CATCH_EXPR(V1 < V2); UE_NET_ASSERT_INTERNAL_MSG(TCmpLess, V1, V2) << Message
+#define UE_NET_ASSERT_LE_MSG(V1, V2, Message) UE_NET_VALIDATE_CATCH_EXPR(V1 <= V2); UE_NET_ASSERT_INTERNAL_MSG(TCmpLessOrEqual, V1, V2) << Message
+#define UE_NET_ASSERT_GT_MSG(V1, V2, Message) UE_NET_VALIDATE_CATCH_EXPR(V1 > V2); UE_NET_ASSERT_INTERNAL_MSG(TCmpGreater, V1, V2) << Message
+#define UE_NET_ASSERT_GE_MSG(V1, V2, Message) UE_NET_VALIDATE_CATCH_EXPR(V1 >= V2); UE_NET_ASSERT_INTERNAL_MSG(TCmpGreaterOrEqual, V1, V2) << Message
+#define UE_NET_ASSERT_TRUE_MSG(V, Message) UE_NET_VALIDATE_CATCH_EXPR(V); UE_NET_ASSERT_BOOL_INTERNAL_MSG(V, true) << Message
+#define UE_NET_ASSERT_FALSE_MSG(V, Message) UE_NET_VALIDATE_CATCH_EXPR(V); UE_NET_ASSERT_BOOL_INTERNAL_MSG(V, false) << Message
+
+#define UE_NET_EXPECT_EQ(V1, V2) UE_NET_VALIDATE_CATCH_EXPR(V1 == V2); UE_NET_EXPECT_INTERNAL(TCmpEqual, V1, V2)
+#define UE_NET_EXPECT_NE(V1, V2) UE_NET_VALIDATE_CATCH_EXPR(V1 != V2); UE_NET_EXPECT_INTERNAL(TCmpNotEqual, V1, V2)
+#define UE_NET_EXPECT_LT(V1, V2) UE_NET_VALIDATE_CATCH_EXPR(V1 < V2); UE_NET_EXPECT_INTERNAL(TCmpLess, V1, V2)
+#define UE_NET_EXPECT_LE(V1, V2) UE_NET_VALIDATE_CATCH_EXPR(V1 <= V2); UE_NET_EXPECT_INTERNAL(TCmpLessOrEqual, V1, V2)
+#define UE_NET_EXPECT_GT(V1, V2) UE_NET_VALIDATE_CATCH_EXPR(V1 > V2); UE_NET_EXPECT_INTERNAL(TCmpGreater, V1, V2)
+#define UE_NET_EXPECT_GE(V1, V2) UE_NET_VALIDATE_CATCH_EXPR(V1 >= V2); UE_NET_EXPECT_INTERNAL(TCmpGreaterOrEqual, V1, V2)
+#define UE_NET_EXPECT_TRUE(V) UE_NET_VALIDATE_CATCH_EXPR(V); UE_NET_EXPECT_BOOL_INTERNAL(V, true)
+#define UE_NET_EXPECT_FALSE(V) UE_NET_VALIDATE_CATCH_EXPR(V); UE_NET_EXPECT_BOOL_INTERNAL(V, false)
+
+#define UE_NET_EXPECT_EQ_MSG(V1, V2, Message) UE_NET_VALIDATE_CATCH_EXPR(V1 == V2); UE_NET_EXPECT_INTERNAL_MSG(TCmpEqual, V1, V2) << Message
+#define UE_NET_EXPECT_NE_MSG(V1, V2, Message) UE_NET_VALIDATE_CATCH_EXPR(V1 != V2); UE_NET_EXPECT_INTERNAL_MSG(TCmpNotEqual, V1, V2) << Message
+#define UE_NET_EXPECT_LT_MSG(V1, V2, Message) UE_NET_VALIDATE_CATCH_EXPR(V1 < V2); UE_NET_EXPECT_INTERNAL_MSG(TCmpLess, V1, V2) << Message
+#define UE_NET_EXPECT_LE_MSG(V1, V2, Message) UE_NET_VALIDATE_CATCH_EXPR(V1 <= V2); UE_NET_EXPECT_INTERNAL_MSG(TCmpLessOrEqual, V1, V2) << Message
+#define UE_NET_EXPECT_GT_MSG(V1, V2, Message) UE_NET_VALIDATE_CATCH_EXPR(V1 > V2); UE_NET_EXPECT_INTERNAL_MSG(TCmpGreater, V1, V2) << Message
+#define UE_NET_EXPECT_GE_MSG(V1, V2, Message) UE_NET_VALIDATE_CATCH_EXPR(V1 >= V2); UE_NET_EXPECT_INTERNAL_MSG(TCmpGreaterOrEqual, V1, V2) << Message
+#define UE_NET_EXPECT_TRUE_MSG(V, Message) UE_NET_VALIDATE_CATCH_EXPR(V); UE_NET_EXPECT_BOOL_INTERNAL_MSG(V, true) << Message
+#define UE_NET_EXPECT_FALSE_MSG(V, Message) UE_NET_VALIDATE_CATCH_EXPR(V); UE_NET_EXPECT_BOOL_INTERNAL_MSG(V, false) << Message
 
 #endif
 
diff --git a/Engine/Plugins/Runtime/ReplicationSystemTestPlugin/Source/Private/Tests/ReplicationSystem/Filtering/TestGroups.cpp b/Engine/Plugins/Runtime/ReplicationSystemTestPlugin/Source/Private/Tests/ReplicationSystem/Filtering/TestGroups.cpp
--- a/Engine/Plugins/Runtime/ReplicationSystemTestPlugin/Source/Private/Tests/ReplicationSystem/Filtering/TestGroups.cpp
+++ b/Engine/Plugins/Runtime/ReplicationSystemTestPlugin/Source/Private/Tests/ReplicationSystem/Filtering/TestGroups.cpp
@@ -70,7 +70,7 @@ protected:
 
 		// Look if the client received the object as expected
 		UObject* ClientObject = Client->GetReplicationBridge()->GetReplicatedObject(Handle);
-		UE_NET_ASSERT_EQ(ClientObject!=nullptr, bExpectedReplicationAllowed);
+		UE_NET_ASSERT_EQ((ClientObject!=nullptr), bExpectedReplicationAllowed);
 	}
 };
 
diff --git a/Engine/Plugins/Runtime/ReplicationSystemTestPlugin/Source/ReplicationSystemTestPlugin.Build.cs b/Engine/Plugins/Runtime/ReplicationSystemTestPlugin/Source/ReplicationSystemTestPlugin.Build.cs
--- a/Engine/Plugins/Runtime/ReplicationSystemTestPlugin/Source/ReplicationSystemTestPlugin.Build.cs
+++ b/Engine/Plugins/Runtime/ReplicationSystemTestPlugin/Source/ReplicationSystemTestPlugin.Build.cs
@@ -45,6 +45,11 @@ namespace UnrealBuildTool.Rules
 					}
 					);
 			}
+			else
+			{
+				// For validation of compatibility with low-level tests even when not running them
+				PrivateDependencyModuleNames.Add("Catch2");
+			}
 
 			PrivateDefinitions.Add(String.Format("UE_NET_WITH_LOW_LEVEL_TESTS={0}", Target.ExplicitTestsTarget ? "1" : "0"));
 
