Iris
* Added unit tests for RPC sending

#rnx

diff --git a/Engine/Plugins/Runtime/ReplicationSystemTestPlugin/Source/Private/Tests/ReplicationSystem/RPC/RPCTestFixture.cpp b/Engine/Plugins/Runtime/ReplicationSystemTestPlugin/Source/Private/Tests/ReplicationSystem/RPC/RPCTestFixture.cpp
new file mode 100644
--- /dev/null
+++ b/Engine/Plugins/Runtime/ReplicationSystemTestPlugin/Source/Private/Tests/ReplicationSystem/RPC/RPCTestFixture.cpp
@@ -0,0 +1,28 @@
+// Copyright Epic Games, Inc. All Rights Reserved.
+
+#include "RPCTestFixture.h"
+
+namespace UE::Net
+{
+
+void FRPCTestFixture::SetUp()
+{
+	UNetBlobHandlerDefinitions* BlobHandlerDefinitions = GetMutableDefault<UNetBlobHandlerDefinitions>();
+	OriginalHandlerDefinitions = BlobHandlerDefinitions->ReadWriteHandlerDefinitions();
+
+	HandlerDefinitions.Add({ TEXT("NetRPCHandler") });
+
+	BlobHandlerDefinitions->ReadWriteHandlerDefinitions() = HandlerDefinitions;
+
+	Super::SetUp();
+}
+
+void FRPCTestFixture::TearDown()
+{
+	Super::TearDown();
+
+	UNetBlobHandlerDefinitions* BlobHandlerDefinitions = GetMutableDefault<UNetBlobHandlerDefinitions>();
+	BlobHandlerDefinitions->ReadWriteHandlerDefinitions() = OriginalHandlerDefinitions;
+}
+
+}
diff --git a/Engine/Plugins/Runtime/ReplicationSystemTestPlugin/Source/Private/Tests/ReplicationSystem/RPC/RPCTestFixture.h b/Engine/Plugins/Runtime/ReplicationSystemTestPlugin/Source/Private/Tests/ReplicationSystem/RPC/RPCTestFixture.h
new file mode 100644
--- /dev/null
+++ b/Engine/Plugins/Runtime/ReplicationSystemTestPlugin/Source/Private/Tests/ReplicationSystem/RPC/RPCTestFixture.h
@@ -0,0 +1,33 @@
+// Copyright Epic Games, Inc. All Rights Reserved.
+
+#pragma once
+
+#include "Tests/ReplicationSystem/ReplicationSystemServerClientTestFixture.h"
+#include "Iris/ReplicationSystem/ReplicationSystem.h"
+#include "Iris/ReplicationSystem/ReplicationSystemInternal.h"
+#include "Iris/ReplicationSystem/NetBlob/NetBlobHandlerDefinitions.h"
+#include "Iris/Serialization/InternalNetSerializationContext.h"
+
+namespace UE::Net
+{
+
+class FRPCTestFixture : public FReplicationSystemServerClientTestFixture
+{
+	typedef FReplicationSystemServerClientTestFixture Super;
+
+public:
+	FRPCTestFixture()
+	{
+	}
+
+protected:
+	virtual void SetUp() override;
+
+	virtual void TearDown() override;
+
+private:
+	TArray<FNetBlobHandlerDefinition> OriginalHandlerDefinitions;
+	TArray<FNetBlobHandlerDefinition> HandlerDefinitions;
+};
+
+}
diff --git a/Engine/Plugins/Runtime/ReplicationSystemTestPlugin/Source/Private/Tests/ReplicationSystem/RPC/ReplicatedTestObjectWithRPC.cpp b/Engine/Plugins/Runtime/ReplicationSystemTestPlugin/Source/Private/Tests/ReplicationSystem/RPC/ReplicatedTestObjectWithRPC.cpp
new file mode 100644
--- /dev/null
+++ b/Engine/Plugins/Runtime/ReplicationSystemTestPlugin/Source/Private/Tests/ReplicationSystem/RPC/ReplicatedTestObjectWithRPC.cpp
@@ -0,0 +1,110 @@
+// Copyright Epic Games, Inc. All Rights Reserved.
+
+#include "ReplicatedTestObjectWithRPC.h"
+
+#include "Iris/ReplicationSystem/NetRefHandle.h"
+#include "Iris/ReplicationSystem/ReplicationProtocol.h"
+#include "Iris/ReplicationSystem/ReplicationSystem.h"
+#include "Iris/ReplicationSystem/ReplicationSystemInternal.h"
+#include "Iris/ReplicationSystem/ReplicationFragmentUtil.h"
+
+#include "Net/UnrealNetwork.h"
+
+//////////////////////////////////////////////////////////////////////////
+// Implementation for UTestReplicatedObjectWithRPC
+//////////////////////////////////////////////////////////////////////////
+UTestReplicatedObjectWithRPC::UTestReplicatedObjectWithRPC()
+	: UReplicatedTestObject()
+{
+}
+
+void UTestReplicatedObjectWithRPC::RegisterReplicationFragments(UE::Net::FFragmentRegistrationContext& Context, UE::Net::EFragmentRegistrationFlags RegistrationFlags)
+{
+	this->ReplicationFragments.Reset();
+	UE::Net::FReplicationFragmentUtil::CreateAndRegisterFragmentsForObject(this, Context, RegistrationFlags, &this->ReplicationFragments);
+}
+
+int32 UTestReplicatedObjectWithRPC::GetFunctionCallspace(UFunction* Function, FFrame* Stack)
+{
+	check(!(Function->FunctionFlags & FUNC_Static));
+	check(Function->FunctionFlags & FUNC_Net);
+
+	const bool bIsOnServer = bIsServerObject;
+
+	// get the top most function
+	while (Function->GetSuperFunction() != nullptr)
+	{
+		Function = Function->GetSuperFunction();
+	}
+
+	// Multicast RPCs
+	if ((Function->FunctionFlags & FUNC_NetMulticast))
+	{
+		if (bIsOnServer)
+		{
+			// Server should execute locally and call remotely
+			return (FunctionCallspace::Local | FunctionCallspace::Remote);
+		}
+		else
+		{
+			return FunctionCallspace::Local;
+		}
+	}
+
+	// if we are the authority
+	if (bIsOnServer)
+	{
+		if (Function->FunctionFlags & FUNC_NetClient)
+		{
+			return FunctionCallspace::Remote;
+		}
+		else
+		{
+			return FunctionCallspace::Local;
+		}
+
+	}
+	// if we are not the authority
+	else
+	{
+		if (Function->FunctionFlags & FUNC_NetServer)
+		{
+			return FunctionCallspace::Remote;
+		}
+		else
+		{
+			// don't replicate
+			return FunctionCallspace::Local;
+		}
+	}
+
+	// Call remotely
+	return FunctionCallspace::Remote;
+}
+
+bool UTestReplicatedObjectWithRPC::CallRemoteFunction(UFunction* Function, void* Parameters, FOutParmRec* OutParms, FFrame* Stack)
+{
+	return ReplicationSystem->SendRPC(this, nullptr, Function, Parameters);
+}
+
+void UTestReplicatedObjectWithRPC::ClientRPC_Implementation()
+{
+	bClientRPCCalled = true;
+}
+
+void UTestReplicatedObjectWithRPC::ClientRPCWithParam_Implementation(int32 IntParam)
+{
+	ClientRPCWithParamCalled = IntParam;
+}
+
+void UTestReplicatedObjectWithRPC::ServerRPC_Implementation()
+{
+	bServerRPCCalled = true;
+}
+
+void UTestReplicatedObjectWithRPC::ServerRPCWithParam_Implementation(int32 IntParam)
+{
+	ServerRPCWithParamCalled = IntParam;
+}
+
+
diff --git a/Engine/Plugins/Runtime/ReplicationSystemTestPlugin/Source/Private/Tests/ReplicationSystem/RPC/ReplicatedTestObjectWithRPC.h b/Engine/Plugins/Runtime/ReplicationSystemTestPlugin/Source/Private/Tests/ReplicationSystem/RPC/ReplicatedTestObjectWithRPC.h
new file mode 100644
--- /dev/null
+++ b/Engine/Plugins/Runtime/ReplicationSystemTestPlugin/Source/Private/Tests/ReplicationSystem/RPC/ReplicatedTestObjectWithRPC.h
@@ -0,0 +1,54 @@
+// Copyright Epic Games, Inc. All Rights Reserved.
+
+#pragma once
+
+#include "Tests/ReplicationSystem/ReplicatedTestObject.h"
+
+#include "ReplicatedTestObjectWithRPC.generated.h"
+
+/**
+ *  A test class for testing RPC replication
+ */
+UCLASS()
+class UTestReplicatedObjectWithRPC : public UReplicatedTestObject
+{
+	GENERATED_BODY()
+
+public:
+	UTestReplicatedObjectWithRPC();
+
+	virtual void RegisterReplicationFragments(UE::Net::FFragmentRegistrationContext& Context, UE::Net::EFragmentRegistrationFlags RegistrationFlags) override;
+
+	virtual int32 GetFunctionCallspace(UFunction* Function, FFrame* Stack) override;
+	virtual bool CallRemoteFunction(UFunction* Function, void* Parameters, FOutParmRec* OutParms, FFrame* Stack) override;
+
+
+	UReplicationSystem* ReplicationSystem = nullptr;
+
+	// Network data only for test
+	TArray<UE::Net::FReplicationFragment*> ReplicationFragments;
+
+	bool bIsServerObject = false;
+
+	//$IRIS TODO: Temp until SetUpRuntimeReplicationData is fixed
+	UPROPERTY()
+	bool bTest;
+
+	// RPC test functions
+	UFUNCTION(Reliable, Client)
+	void ClientRPC();
+	bool bClientRPCCalled = false;
+
+	UFUNCTION(Reliable, Client)
+	void ClientRPCWithParam(int32 IntParam);
+	int32 ClientRPCWithParamCalled = 0;
+
+	UFUNCTION(Reliable, Server)
+	void ServerRPC();
+	bool bServerRPCCalled = false;
+
+	UFUNCTION(Reliable, Server)
+	void ServerRPCWithParam(int32 IntParam);
+	int32 ServerRPCWithParamCalled = 0;
+};
+
diff --git a/Engine/Plugins/Runtime/ReplicationSystemTestPlugin/Source/Private/Tests/ReplicationSystem/RPC/TestObjectRPC.cpp b/Engine/Plugins/Runtime/ReplicationSystemTestPlugin/Source/Private/Tests/ReplicationSystem/RPC/TestObjectRPC.cpp
new file mode 100644
--- /dev/null
+++ b/Engine/Plugins/Runtime/ReplicationSystemTestPlugin/Source/Private/Tests/ReplicationSystem/RPC/TestObjectRPC.cpp
@@ -0,0 +1,85 @@
+// Copyright Epic Games, Inc. All Rights Reserved.
+
+#include "RPCTestFixture.h"
+#include "ReplicatedTestObjectWithRPC.h"
+#include "Iris/Serialization/NetBitStreamReader.h"
+#include "Iris/Serialization/NetBitStreamWriter.h"
+#include "Iris/ReplicationSystem/ReplicationSystem.h"
+
+#include "Iris/ReplicationSystem/ReplicationSystemInternal.h"
+
+namespace UE::Net::Private
+{
+
+	UE_NET_TEST_FIXTURE(FRPCTestFixture, TestBasicObjectRPC)
+	{
+		// Add a client
+		FReplicationSystemTestClient* Client = CreateClient();
+
+		// Spawn object on server
+		UTestReplicatedObjectWithRPC* ServerObject = Server->CreateObject<UTestReplicatedObjectWithRPC>();
+
+		ServerObject->bIsServerObject = true;
+		ServerObject->ReplicationSystem = Server->GetReplicationSystem();
+		Server->ReplicationSystem->SetOwningNetConnection(ServerObject->NetRefHandle, 0x01);
+
+		// Send and deliver packet
+		Server->PreSendUpdate();
+		Server->SendAndDeliverTo(Client, true);
+		Server->PostSendUpdate();
+
+		UTestReplicatedObjectWithRPC* ClientObject = Cast<UTestReplicatedObjectWithRPC>(Client->GetReplicationBridge()->GetReplicatedObject(ServerObject->NetRefHandle));
+		
+		// Verify that created server handle now also exists on client
+		UE_NET_ASSERT_TRUE(ClientObject != nullptr);
+
+		ClientObject->ReplicationSystem = Client->GetReplicationSystem();
+
+		// Call an RPC server->client
+		ServerObject->ClientRPC();
+
+		// Send and deliver packet
+		Server->PreSendUpdate();
+		Server->SendAndDeliverTo(Client, true);
+		Server->PostSendUpdate();
+
+		// Verify RPC reception
+		UE_NET_ASSERT_TRUE(ClientObject->bClientRPCCalled);
+
+		// Call an RPC server->client
+		const int32 IntParam = 0xBABA;
+		ServerObject->ClientRPCWithParam(0xBABA);
+
+		// Send and deliver packet
+		Server->PreSendUpdate();
+		Server->SendAndDeliverTo(Client, true);
+		Server->PostSendUpdate();
+
+		// Verify RPC reception
+		UE_NET_ASSERT_TRUE(ClientObject->ClientRPCWithParamCalled == IntParam);
+
+		// Call an RPC client->server
+		ClientObject->ServerRPC();
+
+		// Send and deliver client packet
+		Client->PreSendUpdate();
+		Client->SendUpdate();
+		Client->DeliverTo(*Server, 0x01, 0x01, true);
+		Client->PostSendUpdate();
+
+		// Verify RPC reception
+		UE_NET_ASSERT_TRUE(ServerObject->bServerRPCCalled);
+
+		// Call an RPC client->server
+		ClientObject->ServerRPCWithParam(IntParam);
+
+		// Send and deliver client packet
+		Client->PreSendUpdate();
+		Client->SendUpdate();
+		Client->DeliverTo(*Server, 0x01, 0x01, true);
+		Client->PostSendUpdate();
+
+		// Verify RPC reception
+		UE_NET_ASSERT_TRUE(ServerObject->ServerRPCWithParamCalled == IntParam);
+	}
+}
\ No newline at end of file
diff --git a/Engine/Plugins/Runtime/ReplicationSystemTestPlugin/Source/Private/Tests/ReplicationSystem/ReplicatedTestObject.cpp b/Engine/Plugins/Runtime/ReplicationSystemTestPlugin/Source/Private/Tests/ReplicationSystem/ReplicatedTestObject.cpp
--- a/Engine/Plugins/Runtime/ReplicationSystemTestPlugin/Source/Private/Tests/ReplicationSystem/ReplicatedTestObject.cpp
+++ b/Engine/Plugins/Runtime/ReplicationSystemTestPlugin/Source/Private/Tests/ReplicationSystem/ReplicatedTestObject.cpp
@@ -467,6 +467,9 @@ void UTestReplicatedIrisObjectWithObjectReference::RegisterReplicationFragments(
 
 uint32 UReplicatedSubObjectOrderObject::RepOrderCounter = 0U;
 
+//////////////////////////////////////////////////////////////////////////
+// Implementation for UTestReplicatedIrisObjectWithNoReplicatedMembers
+//////////////////////////////////////////////////////////////////////////
 UTestReplicatedIrisObjectWithNoReplicatedMembers::UTestReplicatedIrisObjectWithNoReplicatedMembers()
 : UReplicatedTestObject()
 {
@@ -477,6 +480,9 @@ void UTestReplicatedIrisObjectWithNoReplicatedMembers::RegisterReplicationFragme
 	Fragments.SetIsFragmentlessNetObject(true);
 }
 
+//////////////////////////////////////////////////////////////////////////
+// Implementation for UReplicatedSubObjectOrderObject
+//////////////////////////////////////////////////////////////////////////
 void UReplicatedSubObjectOrderObject::GetLifetimeReplicatedProps(TArray< class FLifetimeProperty >& OutLifetimeProps) const
 {
 	DOREPLIFETIME(UReplicatedSubObjectOrderObject, IntA);
diff --git a/Engine/Plugins/Runtime/ReplicationSystemTestPlugin/Source/Private/Tests/ReplicationSystem/ReplicatedTestObject.h b/Engine/Plugins/Runtime/ReplicationSystemTestPlugin/Source/Private/Tests/ReplicationSystem/ReplicatedTestObject.h
--- a/Engine/Plugins/Runtime/ReplicationSystemTestPlugin/Source/Private/Tests/ReplicationSystem/ReplicatedTestObject.h
+++ b/Engine/Plugins/Runtime/ReplicationSystemTestPlugin/Source/Private/Tests/ReplicationSystem/ReplicatedTestObject.h
@@ -411,7 +411,7 @@ public:
 
 
 /**
- *  A test class for Replication that itself uses Property based replication but also has "components" that uses a mix of property based replication and native ReplicationStates
+ *  A test class for Replication on an object with no replicated members
  */
 UCLASS()
 class UTestReplicatedIrisObjectWithNoReplicatedMembers : public UReplicatedTestObject
@@ -421,7 +421,7 @@ class UTestReplicatedIrisObjectWithNoReplicatedMembers : public UReplicatedTestO
 public:
 	UTestReplicatedIrisObjectWithNoReplicatedMembers();
 
-	virtual void RegisterReplicationFragments(UE::Net::FFragmentRegistrationContext& Fragments, UE::Net::EFragmentRegistrationFlags RegistrationFlags) override;
+	virtual void RegisterReplicationFragments(UE::Net::FFragmentRegistrationContext& Context, UE::Net::EFragmentRegistrationFlags RegistrationFlags) override;
 };
 
 /**
