ReplicationSystemTest
* Added unit test for delaying dirty objects when culled.

#rnx

diff --git a/Engine/Plugins/Runtime/ReplicationSystemTestPlugin/Source/Private/Tests/ReplicationSystem/Filtering/MockNetObjectFilter.cpp b/Engine/Plugins/Runtime/ReplicationSystemTestPlugin/Source/Private/Tests/ReplicationSystem/Filtering/MockNetObjectFilter.cpp
--- a/Engine/Plugins/Runtime/ReplicationSystemTestPlugin/Source/Private/Tests/ReplicationSystem/Filtering/MockNetObjectFilter.cpp
+++ b/Engine/Plugins/Runtime/ReplicationSystemTestPlugin/Source/Private/Tests/ReplicationSystem/Filtering/MockNetObjectFilter.cpp
@@ -2,6 +2,10 @@
 
 #include "MockNetObjectFilter.h"
 #include "Containers/ArrayView.h"
+#include "TestFilteringObject.h"
+
+#include "Iris/ReplicationSystem/ReplicationSystem.h"
+#include "Iris/ReplicationSystem/ReplicationSystemInternal.h"
 
 UMockNetObjectFilter::UMockNetObjectFilter()
 : CallStatus({})
@@ -204,3 +208,46 @@ void UMockNetObjectFilterUsingFragmentData::Filter(FNetObjectFilteringParams& Pa
 		Params.OutAllowedObjects.SetBitValue(Pair.Key, !Pair.Value);
 	}
 }
+
+
+//**************************************************************************************************
+// UMockNetObjectFilterWithCondition
+//**************************************************************************************************
+
+void UMockNetObjectFilterWithCondition::OnInit(FNetObjectFilterInitParams& Params)
+{
+	Super::OnInit(Params);
+
+	SetupFilterType(ENetFilterType::PrePoll_Raw);
+
+	ReplicationSystem = Params.ReplicationSystem.Get();
+}
+
+bool UMockNetObjectFilterWithCondition::AddObject(uint32 ObjectIndex, FNetObjectFilterAddObjectParams& Params)
+{
+	bool bResult = Super::AddObject(ObjectIndex, Params);
+
+	// Make sure the object is of the proper class
+	UObject* ReplicatedObject = ReplicationSystem->GetReplicationSystemInternal()->GetNetRefHandleManager().GetReplicatedObjectInstance(ObjectIndex);
+	UTestFilteringObject* FilterObject = CastChecked<UTestFilteringObject>(ReplicatedObject);
+
+	return bResult;
+}
+
+void UMockNetObjectFilterWithCondition::Filter(FNetObjectFilteringParams& Params)
+{
+	++CallStatus.CallCounts.Filter;
+
+	bool bIsProperCall = true;
+
+	Params.FilteredObjects.ForAllSetBits([&](uint32 ObjectIndex)
+	{
+		UTestFilteringObject* FilterObject = CastChecked<UTestFilteringObject>(ReplicationSystem->GetReplicationSystemInternal()->GetNetRefHandleManager().GetReplicatedObjectInstance(ObjectIndex));
+		
+		const bool bIsFilteredOut = FilterObject->GetFilterOut();
+
+		Params.OutAllowedObjects.SetBitValue(ObjectIndex, !bIsFilteredOut);
+	});
+	
+	CallStatus.SuccessfulCallCounts.Filter += bIsProperCall;
+}
\ No newline at end of file
diff --git a/Engine/Plugins/Runtime/ReplicationSystemTestPlugin/Source/Private/Tests/ReplicationSystem/Filtering/MockNetObjectFilter.h b/Engine/Plugins/Runtime/ReplicationSystemTestPlugin/Source/Private/Tests/ReplicationSystem/Filtering/MockNetObjectFilter.h
--- a/Engine/Plugins/Runtime/ReplicationSystemTestPlugin/Source/Private/Tests/ReplicationSystem/Filtering/MockNetObjectFilter.h
+++ b/Engine/Plugins/Runtime/ReplicationSystemTestPlugin/Source/Private/Tests/ReplicationSystem/Filtering/MockNetObjectFilter.h
@@ -108,4 +108,23 @@ private:
 
 	TMap<uint32, UPTRINT> ObjectToFilterOutOffset;
 	TMap<uint32, bool> ObjectToFilterOut;
+};
+
+/**
+ * Filter that checks object data to decide to filter out an object.
+ */
+UCLASS()
+class UMockNetObjectFilterWithCondition : public UMockNetObjectFilter
+{
+	GENERATED_BODY()
+
+protected:
+
+	virtual void OnInit(FNetObjectFilterInitParams&) override;
+	virtual bool AddObject(uint32 ObjectIndex, FNetObjectFilterAddObjectParams&) override;
+	virtual void Filter(FNetObjectFilteringParams&) override;
+
+private:
+
+	const UReplicationSystem* ReplicationSystem = nullptr;
 };
\ No newline at end of file
diff --git a/Engine/Plugins/Runtime/ReplicationSystemTestPlugin/Source/Private/Tests/ReplicationSystem/Filtering/TestFiltering.cpp b/Engine/Plugins/Runtime/ReplicationSystemTestPlugin/Source/Private/Tests/ReplicationSystem/Filtering/TestFiltering.cpp
--- a/Engine/Plugins/Runtime/ReplicationSystemTestPlugin/Source/Private/Tests/ReplicationSystem/Filtering/TestFiltering.cpp
+++ b/Engine/Plugins/Runtime/ReplicationSystemTestPlugin/Source/Private/Tests/ReplicationSystem/Filtering/TestFiltering.cpp
@@ -28,6 +28,9 @@ protected:
 		RestoreFilterDefinitions();
 	}
 
+	virtual FName GetMockFilterName() const { return FName("MockFilter"); }
+	virtual FName GetMockFilterClassName() const { return FName("/Script/ReplicationSystemTestPlugin.MockNetObjectFilter"); }
+
 private:
 	void InitNetObjectFilterDefinitions()
 	{
@@ -43,8 +46,8 @@ private:
 		TArray<FNetObjectFilterDefinition> NewFilterDefinitions;
 		{
 			FNetObjectFilterDefinition& MockDefinition = NewFilterDefinitions.Emplace_GetRef();
-			MockDefinition.FilterName = "MockFilter";
-			MockDefinition.ClassName = "/Script/ReplicationSystemTestPlugin.MockNetObjectFilter";
+			MockDefinition.FilterName = GetMockFilterName();
+			MockDefinition.ClassName = GetMockFilterClassName();
 			MockDefinition.ConfigClassName = "/Script/ReplicationSystemTestPlugin.MockNetObjectFilterConfig";
 		}
 
@@ -76,8 +79,8 @@ private:
 
 	void InitMockNetObjectFilter()
 	{
-		MockNetObjectFilter = ExactCast<UMockNetObjectFilter>(Server->GetReplicationSystem()->GetFilter("MockFilter"));
-		MockFilterHandle = Server->GetReplicationSystem()->GetFilterHandle("MockFilter");
+		MockNetObjectFilter = CastChecked<UMockNetObjectFilter>(Server->GetReplicationSystem()->GetFilter(GetMockFilterName()));
+		MockFilterHandle = Server->GetReplicationSystem()->GetFilterHandle(GetMockFilterName());
 
 		MockNetObjectFilterWithFragments = ExactCast<UMockNetObjectFilterUsingFragmentData>(Server->GetReplicationSystem()->GetFilter("MockFilterWithFragments"));
 		MockFilterWithFragmentsHandle = Server->GetReplicationSystem()->GetFilterHandle("MockFilterWithFragments");
@@ -94,6 +97,15 @@ private:
 	TArray<FNetObjectFilterDefinition> OriginalFilterDefinitions;
 };
 
+
+class FTestFilteringWithConditionFixture : public FTestFilteringFixture
+{
+protected:
+
+	virtual FName GetMockFilterName() const { return FName("MockFilterWithCondition"); }
+	virtual FName GetMockFilterClassName() const { return FName("/Script/ReplicationSystemTestPlugin.MockNetObjectFilterWithCondition"); }
+};
+
 // Owner filtering tests
 UE_NET_TEST_FIXTURE(FTestFilteringFixture, OwnerFilterPreventsObjectFromReplicatingToNonOwner)
 {
@@ -1312,4 +1324,120 @@ UE_NET_TEST_FIXTURE(FTestFilteringFixture, DynamicFilteredOutSubObjectsAreResetW
 }
 
 
+UE_NET_TEST_FIXTURE(FTestFilteringWithConditionFixture, TestCulledDirtyActors)
+{
+	// Setup dynamic filter for the test.
+	{
+		UMockNetObjectFilter::FFunctionCallSetup CallSetup;
+		CallSetup.AddObject.bReturnValue = true;
+		CallSetup.Filter.bFilterOutByDefault = true;
+
+		MockNetObjectFilter->SetFunctionCallSetup(CallSetup);
+		MockNetObjectFilter->ResetFunctionCallStatus();
+	}
+
+	const uint32 RepSystemID = Server->GetReplicationSystemId();
+
+	// Add client
+	FReplicationSystemTestClient* Client = CreateClient();
+
+	// Create multiple filtered objects
+	UTestFilteringObject* ServerObjectA = Server->CreateObject<UTestFilteringObject>();
+	Server->ReplicationSystem->SetFilter(ServerObjectA->NetRefHandle, MockFilterHandle);
+
+	UTestFilteringObject* ServerObjectB = Server->CreateObject<UTestFilteringObject>();
+	Server->ReplicationSystem->SetFilter(ServerObjectB->NetRefHandle, MockFilterHandle);
+
+	UTestFilteringObject* ServerObjectC = Server->CreateObject<UTestFilteringObject>();
+	Server->ReplicationSystem->SetFilter(ServerObjectC->NetRefHandle, MockFilterHandle);
+
+	// Create a non-filtered object
+	UTestFilteringObject* ServerObjectNoFilter = Server->CreateObject<UTestFilteringObject>();
+
+	// Filter them in
+	{	
+		constexpr bool bFilterIn = false;
+		ServerObjectA->SetFilterOut(bFilterIn);
+		ServerObjectB->SetFilterOut(bFilterIn);
+		ServerObjectC->SetFilterOut(bFilterIn);
+
+		// Send and deliver packets
+		Server->UpdateAndSend({Client});
+
+		// Check that the filtered object do exist on the client.
+		UE_NET_ASSERT_NE(Client->GetReplicationBridge()->GetReplicatedObject(ServerObjectA->NetRefHandle), nullptr);
+		UE_NET_ASSERT_NE(Client->GetReplicationBridge()->GetReplicatedObject(ServerObjectB->NetRefHandle), nullptr);
+		UE_NET_ASSERT_NE(Client->GetReplicationBridge()->GetReplicatedObject(ServerObjectC->NetRefHandle), nullptr);
+	}
+
+	// Now filter them out
+	{
+		constexpr bool bFilterOut = true;
+		ServerObjectA->SetFilterOut(bFilterOut);
+		ServerObjectB->SetFilterOut(bFilterOut);
+		ServerObjectC->SetFilterOut(bFilterOut);
+
+		// Send and deliver packets
+		Server->UpdateAndSend({ Client });
+
+		// Check that the filtered object do not exist on the client.
+		UE_NET_ASSERT_EQ(Client->GetReplicationBridge()->GetReplicatedObject(ServerObjectA->NetRefHandle), nullptr);
+		UE_NET_ASSERT_EQ(Client->GetReplicationBridge()->GetReplicatedObject(ServerObjectB->NetRefHandle), nullptr);
+		UE_NET_ASSERT_EQ(Client->GetReplicationBridge()->GetReplicatedObject(ServerObjectC->NetRefHandle), nullptr);
+	}
+
+	// Mark objects dirty
+	{
+		ServerObjectA->ReplicatedCounter = 0x01;
+		ServerObjectB->ReplicatedCounter = 0x01;
+		ServerObjectC->ReplicatedCounter = 0x01;
+		Server->GetReplicationSystem()->MarkDirty(ServerObjectA->NetRefHandle);
+		Server->GetReplicationSystem()->MarkDirty(ServerObjectB->NetRefHandle);
+		Server->GetReplicationSystem()->MarkDirty(ServerObjectC->NetRefHandle);
+
+		// Send and deliver packets
+		Server->UpdateAndSend({ Client });
+
+		// Should still not exist on the client.
+		UE_NET_ASSERT_EQ(Client->GetReplicationBridge()->GetReplicatedObject(ServerObjectA->NetRefHandle), nullptr);
+		UE_NET_ASSERT_EQ(Client->GetReplicationBridge()->GetReplicatedObject(ServerObjectB->NetRefHandle), nullptr);
+		UE_NET_ASSERT_EQ(Client->GetReplicationBridge()->GetReplicatedObject(ServerObjectC->NetRefHandle), nullptr);
+	}
+
+	// Put one of them back in the scope
+	{
+		constexpr bool bFilterIn = false;
+		ServerObjectA->SetFilterOut(bFilterIn);
+
+		// Send and deliver packets
+		Server->UpdateAndSend({ Client });
+
+		// This one exists
+		UTestFilteringObject* ClientObjectA = Cast<UTestFilteringObject>(Client->GetReplicationBridge()->GetReplicatedObject(ServerObjectA->NetRefHandle));
+		UE_NET_ASSERT_NE(ClientObjectA, nullptr);
+		UE_NET_ASSERT_TRUE(ClientObjectA->ReplicatedCounter == 0x01);
+
+		// These don't
+		UE_NET_ASSERT_EQ(Client->GetReplicationBridge()->GetReplicatedObject(ServerObjectB->NetRefHandle), nullptr);
+		UE_NET_ASSERT_EQ(Client->GetReplicationBridge()->GetReplicatedObject(ServerObjectC->NetRefHandle), nullptr);
+	}
+
+	// Add another one back in the scope
+	{
+		constexpr bool bFilterIn = false;
+		ServerObjectB->SetFilterOut(bFilterIn);
+
+		// Send and deliver packets
+		Server->UpdateAndSend({ Client });
+
+		// This one exists
+		UTestFilteringObject* ClientObjectB = Cast<UTestFilteringObject>(Client->GetReplicationBridge()->GetReplicatedObject(ServerObjectB->NetRefHandle));
+		UE_NET_ASSERT_NE(ClientObjectB, nullptr);
+		UE_NET_ASSERT_TRUE(ClientObjectB->ReplicatedCounter == 0x01);
+
+		// These don't
+		UE_NET_ASSERT_EQ(Client->GetReplicationBridge()->GetReplicatedObject(ServerObjectC->NetRefHandle), nullptr);
+	}
 }
+
+} // end namespace UE::Net::Private
diff --git a/Engine/Plugins/Runtime/ReplicationSystemTestPlugin/Source/Private/Tests/ReplicationSystem/Filtering/TestFilteringObject.cpp b/Engine/Plugins/Runtime/ReplicationSystemTestPlugin/Source/Private/Tests/ReplicationSystem/Filtering/TestFilteringObject.cpp
--- a/Engine/Plugins/Runtime/ReplicationSystemTestPlugin/Source/Private/Tests/ReplicationSystem/Filtering/TestFilteringObject.cpp
+++ b/Engine/Plugins/Runtime/ReplicationSystemTestPlugin/Source/Private/Tests/ReplicationSystem/Filtering/TestFilteringObject.cpp
@@ -19,6 +19,7 @@ void UTestFilteringObject::GetLifetimeReplicatedProps( TArray< class FLifetimePr
 {
 	Super::GetLifetimeReplicatedProps(OutLifetimeProps);
 	DOREPLIFETIME(ThisClass, NetTest_FilterOut);
+	DOREPLIFETIME(ThisClass, ReplicatedCounter);
 }
 
 
diff --git a/Engine/Plugins/Runtime/ReplicationSystemTestPlugin/Source/Private/Tests/ReplicationSystem/Filtering/TestFilteringObject.h b/Engine/Plugins/Runtime/ReplicationSystemTestPlugin/Source/Private/Tests/ReplicationSystem/Filtering/TestFilteringObject.h
--- a/Engine/Plugins/Runtime/ReplicationSystemTestPlugin/Source/Private/Tests/ReplicationSystem/Filtering/TestFilteringObject.h
+++ b/Engine/Plugins/Runtime/ReplicationSystemTestPlugin/Source/Private/Tests/ReplicationSystem/Filtering/TestFilteringObject.h
@@ -12,6 +12,10 @@ class UTestFilteringObject : public UReplicatedTestObject
 
 public:
 	void SetFilterOut(bool bFilterOut);
+	bool GetFilterOut() const { return NetTest_FilterOut; }
+
+	UPROPERTY(Replicated)
+	uint32 ReplicatedCounter = 0;
 
 protected:
 	virtual void RegisterReplicationFragments(UE::Net::FFragmentRegistrationContext& Fragments, UE::Net::EFragmentRegistrationFlags RegistrationFlags) override;
diff --git a/Engine/Plugins/Runtime/ReplicationSystemTestPlugin/Source/Private/Tests/ReplicationSystem/ReplicationSystemServerClientTestFixture.cpp b/Engine/Plugins/Runtime/ReplicationSystemTestPlugin/Source/Private/Tests/ReplicationSystem/ReplicationSystemServerClientTestFixture.cpp
--- a/Engine/Plugins/Runtime/ReplicationSystemTestPlugin/Source/Private/Tests/ReplicationSystem/ReplicationSystemServerClientTestFixture.cpp
+++ b/Engine/Plugins/Runtime/ReplicationSystemTestPlugin/Source/Private/Tests/ReplicationSystem/ReplicationSystemServerClientTestFixture.cpp
@@ -369,6 +369,22 @@ void FReplicationSystemTestServer::DeliverTo(FReplicationSystemTestClient* Clien
 	FReplicationSystemTestNode::DeliverTo(*Client, Client->ConnectionIdOnServer, Client->LocalConnectionId, bDeliver);
 }
 
+bool FReplicationSystemTestServer::UpdateAndSend(const TArray<FReplicationSystemTestClient*>& Clients, bool bDeliver /*= true*/)
+{
+	bool bSuccess = true;
+
+	PreSendUpdate();
+
+	for (FReplicationSystemTestClient* Client : Clients)
+	{
+		bSuccess &= SendAndDeliverTo(Client, bDeliver);
+	}
+	
+	PostSendUpdate();
+
+	return bSuccess;
+}
+
 // FReplicationSystemServerClientTestFixture implementation
 void FReplicationSystemServerClientTestFixture::SetUp()
 {
diff --git a/Engine/Plugins/Runtime/ReplicationSystemTestPlugin/Source/Private/Tests/ReplicationSystem/ReplicationSystemServerClientTestFixture.h b/Engine/Plugins/Runtime/ReplicationSystemTestPlugin/Source/Private/Tests/ReplicationSystem/ReplicationSystemServerClientTestFixture.h
--- a/Engine/Plugins/Runtime/ReplicationSystemTestPlugin/Source/Private/Tests/ReplicationSystem/ReplicationSystemServerClientTestFixture.h
+++ b/Engine/Plugins/Runtime/ReplicationSystemTestPlugin/Source/Private/Tests/ReplicationSystem/ReplicationSystemServerClientTestFixture.h
@@ -179,6 +179,9 @@ public:
 
 	// Explicitly set delivery status
 	void DeliverTo(FReplicationSystemTestClient* Client, bool bDeliver);
+
+	// Tick and send packets to one or all clients
+	bool UpdateAndSend(const TArray<FReplicationSystemTestClient*>& Clients, bool bDeliver=true);
 };
 
 class FReplicationSystemServerClientTestFixture : public FNetworkAutomationTestSuiteFixture
