Allow replication of UObjects with UFunctions but no UProperties

#rnx

diff --git a/Engine/Plugins/Runtime/ReplicationSystemTestPlugin/Source/Private/Tests/ReplicationSystem/RPC/ReplicatedTestObjectWithRPC.h b/Engine/Plugins/Runtime/ReplicationSystemTestPlugin/Source/Private/Tests/ReplicationSystem/RPC/ReplicatedTestObjectWithRPC.h
--- a/Engine/Plugins/Runtime/ReplicationSystemTestPlugin/Source/Private/Tests/ReplicationSystem/RPC/ReplicatedTestObjectWithRPC.h
+++ b/Engine/Plugins/Runtime/ReplicationSystemTestPlugin/Source/Private/Tests/ReplicationSystem/RPC/ReplicatedTestObjectWithRPC.h
@@ -30,10 +30,6 @@ public:
 
 	bool bIsServerObject = false;
 
-	//$IRIS TODO: Temp until SetUpRuntimeReplicationData is fixed
-	UPROPERTY()
-	bool bTest;
-
 	// RPC test functions
 	UFUNCTION(Reliable, Client)
 	void ClientRPC();
diff --git a/Engine/Source/Runtime/CoreUObject/Private/UObject/Class.cpp b/Engine/Source/Runtime/CoreUObject/Private/UObject/Class.cpp
--- a/Engine/Source/Runtime/CoreUObject/Private/UObject/Class.cpp
+++ b/Engine/Source/Runtime/CoreUObject/Private/UObject/Class.cpp
@@ -4682,7 +4682,7 @@ static FAutoConsoleVariable CVarValidateReplicatedPropertyRegistration(TEXT("net
 
 void UClass::SetUpRuntimeReplicationData()
 {
-	if (!HasAnyClassFlags(CLASS_ReplicationDataIsSetUp) && PropertyLink != NULL)
+	if (!HasAnyClassFlags(CLASS_ReplicationDataIsSetUp))
 	{
 		NetFields.Empty();
 
