Added unit test for CombineMultiple and FNetBitArray

#rnx

diff --git a/Engine/Plugins/Runtime/ReplicationSystemTestPlugin/Source/Private/Tests/Core/TestNetBitArray.cpp b/Engine/Plugins/Runtime/ReplicationSystemTestPlugin/Source/Private/Tests/Core/TestNetBitArray.cpp
--- a/Engine/Plugins/Runtime/ReplicationSystemTestPlugin/Source/Private/Tests/Core/TestNetBitArray.cpp
+++ b/Engine/Plugins/Runtime/ReplicationSystemTestPlugin/Source/Private/Tests/Core/TestNetBitArray.cpp
@@ -1280,10 +1280,10 @@ UE_NET_TEST(FNetBitArrayView, TestCombineMultiple)
 
 	// Test Multiple arrays
 	{
-		const uint32 ExpectedWordBuffer[] = { 0xFF10, 0xF10F, 0x10FF, 0x0FF1 };
 		uint32 WordBufferA[] = { 0xFF00, 0xF00F, 0x00FF, 0x0FF0 };
 		uint32 WordBufferB[] = { 0x0011, 0x0110, 0x1100, 0x1001 };
 		uint32 WordBufferC[] = { 0x0010, 0x0100, 0x1000, 0x0001 };
+		const uint32 ExpectedWordBuffer[] = { 0xFF10, 0xF10F, 0x10FF, 0x0FF1 };
 
 		FNetBitArrayView BitArrayA(&WordBufferA[0], 128);
 		const FNetBitArrayView BitArrayB(&WordBufferB[0], 128);
@@ -1296,6 +1296,29 @@ UE_NET_TEST(FNetBitArrayView, TestCombineMultiple)
 			UE_NET_ASSERT_EQ(ExpectedWordBuffer[it], WordBufferA[it]);
 		}
 	}
+
+	// Test Multiple arrays
+	{
+		uint32 WordBufferA[] = { 0xFF00, 0xF00F, 0x00FF, 0x0FF0 };
+		uint32 WordBufferB[] = { 0x0011, 0x0110, 0x1100, 0x1001 };
+		uint32 WordBufferC[] = { 0x0010, 0x0100, 0x1000, 0x0001 };
+		const uint32 ExpectedWordBuffer[] = { 0xFF10, 0xF10F, 0x10FF, 0x0FF1 };
+
+		FNetBitArray BitArrayA, BitArrayB, BitArrayC;
+		BitArrayA.Init(128);
+		BitArrayB.Init(128);
+		BitArrayC.Init(128);
+		MakeNetBitArrayView(BitArrayA).Copy(FNetBitArrayView(&WordBufferA[0], 128));
+		MakeNetBitArrayView(BitArrayB).Copy(FNetBitArrayView(&WordBufferB[0], 128));
+		MakeNetBitArrayView(BitArrayC).Copy(FNetBitArrayView(&WordBufferC[0], 128));
+
+		BitArrayA.CombineMultiple(FirstWordOp, BitArrayB, SecondWordOp, BitArrayC);
+
+		for (uint32 it = 0; it < UE_ARRAY_COUNT(ExpectedWordBuffer); ++it)
+		{
+			UE_NET_ASSERT_EQ(ExpectedWordBuffer[it], BitArrayA.GetData()[it]);
+		}
+	}
 }
 
 UE_NET_TEST(FNetBitArrayView, TestSet)
