Fixed issue where FNetBitArray::FindLastZero would return index out of bounds for array sizes not aligned with the size of the storage.

diff --git a/Engine/Plugins/Runtime/ReplicationSystemTestPlugin/Source/Private/Tests/Core/TestNetBitArray.cpp b/Engine/Plugins/Runtime/ReplicationSystemTestPlugin/Source/Private/Tests/Core/TestNetBitArray.cpp
--- a/Engine/Plugins/Runtime/ReplicationSystemTestPlugin/Source/Private/Tests/Core/TestNetBitArray.cpp
+++ b/Engine/Plugins/Runtime/ReplicationSystemTestPlugin/Source/Private/Tests/Core/TestNetBitArray.cpp
@@ -671,6 +671,33 @@ UE_NET_TEST(FNetBitArrayView, TestFindLastZero)
 
 		UE_NET_ASSERT_EQ(BitArray.FindLastZero(), 49U);
 	}
+
+	// All bits zero, with unaligned bitcount
+	{
+		uint32 WordBuffer[] = { 0x0 };
+		FNetBitArrayView BitArray(WordBuffer, 31);
+
+		// Verify that we do not get an out of bounds value
+		UE_NET_ASSERT_LT(BitArray.FindLastZero(), 31U);
+		UE_NET_ASSERT_EQ(BitArray.FindLastZero(), 30U);
+	}
+
+	// Set last bit, unaligned bitcount
+	{
+		uint32 WordBuffer[] = { 0x40000000 };
+		FNetBitArrayView BitArray(WordBuffer, 31);
+
+		UE_NET_ASSERT_EQ(BitArray.FindLastZero(), 29U);
+	}
+
+	// All bits set, unaligned bitcount
+	{
+		uint32 WordBuffer[] = { 0x7fffffff };
+		FNetBitArrayView BitArray(WordBuffer, 31);
+
+		UE_NET_ASSERT_EQ(BitArray.FindLastZero(), BitArray.InvalidIndex);
+	}
+
 }
 
 UE_NET_TEST(FNetBitArrayView, TestFindLastOne)
diff --git a/Engine/Source/Runtime/Net/Core/Public/Net/Core/NetBitArray.h b/Engine/Source/Runtime/Net/Core/Public/Net/Core/NetBitArray.h
--- a/Engine/Source/Runtime/Net/Core/Public/Net/Core/NetBitArray.h
+++ b/Engine/Source/Runtime/Net/Core/Public/Net/Core/NetBitArray.h
@@ -668,8 +668,9 @@ private:
 			if (Word != WordMask)
 			{
 				const uint32 BitOffset = WordIt*WordBitCount;
-				const uint32 Index = BitOffset + WordBitCount - 1U - FPlatformMath::CountLeadingZeros(~Word);
-				return Index;
+				const uint32 Index = BitOffset + WordBitCount - 1U - FPlatformMath::CountLeadingZeros(~Word & WordMask);
+				// Need to make sure the index is not out of bounds
+				return (Index < BitCount ? Index : FNetBitArrayBase::InvalidIndex);
 			}
 
 			WordMask = StorageWordType(~StorageWordType(0));
