Iris
- Fixed issue with tearoff of actor with subobject triggering ensure
- Added test recreating the bug
- Added Description string parameter to FReplicationSystemTestNode packet sending methods to facilitate logging of more context when writing tests

diff --git a/Engine/Plugins/Runtime/ReplicationSystemTestPlugin/Source/Private/Tests/ReplicationSystem/ReplicationSystemServerClientTestFixture.cpp b/Engine/Plugins/Runtime/ReplicationSystemTestPlugin/Source/Private/Tests/ReplicationSystem/ReplicationSystemServerClientTestFixture.cpp
--- a/Engine/Plugins/Runtime/ReplicationSystemTestPlugin/Source/Private/Tests/ReplicationSystem/ReplicationSystemServerClientTestFixture.cpp
+++ b/Engine/Plugins/Runtime/ReplicationSystemTestPlugin/Source/Private/Tests/ReplicationSystem/ReplicationSystemServerClientTestFixture.cpp
@@ -236,7 +236,7 @@ void FReplicationSystemTestNode::PreSendUpdate()
 	ReplicationSystem->PreSendUpdate(1.f);
 }
 
-bool FReplicationSystemTestNode::SendUpdate(uint32 ConnectionId)
+bool FReplicationSystemTestNode::SendUpdate(uint32 ConnectionId, const TCHAR* Desc)
 {
 	FPacketData Packet;
 
@@ -257,6 +257,10 @@ bool FReplicationSystemTestNode::SendUpdate(uint32 ConnectionId)
 		Writer.CommitWrites();
 		Packet.BitCount = Writer.GetPosBits();
 		Packet.PacketId = PacketId++;
+		if (Desc)
+		{
+			Packet.Desc = FString(Desc);
+		}
 
 		Connection.WriteRecords.Enqueue(Record);
 		Connection.WrittenPackets.Enqueue(Packet);
@@ -268,6 +272,7 @@ bool FReplicationSystemTestNode::SendUpdate(uint32 ConnectionId)
 	{
 		UE_NET_TRACE_FLUSH_COLLECTOR(Context.GetTraceCollector(), GetNetTraceId(), Connection.ConnectionId, ENetTracePacketType::Outgoing);
 		UE_NET_TRACE_PACKET_SEND(GetNetTraceId(), Connection.ConnectionId, Packet.PacketId, Packet.BitCount);
+		UE_LOG(LogIris, Log, TEXT("ReplicationSystemTestFixture: Conn: %u Send PacketId: %u %s"), ConnectionId, Packet.PacketId, *Packet.Desc);
 	}
 
 	UE_NET_TRACE_DESTROY_COLLECTOR(Context.GetTraceCollector());
@@ -293,6 +298,7 @@ void FReplicationSystemTestNode::DeliverTo(FReplicationSystemTestNode& Dest, uin
 
 		Context.SetTraceCollector(UE_NET_TRACE_CREATE_COLLECTOR(ENetTraceVerbosity::Trace));
 
+		UE_LOG(LogIris, Log, TEXT("ReplicationSystemTestFixture: Conn: %u Deliver PacketId: %u %s"), LocalConnectionId, Packet.PacketId, *Packet.Desc);
 		Dest.RecvUpdate(RemoteConnectionId, Context);
 
 		UE_NET_TRACE_FLUSH_COLLECTOR(Context.GetTraceCollector(), Dest.GetNetTraceId(), RemoteConnectionId, ENetTracePacketType::Incoming);
@@ -301,6 +307,7 @@ void FReplicationSystemTestNode::DeliverTo(FReplicationSystemTestNode& Dest, uin
 	}
 	else
 	{
+		UE_LOG(LogIris, Log, TEXT("ReplicationSystemTestFixture: Conn: %u Dropped PacketId: %u %s"), LocalConnectionId, Packet.PacketId, *Packet.Desc);
 		UE_NET_TRACE_PACKET_DROPPED(Dest.GetNetTraceId(), RemoteConnectionId, Packet.PacketId, ENetTracePacketType::Incoming);
 		UE_NET_TRACE_PACKET_DROPPED(GetNetTraceId(), LocalConnectionId, Packet.PacketId, ENetTracePacketType::Outgoing);
 	}
@@ -345,9 +352,9 @@ FReplicationSystemTestServer::FReplicationSystemTestServer(const TCHAR* Name)
 {
 }
 
-bool FReplicationSystemTestServer::SendAndDeliverTo(FReplicationSystemTestClient* Client, bool bDeliver)
+bool FReplicationSystemTestServer::SendAndDeliverTo(FReplicationSystemTestClient* Client, bool bDeliver, const TCHAR* Desc)
 {
-	if (SendUpdate(Client->ConnectionIdOnServer))
+	if (SendUpdate(Client->ConnectionIdOnServer, Desc))
 	{
 		DeliverTo(Client, bDeliver);
 
@@ -358,9 +365,9 @@ bool FReplicationSystemTestServer::SendAndDeliverTo(FReplicationSystemTestClient
 }
 
 // Send data, return true if data was written
-bool FReplicationSystemTestServer::SendTo(FReplicationSystemTestClient* Client)
+bool FReplicationSystemTestServer::SendTo(FReplicationSystemTestClient* Client, const TCHAR* Desc)
 {
-	return SendUpdate(Client->ConnectionIdOnServer);
+	return SendUpdate(Client->ConnectionIdOnServer, Desc);
 }
 
 // If bDeliver is true deliver data to client and report packet as delivered
diff --git a/Engine/Plugins/Runtime/ReplicationSystemTestPlugin/Source/Private/Tests/ReplicationSystem/ReplicationSystemServerClientTestFixture.h b/Engine/Plugins/Runtime/ReplicationSystemTestPlugin/Source/Private/Tests/ReplicationSystem/ReplicationSystemServerClientTestFixture.h
--- a/Engine/Plugins/Runtime/ReplicationSystemTestPlugin/Source/Private/Tests/ReplicationSystem/ReplicationSystemServerClientTestFixture.h
+++ b/Engine/Plugins/Runtime/ReplicationSystemTestPlugin/Source/Private/Tests/ReplicationSystem/ReplicationSystemServerClientTestFixture.h
@@ -59,6 +59,7 @@ public:
 		alignas(16) uint8 PacketBuffer[MaxPacketSize];
 		uint32 BitCount;
 		uint32 PacketId;
+		FString Desc;
 	};
 
 	struct FConnectionInfo
@@ -131,8 +132,8 @@ public:
 	void PostSendUpdate();
 
 	// Update methods for server connections
-	bool SendUpdate(uint32 ConnectionId);
-	bool SendUpdate() { return SendUpdate(1); }
+	bool SendUpdate(uint32 ConnectionId, const TCHAR* Desc = nullptr);
+	bool SendUpdate(const TCHAR* Desc = nullptr) { return SendUpdate(1, Desc); }
 
 	void DeliverTo(FReplicationSystemTestNode& Dest, uint32 LocalConnectionId, uint32 RemoteConnectionId, bool bDeliver);
 	void RecvUpdate(uint32 ConnectionId, FNetSerializationContext& Context);
@@ -175,10 +176,10 @@ public:
 	explicit FReplicationSystemTestServer(const TCHAR* Name);
 
 	// Send data and deliver to the client if bDeliver is true
-	bool SendAndDeliverTo(FReplicationSystemTestClient* Client, bool bDeliver);
+	bool SendAndDeliverTo(FReplicationSystemTestClient* Client, bool bDeliver, const TCHAR* Desc = nullptr);
 
 	// Send data, return true if data was written
-	bool SendTo(FReplicationSystemTestClient* Client);
+	bool SendTo(FReplicationSystemTestClient* Client, const TCHAR* Desc = nullptr);
 
 	// Explicitly set delivery status
 	void DeliverTo(FReplicationSystemTestClient* Client, bool bDeliver);
diff --git a/Engine/Plugins/Runtime/ReplicationSystemTestPlugin/Source/Private/Tests/ReplicationSystem/TestFlushBeforeDestroy.cpp b/Engine/Plugins/Runtime/ReplicationSystemTestPlugin/Source/Private/Tests/ReplicationSystem/TestFlushBeforeDestroy.cpp
--- a/Engine/Plugins/Runtime/ReplicationSystemTestPlugin/Source/Private/Tests/ReplicationSystem/TestFlushBeforeDestroy.cpp
+++ b/Engine/Plugins/Runtime/ReplicationSystemTestPlugin/Source/Private/Tests/ReplicationSystem/TestFlushBeforeDestroy.cpp
@@ -561,5 +561,96 @@ UE_NET_TEST_FIXTURE(FTestFlushBeforeDestroyFixture, TestImmediateTearOffExisting
 	UE_NET_ASSERT_TRUE(Cast<UTestReplicatedIrisObject>(Client->GetReplicationBridge()->GetReplicatedObject(ServerSubObject->NetRefHandle)) == nullptr);
 }
 
+// Test to recreate a very specific bug where owner being torn-off has in flight rpc requiring a flush
+UE_NET_TEST_FIXTURE(FTestFlushBeforeDestroyFixture, TestImmediateTearOffWithSubObjectAndInFlightAttachmentsAndPacketLoss)
+{
+	UReplicationSystem* ReplicationSystem = Server->ReplicationSystem;
+
+	// Add a client
+	FReplicationSystemTestClient* Client = CreateClient();
+	RegisterNetBlobHandlers(Client);
+
+	// Spawn object on server
+	UTestReplicatedIrisObject* ServerObject = Server->CreateObject(0,0);
+
+	// Spawn second object on server as a subobject
+	UTestReplicatedIrisObject* ServerSubObject = Server->CreateSubObject(ServerObject->NetRefHandle, 0, 0);
+
+	// Trigger replication
+	ServerObject->IntA = 1;
+	ServerSubObject->IntA = 1;
+
+	// Send and deliver packet
+	Server->PreSendUpdate();
+	Server->SendAndDeliverTo(Client, true, TEXT("Create Objects"));
+	Server->PostSendUpdate();
+
+	// Store Pointer to objects
+	UTestReplicatedIrisObject* ClientObjectThatWillBeTornOff = Cast<UTestReplicatedIrisObject>(Client->GetReplicationBridge()->GetReplicatedObject(ServerObject->NetRefHandle));
+	UE_NET_ASSERT_TRUE(ClientObjectThatWillBeTornOff != nullptr);
+	UE_NET_ASSERT_EQ(ServerObject->IntA, ClientObjectThatWillBeTornOff->IntA);
+
+	UTestReplicatedIrisObject* ClientSubObjectThatWillBeTornOff = Cast<UTestReplicatedIrisObject>(Client->GetReplicationBridge()->GetReplicatedObject(ServerSubObject->NetRefHandle));
+	UE_NET_ASSERT_TRUE(ClientSubObjectThatWillBeTornOff != nullptr);
+	UE_NET_ASSERT_EQ(ServerSubObject->IntA, ClientSubObjectThatWillBeTornOff->IntA);
+
+	// Modify the value of object only
+	ServerObject->IntA = 2;
+
+	// Create attachment to force flush behavior by having a rpc in flight
+	{
+		constexpr uint32 PayloadBitCount = 24;
+		const TRefCountPtr<FNetObjectAttachment>& Attachment = MockNetObjectAttachmentHandler->CreateReliableNetObjectAttachment(PayloadBitCount);
+		FNetObjectReference AttachmentTarget = FObjectReferenceCache::MakeNetObjectReference(ServerObject->NetRefHandle);
+		Server->GetReplicationSystem()->QueueNetObjectAttachment(Client->ConnectionIdOnServer, AttachmentTarget, Attachment);
+	}
+
+	Server->PreSendUpdate();
+	Server->SendTo(Client, TEXT("State data + Attachment"));
+	Server->PostSendUpdate();
+
+	// Modify the value of object only
+	++ServerObject->IntA ;
+
+	Server->PreSendUpdate();
+	Server->SendTo(Client, TEXT("State data"));
+	Server->PostSendUpdate();
+
+	// TearOff the object using immediate tear-off
+	Server->ReplicationBridge->EndReplication(ServerObject, EEndReplicationFlags::TearOff);
+
+	Server->PreSendUpdate();
+	Server->SendTo(Client, TEXT("Tear off"));
+	Server->PostSendUpdate();
+
+	// Deliver packet to drive PendingTearOff -> WaitOnFlush
+	Server->DeliverTo(Client, true);
+
+	// Notify that we dropped tear off data
+	Server->DeliverTo(Client, false);
+
+	// This earlier caused an unwanted state transition
+	Server->PreSendUpdate();
+	Server->SendTo(Client, TEXT("Packet after tearoff"));
+	Server->PostSendUpdate();
+
+	// Drop the packet containing the original tear-off
+	Server->DeliverTo(Client, false);
+
+	// Deliver a packet
+	Server->DeliverTo(Client, true);
+
+	// This should contain resend of lost state
+	Server->PreSendUpdate();
+	Server->SendAndDeliverTo(Client, true, TEXT("Resending tearoff"));
+	Server->PostSendUpdate();
+
+	// Verify that ClientObject is torn-off and that the final state was applied
+	UE_NET_ASSERT_EQ(ServerObject->IntA, ClientObjectThatWillBeTornOff->IntA);
+	UE_NET_ASSERT_EQ(ServerSubObject->IntA, ClientSubObjectThatWillBeTornOff->IntA);
+	UE_NET_ASSERT_TRUE(Cast<UTestReplicatedIrisObject>(Client->GetReplicationBridge()->GetReplicatedObject(ServerSubObject->NetRefHandle)) == nullptr);
+	UE_NET_ASSERT_TRUE(Cast<UTestReplicatedIrisObject>(Client->GetReplicationBridge()->GetReplicatedObject(ServerObject->NetRefHandle)) == nullptr);
+}
+
 
 }
diff --git a/Engine/Source/Runtime/Experimental/Iris/Core/Private/Iris/ReplicationSystem/ReplicationWriter.cpp b/Engine/Source/Runtime/Experimental/Iris/Core/Private/Iris/ReplicationSystem/ReplicationWriter.cpp
--- a/Engine/Source/Runtime/Experimental/Iris/Core/Private/Iris/ReplicationSystem/ReplicationWriter.cpp
+++ b/Engine/Source/Runtime/Experimental/Iris/Core/Private/Iris/ReplicationSystem/ReplicationWriter.cpp
@@ -550,34 +550,40 @@ uint32 FReplicationWriter::GetDefaultFlushFlags() const
 	return EFlushFlags::FlushFlags_FlushReliable;
 }
 
-uint32 FReplicationWriter::GetFlushStatus(uint32 InternalIndex, const FReplicationInfo& Info, uint32 InFlushFlags) const
+uint32 FReplicationWriter::GetFlushStatus(uint32 InternalIndex, const FReplicationInfo& Info, uint32 FlushFlagsToTest) const
 {
 	uint32 FlushFlags = EFlushFlags::FlushFlags_None;
 
-	if (InFlushFlags == EFlushFlags::FlushFlags_None)
+	if (FlushFlagsToTest == EFlushFlags::FlushFlags_None)
 	{
 		return FlushFlags;
 	}
 
-	if (!!(InFlushFlags & EFlushFlags::FlushFlags_FlushState) && (Info.HasDirtyChangeMask || HasInFlightStateChanges(InternalIndex, Info) || IsObjectPartOfActiveHugeObject(InternalIndex, Info)))
+	if (!!(FlushFlagsToTest & EFlushFlags::FlushFlags_FlushState) && (Info.HasDirtyChangeMask || HasInFlightStateChanges(InternalIndex, Info) || IsObjectPartOfActiveHugeObject(InternalIndex, Info)))
 	{
 		FlushFlags |= EFlushFlags::FlushFlags_FlushState;
 	}
 
-	if (!!(InFlushFlags & EFlushFlags::FlushFlags_FlushReliable) && !Attachments.IsAllReliableSentAndAcked(ENetObjectAttachmentType::Normal, InternalIndex))
+	if (!!(FlushFlagsToTest & EFlushFlags::FlushFlags_FlushReliable) && !Attachments.IsAllReliableSentAndAcked(ENetObjectAttachmentType::Normal, InternalIndex))
 	{
 		FlushFlags |= EFlushFlags::FlushFlags_FlushReliable;
 	}
 
-	if (!Info.IsSubObject && FlushFlags != FlushFlags_All)
+	// Do we have a tear-off for the subobject in-flight?
+	if (!!(FlushFlagsToTest & EFlushFlags::FlushFlags_FlushTornOffSubObjects) && (Info.IsSubObject && Info.TearOff && (Info.GetState() == EReplicatedObjectState::WaitOnDestroyConfirmation)))
+	{
+		FlushFlags |= EFlushFlags::FlushFlags_FlushTornOffSubObjects;
+	}
+
+	if (!Info.IsSubObject && FlushFlags != FlushFlagsToTest)
 	{
 		// Check status of SubObjects as well.
 		for (uint32 SubObjectIndex : NetRefHandleManager->GetSubObjects(InternalIndex))
 		{
 			const FReplicationInfo& SubObjectInfo = GetReplicationInfo(SubObjectIndex);
-			FlushFlags |= GetFlushStatus(SubObjectIndex, SubObjectInfo, InFlushFlags);
+			FlushFlags |= GetFlushStatus(SubObjectIndex, SubObjectInfo, FlushFlagsToTest);
 
-			if (FlushFlags == FlushFlags_All)
+			if (FlushFlags == FlushFlagsToTest)
 			{
 				break;
 			}
@@ -1915,7 +1921,7 @@ FReplicationWriter::EWriteObjectStatus FReplicationWriter::WriteObjectAndSubObje
 	BatchEntry.bHasUnsentAttachments = bHasAttachments;
 
 	// Check if we must defer tearoff until after flush
-	const bool bSentTearOff = Info.TearOff && (GetFlushStatus(InternalIndex, Info, Info.FlushFlags) == EFlushFlags::FlushFlags_None);
+	const bool bSentTearOff = Info.TearOff && (GetFlushStatus(InternalIndex, Info, Info.FlushFlags | EFlushFlags::FlushFlags_FlushTornOffSubObjects) == EFlushFlags::FlushFlags_None);
 
 	Context.SetIsInitState(bIsInitialState);
 
diff --git a/Engine/Source/Runtime/Experimental/Iris/Core/Private/Iris/ReplicationSystem/ReplicationWriter.h b/Engine/Source/Runtime/Experimental/Iris/Core/Private/Iris/ReplicationSystem/ReplicationWriter.h
--- a/Engine/Source/Runtime/Experimental/Iris/Core/Private/Iris/ReplicationSystem/ReplicationWriter.h
+++ b/Engine/Source/Runtime/Experimental/Iris/Core/Private/Iris/ReplicationSystem/ReplicationWriter.h
@@ -95,7 +95,8 @@ public:
 		FlushFlags_None				= 0U,
 		FlushFlags_FlushState		= 1U << 0U,											// Make sure that all current state data is acknowledged before we stop replicating the object
 		FlushFlags_FlushReliable	= FlushFlags_FlushState << 1U,						// Make sure that all enqueued Reliable RPCs are delivered before we stop replicating the object
-		FlushFlags_All				= FlushFlags_FlushState | FlushFlags_FlushReliable,
+		FlushFlags_FlushTornOffSubObjects	= FlushFlags_FlushReliable << 1U,					// Make sure that we flush TearOff and replicated destroy properly
+		FlushFlags_All				= FlushFlags_FlushState | FlushFlags_FlushReliable | FlushFlags_FlushTornOffSubObjects,
 		FlushFlags_Default			= FlushFlags_FlushReliable,
 	};
 
@@ -126,7 +127,7 @@ public:
 				uint64 IsDeltaCompressionEnabled : 1;					// Set to 1 if deltacompression is enabled for this object
 				uint64 LastAckedBaselineIndex : 2;						// Last acknowledged baseline index which we can use for deltacompresion
 				uint64 PendingBaselineIndex : 2;						// Baseline index pending acknowledgment from client
-				uint64 FlushFlags : 2;									// Flags indicating what we are waiting for when flushing
+				uint64 FlushFlags : 3;									// Flags indicating what we are waiting for when flushing
 			};
 		};
 
@@ -352,7 +353,7 @@ private:
 private:
 
 	uint32 GetDefaultFlushFlags() const;
-	uint32 GetFlushStatus(uint32 InternalIndex, const FReplicationInfo& Info, uint32 InFlushFlags = EFlushFlags::FlushFlags_Default) const;
+	uint32 GetFlushStatus(uint32 InternalIndex, const FReplicationInfo& Info, uint32 FlushFlagsToTest = EFlushFlags::FlushFlags_Default) const;
 	void SetPendingDestroyOrSubObjectPendingDestroyState(uint32 Index, FReplicationInfo& Info);
 
 	bool IsObjectIndexForOOBAttachment(uint32 InternalIndex) const { return InternalIndex == ObjectIndexForOOBAttachment; }
