Iris - Fix out of bounds read in test.

#rnx

diff --git a/Engine/Plugins/Runtime/ReplicationSystemTestPlugin/Source/Private/Tests/Serialization/TestVectorNetSerializer.cpp b/Engine/Plugins/Runtime/ReplicationSystemTestPlugin/Source/Private/Tests/Serialization/TestVectorNetSerializer.cpp
--- a/Engine/Plugins/Runtime/ReplicationSystemTestPlugin/Source/Private/Tests/Serialization/TestVectorNetSerializer.cpp
+++ b/Engine/Plugins/Runtime/ReplicationSystemTestPlugin/Source/Private/Tests/Serialization/TestVectorNetSerializer.cpp
@@ -27,17 +27,18 @@ protected:
 
 struct FTestFloatTriplet
 {
-	float X;
-	float Y;
-	float Z;
+	FVector::FReal X;
+	FVector::FReal Y;
+	FVector::FReal Z;
+	FVector::FReal Padding;
 };
 
 const FTestFloatTriplet BadVectorValues[] =
 {
-	{-INFINITY, NAN, -NAN},
-	{-INFINITY, -INFINITY, -INFINITY},
-	{NAN, NAN, NAN},
-	{-NAN, -NAN, -NAN}
+	{-INFINITY, NAN, -NAN, 0.0f},
+	{-INFINITY, -INFINITY, -INFINITY, 0.0f},
+	{NAN, NAN, NAN, 0.0f},
+	{-NAN, -NAN, -NAN, 0.0f}
 };
 
 const FVector FTestVectorNetSerializer::Values[] =
