Iris - Added missing prioritzation calls in test.

#rnx

diff --git a/Engine/Plugins/Runtime/ReplicationSystemTestPlugin/Source/Private/Tests/ReplicationSystem/Prioritization/TestSphereNetObjectPrioritizer.cpp b/Engine/Plugins/Runtime/ReplicationSystemTestPlugin/Source/Private/Tests/ReplicationSystem/Prioritization/TestSphereNetObjectPrioritizer.cpp
--- a/Engine/Plugins/Runtime/ReplicationSystemTestPlugin/Source/Private/Tests/ReplicationSystem/Prioritization/TestSphereNetObjectPrioritizer.cpp
+++ b/Engine/Plugins/Runtime/ReplicationSystemTestPlugin/Source/Private/Tests/ReplicationSystem/Prioritization/TestSphereNetObjectPrioritizer.cpp
@@ -203,6 +203,13 @@ UE_NET_TEST_FIXTURE(FTestSphereNetObjectPrioritizer, ViewPositionSameAsObjectGiv
 		PrioParams.View.Views.Add(MakeReplicationView(FVector(0, 0, 0), FVector(1, 0, 0), FMath::DegreesToRadians(60)).Views[0]);
 		PrioParams.View.Views.Add(MakeReplicationView(-TestPosition, FVector(1, 0, 0), FMath::DegreesToRadians(60)).Views[0]);
 
+		SphereNetObjectPrioritizer->PrioritizeWithParams(PrioParams);
+		UE_NET_ASSERT_EQ(OutPriority, SphereNetObjectPrioritizerConfig->InnerPriority);
+
+		// Reverse views. The result should be the same.
+		Algo::Reverse(PrioParams.View.Views);
+		SphereNetObjectPrioritizer->PrioritizeWithParams(PrioParams);
+		UE_NET_ASSERT_EQ(OutPriority, SphereNetObjectPrioritizerConfig->InnerPriority);
 	}
 }
 
