Do not depend on binaries that do not exist.

#rnx

diff --git a/Engine/Plugins/Runtime/ReplicationSystemTestPlugin/Source/ReplicationSystemTestPlugin.Build.cs b/Engine/Plugins/Runtime/ReplicationSystemTestPlugin/Source/ReplicationSystemTestPlugin.Build.cs
--- a/Engine/Plugins/Runtime/ReplicationSystemTestPlugin/Source/ReplicationSystemTestPlugin.Build.cs
+++ b/Engine/Plugins/Runtime/ReplicationSystemTestPlugin/Source/ReplicationSystemTestPlugin.Build.cs
@@ -45,11 +45,15 @@ namespace UnrealBuildTool.Rules
 					}
 					);
 			}
-			else
+			else if (Target.Platform != UnrealTargetPlatform.LinuxArm64)
 			{
 				// For validation of compatibility with low-level tests even when not running them
 				PrivateDependencyModuleNames.Add("Catch2");
 			}
+			else
+			{
+				PrivateIncludePathModuleNames.Add("Catch2");
+			}
 
 			PrivateDefinitions.Add(String.Format("UE_NET_WITH_LOW_LEVEL_TESTS={0}", Target.ExplicitTestsTarget ? "1" : "0"));
 
