Iris - Fix EnumNetSerializer issue with enum out of bounds values. Fixed tests to use same config init path as the serializer.

diff --git a/Engine/Plugins/Runtime/ReplicationSystemTestPlugin/Source/Private/Tests/Serialization/TestEnumNetSerializer.cpp b/Engine/Plugins/Runtime/ReplicationSystemTestPlugin/Source/Private/Tests/Serialization/TestEnumNetSerializer.cpp
--- a/Engine/Plugins/Runtime/ReplicationSystemTestPlugin/Source/Private/Tests/Serialization/TestEnumNetSerializer.cpp
+++ b/Engine/Plugins/Runtime/ReplicationSystemTestPlugin/Source/Private/Tests/Serialization/TestEnumNetSerializer.cpp
@@ -3,6 +3,7 @@
 #include "TestNetSerializerFixture.h"
 #include "Iris/Core/BitTwiddling.h"
 #include "Iris/Serialization/EnumNetSerializers.h"
+#include "Iris/Serialization/InternalEnumNetSerializers.h"
 #include "EnumTestTypes.h"
 #include <limits>
 
@@ -106,42 +107,15 @@ void FTestEnumIntNetSerializer<SerializerConfig, SourceType, EnumType>::SetUp()
 	}
 
 	Enum = StaticEnum<EnumType>();
-	// NumEnums actually also contain the generated _MAX enum value which might not even be a valid value by the backed type. Skip it!
-	const int32 EnumValueCount = Enum->NumEnums() - 1;
-
-	// Setup the NetSerializerConfig
-	using LargeIntegerType = typename TChooseClass<TIsSigned<SourceType>::Value, int64, uint64>::Result;
+	InitEnumNetSerializerConfig(Config, Enum);
 
+	// Setup test values
 	{
-		// Find smallest and largest values
-		if (EnumValueCount == 0)
-		{
-			Config.LowerBound = 0;
-			Config.UpperBound = 0;
-			Config.BitCount = GetBitsNeededForRange(Config.LowerBound, Config.UpperBound);
-			Config.Enum = Enum;
-		}
-		else
-		{
-			// N.B. This code is designed to also work with all uint64 enums, which UEnum doesn't handle perfectly.
-			LargeIntegerType SmallestValue = std::numeric_limits<LargeIntegerType>::max();
-			LargeIntegerType LargestValue = std::numeric_limits<LargeIntegerType>::min();
-			for (int32 EnumIt = 0, EnumEndIt = EnumValueCount; EnumIt != EnumEndIt; ++EnumIt)
-			{
-				const LargeIntegerType Value = static_cast<LargeIntegerType>(Enum->GetValueByIndex(EnumIt));
-				SmallestValue = FMath::Min(SmallestValue, Value);
-				LargestValue = FMath::Max(LargestValue, Value);
-			}
+		// NumEnums actually also contain the generated _MAX enum value which might not even be a valid value by the backed type. Skip it!
+		const int32 EnumValueCount = Enum->NumEnums() - 1;
 
-			Config.LowerBound = static_cast<SourceType>(SmallestValue);
-			Config.UpperBound = static_cast<SourceType>(LargestValue);
-			Config.BitCount = GetBitsNeededForRange(Config.LowerBound, Config.UpperBound);
-			Config.Enum = Enum;
-		}
-	}
+		using LargeIntegerType = typename TChooseClass<TIsSigned<SourceType>::Value, int64, uint64>::Result;
 
-	// Setup test values
-	{
 		// Valid values
 		TArray<SourceType> TempValues;
 		TempValues.Reserve(EnumValueCount);
diff --git a/Engine/Source/Runtime/Experimental/Iris/Core/Private/Iris/Serialization/InternalEnumNetSerializers.cpp b/Engine/Source/Runtime/Experimental/Iris/Core/Private/Iris/Serialization/InternalEnumNetSerializers.cpp
--- a/Engine/Source/Runtime/Experimental/Iris/Core/Private/Iris/Serialization/InternalEnumNetSerializers.cpp
+++ b/Engine/Source/Runtime/Experimental/Iris/Core/Private/Iris/Serialization/InternalEnumNetSerializers.cpp
@@ -18,6 +18,8 @@ public:
 		const int32 EnumValueCount = Enum->NumEnums();
 
 		using LargeIntegerType = std::conditional_t<std::is_signed_v<SourceType>, int64, uint64>;
+		constexpr LargeIntegerType MaxSupportedValue = static_cast<LargeIntegerType>(std::numeric_limits<SourceType>::max());
+
 
 		// Find smallest and largest values.
 		if (EnumValueCount <= 0)
@@ -38,7 +40,8 @@ public:
 			{
 				const LargeIntegerType Value = static_cast<LargeIntegerType>(Enum->GetValueByIndex(EnumIt));
 				SmallestValue = FMath::Min(SmallestValue, Value);
-				LargestValue = FMath::Max(LargestValue, Value);
+				// Prevent out of bounds auto-generated _MAX values from sneaking in.
+				LargestValue = FMath::Max(LargestValue, FMath::Min(Value, MaxSupportedValue));
 			}
 
 			OutConfig.LowerBound = static_cast<SourceType>(SmallestValue);
diff --git a/Engine/Source/Runtime/Experimental/Iris/Core/Private/Iris/Serialization/InternalEnumNetSerializers.h b/Engine/Source/Runtime/Experimental/Iris/Core/Private/Iris/Serialization/InternalEnumNetSerializers.h
--- a/Engine/Source/Runtime/Experimental/Iris/Core/Private/Iris/Serialization/InternalEnumNetSerializers.h
+++ b/Engine/Source/Runtime/Experimental/Iris/Core/Private/Iris/Serialization/InternalEnumNetSerializers.h
@@ -7,14 +7,14 @@
 namespace UE::Net::Private
 {
 
-bool InitEnumNetSerializerConfig(FEnumInt8NetSerializerConfig& OutConfig, const UEnum* Enum);
-bool InitEnumNetSerializerConfig(FEnumInt16NetSerializerConfig& OutConfig, const UEnum* Enum);
-bool InitEnumNetSerializerConfig(FEnumInt32NetSerializerConfig& OutConfig, const UEnum* Enum);
-bool InitEnumNetSerializerConfig(FEnumInt64NetSerializerConfig& OutConfig, const UEnum* Enum);
+IRISCORE_API bool InitEnumNetSerializerConfig(FEnumInt8NetSerializerConfig& OutConfig, const UEnum* Enum);
+IRISCORE_API bool InitEnumNetSerializerConfig(FEnumInt16NetSerializerConfig& OutConfig, const UEnum* Enum);
+IRISCORE_API bool InitEnumNetSerializerConfig(FEnumInt32NetSerializerConfig& OutConfig, const UEnum* Enum);
+IRISCORE_API bool InitEnumNetSerializerConfig(FEnumInt64NetSerializerConfig& OutConfig, const UEnum* Enum);
 
-bool InitEnumNetSerializerConfig(FEnumUint8NetSerializerConfig& OutConfig, const UEnum* Enum);
-bool InitEnumNetSerializerConfig(FEnumUint16NetSerializerConfig& OutConfig, const UEnum* Enum);
-bool InitEnumNetSerializerConfig(FEnumUint32NetSerializerConfig& OutConfig, const UEnum* Enum);
-bool InitEnumNetSerializerConfig(FEnumUint64NetSerializerConfig& OutConfig, const UEnum* Enum);
+IRISCORE_API bool InitEnumNetSerializerConfig(FEnumUint8NetSerializerConfig& OutConfig, const UEnum* Enum);
+IRISCORE_API bool InitEnumNetSerializerConfig(FEnumUint16NetSerializerConfig& OutConfig, const UEnum* Enum);
+IRISCORE_API bool InitEnumNetSerializerConfig(FEnumUint32NetSerializerConfig& OutConfig, const UEnum* Enum);
+IRISCORE_API bool InitEnumNetSerializerConfig(FEnumUint64NetSerializerConfig& OutConfig, const UEnum* Enum);
 
 }
