Iris
- Added friend declaration for inlined subclass to comply with recent compilers.

diff --git a/Engine/Plugins/Runtime/ReplicationSystemTestPlugin/Source/Private/Tests/ReplicationSystem/ReplicatedTestObject.h b/Engine/Plugins/Runtime/ReplicationSystemTestPlugin/Source/Private/Tests/ReplicationSystem/ReplicatedTestObject.h
--- a/Engine/Plugins/Runtime/ReplicationSystemTestPlugin/Source/Private/Tests/ReplicationSystem/ReplicatedTestObject.h
+++ b/Engine/Plugins/Runtime/ReplicationSystemTestPlugin/Source/Private/Tests/ReplicationSystem/ReplicatedTestObject.h
@@ -475,6 +475,9 @@ public:
 	};
 
 protected:
+
+	friend FSupressCreateInstanceFailedEnsureScope;
+
 	// Type specifics for serializing creation data this will most likely be made into a separate interface to support different types of header data for different types
 	// But if we can avoid having custom data per type for instantiating remote objects as we would like to be able to fully express the state of a replicated object using the define protocol alone. 
 	// This will probably be specified as a Iris generated struct
