Iris - COND_SkipReplay properties are now treated as having a lifetime condition. Added missing tests for various lifetimeconditions.

#rnx

diff --git a/Engine/Plugins/Runtime/ReplicationSystemTestPlugin/Source/Private/Tests/ReplicationSystem/Conditionals/TestConditionals.cpp b/Engine/Plugins/Runtime/ReplicationSystemTestPlugin/Source/Private/Tests/ReplicationSystem/Conditionals/TestConditionals.cpp
--- a/Engine/Plugins/Runtime/ReplicationSystemTestPlugin/Source/Private/Tests/ReplicationSystem/Conditionals/TestConditionals.cpp
+++ b/Engine/Plugins/Runtime/ReplicationSystemTestPlugin/Source/Private/Tests/ReplicationSystem/Conditionals/TestConditionals.cpp
@@ -29,6 +29,7 @@ UE_NET_TEST_FIXTURE(FTestConditionalsFixture, ToOwnerStateIsReplicatedToOwner)
 	// Set some values in ToOwner only state
 	ServerObject->ConnectionFilteredComponents[0]->ToOwnerA = 13;
 	ServerObject->ConnectionFilteredComponents[0]->ToOwnerB = 37;
+	ServerObject->ConnectionFilteredComponents[0]->ReplayOrOwner = 47;
 
 	// Set owner
 	Server->ReplicationSystem->SetOwningNetConnection(ServerObject->NetRefHandle, Client->ConnectionIdOnServer);
@@ -45,6 +46,7 @@ UE_NET_TEST_FIXTURE(FTestConditionalsFixture, ToOwnerStateIsReplicatedToOwner)
 	// Check that the To Owner members have the same values as on the sending side.
 	UE_NET_ASSERT_EQ(ClientObject->ConnectionFilteredComponents[0]->ToOwnerA, ServerObject->ConnectionFilteredComponents[0]->ToOwnerA);
 	UE_NET_ASSERT_EQ(ClientObject->ConnectionFilteredComponents[0]->ToOwnerB, ServerObject->ConnectionFilteredComponents[0]->ToOwnerB);
+	UE_NET_ASSERT_EQ(ClientObject->ConnectionFilteredComponents[0]->ReplayOrOwner, ServerObject->ConnectionFilteredComponents[0]->ReplayOrOwner);
 }
 
 UE_NET_TEST_FIXTURE(FTestConditionalsFixture, ToOwnerStateIsNotReplicatedToNonOwner)
@@ -61,6 +63,7 @@ UE_NET_TEST_FIXTURE(FTestConditionalsFixture, ToOwnerStateIsNotReplicatedToNonOw
 	// Set some values in ToOwner only state
 	ServerObject->ConnectionFilteredComponents[0]->ToOwnerA = 13;
 	ServerObject->ConnectionFilteredComponents[0]->ToOwnerB = 37;
+	ServerObject->ConnectionFilteredComponents[0]->ReplayOrOwner = 47;
 
 	// Do not set owner. Optionally one could set it to some connection other than the client's.
 
@@ -76,6 +79,7 @@ UE_NET_TEST_FIXTURE(FTestConditionalsFixture, ToOwnerStateIsNotReplicatedToNonOw
 	// Check that the To Owner members haven't been modified
 	UE_NET_ASSERT_EQ(ClientObject->ConnectionFilteredComponents[0]->ToOwnerA, 0);
 	UE_NET_ASSERT_EQ(ClientObject->ConnectionFilteredComponents[0]->ToOwnerB, 0);
+	UE_NET_ASSERT_EQ(ClientObject->ConnectionFilteredComponents[0]->ReplayOrOwner, 0);
 }
 
 UE_NET_TEST_FIXTURE(FTestConditionalsFixture, SkipOwnerStateIsNotReplicatedToOwner)
@@ -888,5 +892,122 @@ UE_NET_TEST_FIXTURE(FTestConditionalsFixture, HierarchicalSubObjectIsNotReplicat
 
 }
 
+UE_NET_TEST_FIXTURE(FTestConditionalsFixture, NoneIsReplicated)
+{
+	// Add client
+	FReplicationSystemTestClient* Client = CreateClient();
+
+	// Spawn object on server
+	constexpr uint32 ConnectionFilteredComponentCount = 1;
+	UTestReplicatedIrisObject::FComponents ObjectComponents;
+	ObjectComponents.ConnectionFilteredComponentCount = ConnectionFilteredComponentCount;
+	UTestReplicatedIrisObject* ServerObject = Server->CreateObject(ObjectComponents);
+
+	// Set values that should be replicated
+	ServerObject->ConnectionFilteredComponents[0]->NoneInt = 13;
+
+	// Could set autonomous condition to be some connection other than the create client, or not set at all as is done now.
+
+	// Send and deliver packet
+	Server->PreSendUpdate();
+	Server->SendAndDeliverTo(Client, DeliverPacket);
+	Server->PostSendUpdate();
+	
+	// Object should have been created on the client
+	const UTestReplicatedIrisObject* ClientObject = Cast<UTestReplicatedIrisObject>(Client->GetReplicationBridge()->GetReplicatedObject(ServerObject->NetRefHandle));
+	UE_NET_ASSERT_NE(ClientObject, nullptr);
+
+	// Check that the COND_None member has been modified
+	UE_NET_ASSERT_EQ(ClientObject->ConnectionFilteredComponents[0]->NoneInt, ServerObject->ConnectionFilteredComponents[0]->NoneInt);
+}
+
+UE_NET_TEST_FIXTURE(FTestConditionalsFixture, NeverIsNotReplicated)
+{
+	// Add client
+	FReplicationSystemTestClient* Client = CreateClient();
+
+	// Spawn object on server
+	constexpr uint32 ConnectionFilteredComponentCount = 1;
+	UTestReplicatedIrisObject::FComponents ObjectComponents;
+	ObjectComponents.ConnectionFilteredComponentCount = ConnectionFilteredComponentCount;
+	UTestReplicatedIrisObject* ServerObject = Server->CreateObject(ObjectComponents);
+
+	const int32 ExpectedClientValue = ServerObject->ConnectionFilteredComponents[0]->NeverInt;
+
+	// Set value
+	ServerObject->ConnectionFilteredComponents[0]->NeverInt = 13;
+
+	// Could set autonomous condition to be some connection other than the create client, or not set at all as is done now.
+
+	// Send and deliver packet
+	Server->PreSendUpdate();
+	Server->SendAndDeliverTo(Client, DeliverPacket);
+	Server->PostSendUpdate();
+	
+	// Object should have been created on the client
+	const UTestReplicatedIrisObject* ClientObject = Cast<UTestReplicatedIrisObject>(Client->GetReplicationBridge()->GetReplicatedObject(ServerObject->NetRefHandle));
+	UE_NET_ASSERT_NE(ClientObject, nullptr);
+
+	// Check that the COND_Never member remains unmodified.
+	UE_NET_ASSERT_EQ(ClientObject->ConnectionFilteredComponents[0]->NeverInt, ExpectedClientValue);
+}
+
+UE_NET_TEST_FIXTURE(FTestConditionalsFixture, SkipReplayIsReplicatedToRegularConnection)
+{
+	// Add client
+	FReplicationSystemTestClient* Client = CreateClient();
+
+	// Spawn object on server
+	constexpr uint32 ConnectionFilteredComponentCount = 1;
+	UTestReplicatedIrisObject::FComponents ObjectComponents;
+	ObjectComponents.ConnectionFilteredComponentCount = ConnectionFilteredComponentCount;
+	UTestReplicatedIrisObject* ServerObject = Server->CreateObject(ObjectComponents);
+
+	// Set value
+	ServerObject->ConnectionFilteredComponents[0]->SkipReplayInt = 13;
+
+	// Send and deliver packet
+	Server->PreSendUpdate();
+	Server->SendAndDeliverTo(Client, DeliverPacket);
+	Server->PostSendUpdate();
+	
+	// Object should have been created on the client
+	const UTestReplicatedIrisObject* ClientObject = Cast<UTestReplicatedIrisObject>(Client->GetReplicationBridge()->GetReplicatedObject(ServerObject->NetRefHandle));
+	UE_NET_ASSERT_NE(ClientObject, nullptr);
+
+	// Check that the COND_SkipReplay member has been modified
+	UE_NET_ASSERT_EQ(ClientObject->ConnectionFilteredComponents[0]->SkipReplayInt, ServerObject->ConnectionFilteredComponents[0]->SkipReplayInt);
+}
+
+UE_NET_TEST_FIXTURE(FTestConditionalsFixture, ReplayOnlyIsNotReplicatedToRegularConnection)
+{
+	// Add client
+	FReplicationSystemTestClient* Client = CreateClient();
+
+	// Spawn object on server
+	constexpr uint32 ConnectionFilteredComponentCount = 1;
+	UTestReplicatedIrisObject::FComponents ObjectComponents;
+	ObjectComponents.ConnectionFilteredComponentCount = ConnectionFilteredComponentCount;
+	UTestReplicatedIrisObject* ServerObject = Server->CreateObject(ObjectComponents);
+
+	const int32 ExpectedClientValue = ServerObject->ConnectionFilteredComponents[0]->NeverInt;
+
+	// Set value
+	ServerObject->ConnectionFilteredComponents[0]->ReplayOnlyInt = 13;
+
+	// Could set autonomous condition to be some connection other than the create client, or not set at all as is done now.
+
+	// Send and deliver packet
+	Server->PreSendUpdate();
+	Server->SendAndDeliverTo(Client, DeliverPacket);
+	Server->PostSendUpdate();
+	
+	// Object should have been created on the client
+	const UTestReplicatedIrisObject* ClientObject = Cast<UTestReplicatedIrisObject>(Client->GetReplicationBridge()->GetReplicatedObject(ServerObject->NetRefHandle));
+	UE_NET_ASSERT_NE(ClientObject, nullptr);
+
+	// Check that the COND_SkipReplay member has been modified
+	UE_NET_ASSERT_EQ(ClientObject->ConnectionFilteredComponents[0]->ReplayOnlyInt, ExpectedClientValue);
+}
 
 }
diff --git a/Engine/Plugins/Runtime/ReplicationSystemTestPlugin/Source/Private/Tests/ReplicationSystem/ReplicatedTestObject.cpp b/Engine/Plugins/Runtime/ReplicationSystemTestPlugin/Source/Private/Tests/ReplicationSystem/ReplicatedTestObject.cpp
--- a/Engine/Plugins/Runtime/ReplicationSystemTestPlugin/Source/Private/Tests/ReplicationSystem/ReplicatedTestObject.cpp
+++ b/Engine/Plugins/Runtime/ReplicationSystemTestPlugin/Source/Private/Tests/ReplicationSystem/ReplicatedTestObject.cpp
@@ -295,6 +295,7 @@ void UTestReplicatedIrisLifetimeConditionalsPropertyState::GetLifetimeReplicated
 {
 	DOREPLIFETIME_CONDITION(ThisClass, ToOwnerA, COND_OwnerOnly);
 	DOREPLIFETIME_CONDITION(ThisClass, ToOwnerB, COND_OwnerOnly);
+	DOREPLIFETIME_CONDITION(ThisClass, ReplayOrOwner, COND_ReplayOrOwner);
 
 	DOREPLIFETIME_CONDITION(ThisClass, SkipOwnerA, COND_SkipOwner);
 	DOREPLIFETIME_CONDITION(ThisClass, SkipOwnerB, COND_SkipOwner);
@@ -304,6 +305,10 @@ void UTestReplicatedIrisLifetimeConditionalsPropertyState::GetLifetimeReplicated
 	DOREPLIFETIME_CONDITION(ThisClass, SimulatedOrPhysicsInt, COND_SimulatedOrPhysics);
 	DOREPLIFETIME_CONDITION(ThisClass, SimulatedOnlyNoReplayInt, COND_SimulatedOnlyNoReplay);
 	DOREPLIFETIME_CONDITION(ThisClass, SimulatedOrPhysicsNoReplayInt, COND_SimulatedOrPhysicsNoReplay);
+	DOREPLIFETIME_CONDITION(ThisClass, NoneInt, COND_None);
+	DOREPLIFETIME_CONDITION(ThisClass, NeverInt, COND_Never);
+	DOREPLIFETIME_CONDITION(ThisClass, SkipReplayInt, COND_SkipReplay);
+	DOREPLIFETIME_CONDITION(ThisClass, ReplayOnlyInt, COND_ReplayOnly);
 }
 
 //////////////////////////////////////////////////////////////////////////
diff --git a/Engine/Plugins/Runtime/ReplicationSystemTestPlugin/Source/Private/Tests/ReplicationSystem/ReplicatedTestObject.h b/Engine/Plugins/Runtime/ReplicationSystemTestPlugin/Source/Private/Tests/ReplicationSystem/ReplicatedTestObject.h
--- a/Engine/Plugins/Runtime/ReplicationSystemTestPlugin/Source/Private/Tests/ReplicationSystem/ReplicatedTestObject.h
+++ b/Engine/Plugins/Runtime/ReplicationSystemTestPlugin/Source/Private/Tests/ReplicationSystem/ReplicatedTestObject.h
@@ -226,6 +226,8 @@ public:
 	int32 ToOwnerA = 0;
 	UPROPERTY(Replicated, Transient)
 	int32 ToOwnerB = 0;
+	UPROPERTY(Replicated, Transient)
+	int32 ReplayOrOwner = 0;
 
 	UPROPERTY(Replicated, Transient)
 	int32 SkipOwnerA = 0;
@@ -242,6 +244,14 @@ public:
 	int32 SimulatedOnlyNoReplayInt = 0;
 	UPROPERTY(Replicated, Transient)
 	int32 SimulatedOrPhysicsNoReplayInt = 0;
+	UPROPERTY(Replicated, Transient)
+	int32 NoneInt = 0;
+	UPROPERTY(Replicated, Transient)
+	int32 NeverInt = 0;
+	UPROPERTY(Replicated, Transient)
+	int32 SkipReplayInt = 0;
+	UPROPERTY(Replicated, Transient)
+	int32 ReplayOnlyInt = 0;
 
 	// Network data only for test
 	TArray<UE::Net::FReplicationFragment*> ReplicationFragments;
diff --git a/Engine/Source/Runtime/Experimental/Iris/Core/Private/Iris/ReplicationState/ReplicationStateDescriptorBuilder.cpp b/Engine/Source/Runtime/Experimental/Iris/Core/Private/Iris/ReplicationState/ReplicationStateDescriptorBuilder.cpp
--- a/Engine/Source/Runtime/Experimental/Iris/Core/Private/Iris/ReplicationState/ReplicationStateDescriptorBuilder.cpp
+++ b/Engine/Source/Runtime/Experimental/Iris/Core/Private/Iris/ReplicationState/ReplicationStateDescriptorBuilder.cpp
@@ -1867,6 +1867,7 @@ EMemberPropertyTraits FPropertyReplicationStateDescriptorBuilder::GetConnectionF
 		{COND_ReplayOnly, EMemberPropertyTraits::HasLifetimeConditionals},
 		{COND_SimulatedOnlyNoReplay, EMemberPropertyTraits::HasLifetimeConditionals},
 		{COND_SimulatedOrPhysicsNoReplay, EMemberPropertyTraits::HasLifetimeConditionals},
+		{COND_SkipReplay, EMemberPropertyTraits::HasLifetimeConditionals},
 		{COND_Dynamic, EMemberPropertyTraits::HasLifetimeConditionals},
 	};
 
