Iris
* Move PreUpdate into its own loop and call it on all polling objects.
* Properly manage objects that set dirtyness on other objects inside PreUpdate/PreReplication
   * Now objects that get set dirty inside PreReplication but are not set to be polled will still be set dirty on the next frame.
* Deleted code that tested if such dependencies existed

#rnx

diff --git a/Engine/Plugins/Runtime/ReplicationSystemTestPlugin/Source/Private/Tests/ReplicationSystem/TestDirtyNetObjectTracker.cpp b/Engine/Plugins/Runtime/ReplicationSystemTestPlugin/Source/Private/Tests/ReplicationSystem/TestDirtyNetObjectTracker.cpp
--- a/Engine/Plugins/Runtime/ReplicationSystemTestPlugin/Source/Private/Tests/ReplicationSystem/TestDirtyNetObjectTracker.cpp
+++ b/Engine/Plugins/Runtime/ReplicationSystemTestPlugin/Source/Private/Tests/ReplicationSystem/TestDirtyNetObjectTracker.cpp
@@ -86,7 +86,7 @@ UE_NET_TEST_FIXTURE(FDirtyNetObjectTrackerTestFixture, CanClearDirtyObjects)
 	CleanedObjects.SetBit(NetObjectIndexRangeEnd);
 
 	DirtyNetObjectTracker->UpdateAccumulatedDirtyList();
-	DirtyNetObjectTracker->ClearDirtyNetObjects(MakeNetBitArrayView(CleanedObjects));
+	DirtyNetObjectTracker->ReconcilePolledList(MakeNetBitArrayView(CleanedObjects));
 
 	const FNetBitArrayView AccumulatedDirtyObjects = DirtyNetObjectTracker->GetAccumulatedDirtyNetObjects();
 	UE_NET_ASSERT_FALSE(AccumulatedDirtyObjects.IsAnyBitSet());
@@ -107,7 +107,7 @@ UE_NET_TEST_FIXTURE(FDirtyNetObjectTrackerTestFixture, DelayedDirtyBitTracking)
 		CleanedObjects.SetBit(FirstObjectIndex);
 
 		DirtyNetObjectTracker->UpdateAccumulatedDirtyList();
-		DirtyNetObjectTracker->ClearDirtyNetObjects(MakeNetBitArrayView(CleanedObjects));
+		DirtyNetObjectTracker->ReconcilePolledList(MakeNetBitArrayView(CleanedObjects));
 	}
 
 	const FNetBitArrayView AccumulatedDirtyObjects = DirtyNetObjectTracker->GetAccumulatedDirtyNetObjects();
@@ -122,7 +122,7 @@ UE_NET_TEST_FIXTURE(FDirtyNetObjectTrackerTestFixture, DelayedDirtyBitTracking)
 		CleanedObjects.SetBit(SecondObjectIndex);
 
 		DirtyNetObjectTracker->UpdateAccumulatedDirtyList();
-		DirtyNetObjectTracker->ClearDirtyNetObjects(MakeNetBitArrayView(CleanedObjects));
+		DirtyNetObjectTracker->ReconcilePolledList(MakeNetBitArrayView(CleanedObjects));
 	}
 
 	UE_NET_ASSERT_FALSE(AccumulatedDirtyObjects.GetBit(FirstObjectIndex));
@@ -180,7 +180,7 @@ UE_NET_TEST_FIXTURE(FReplicationSystemServerClientTestFixture, GlobalDirtyTracke
 	Server->DestroyObject(ServerObject);
 }
 
-/* This test exposes a problem when an object's PreUpdate calls dirty on a different object.
+/** Test that validates behavior when dirtying other actors inside PreUpdate/PreReplication */
 UE_NET_TEST_FIXTURE(FReplicationSystemServerClientTestFixture, DirtyInsidePreUpdateTest)
 {
 	// Add client
@@ -188,7 +188,7 @@ UE_NET_TEST_FIXTURE(FReplicationSystemServerClientTestFixture, DirtyInsidePreUpd
 
 	// Spawn object on server that is polled late in order to test ForceNetUpdate
 	UObjectReplicationBridge::FCreateNetRefHandleParams Params;
-	const uint32 PollPeriod = 10;
+	const uint32 PollPeriod = 100;
 	const float PollFrequency = Server->ConvertPollPeriodIntoFrequency(PollPeriod);
 	Params.PollFrequency = PollFrequency;
 	Params.bCanReceive = true;
@@ -206,13 +206,12 @@ UE_NET_TEST_FIXTURE(FReplicationSystemServerClientTestFixture, DirtyInsidePreUpd
 	UTestReplicatedIrisObject* ClientObjectB = Cast<UTestReplicatedIrisObject>(Client->GetReplicationBridge()->GetReplicatedObject(ServerObjectB->NetRefHandle));
 	UE_NET_ASSERT_NE(ClientObjectB, nullptr);
 	
-
 	// Send and deliver packet twice
 	Server->UpdateAndSend({ Client });
 	Server->UpdateAndSend({ Client });
 
 	// Set a replicated variable, but don't mark it dirty
-	ServerObjectA->IntA = 0xFF;
+	ServerObjectA->IntA = 0xAA;
 
 	// Send and deliver packet
 	Server->UpdateAndSend({ Client });
@@ -221,7 +220,7 @@ UE_NET_TEST_FIXTURE(FReplicationSystemServerClientTestFixture, DirtyInsidePreUpd
 	UE_NET_ASSERT_NE(ClientObjectA->IntA, ServerObjectA->IntA);
 
 	// Set a replicated variable, but don't mark it dirty
-	ServerObjectB->IntA = 0xFF;
+	ServerObjectB->IntA = 0xAA;
 
 	// Send and deliver packet
 	Server->UpdateAndSend({ Client });
@@ -230,34 +229,70 @@ UE_NET_TEST_FIXTURE(FReplicationSystemServerClientTestFixture, DirtyInsidePreUpd
 	UE_NET_ASSERT_NE(ClientObjectA->IntA, ServerObjectA->IntA);
 	UE_NET_ASSERT_NE(ClientObjectB->IntA, ServerObjectB->IntA);
 	
-	// Make the first object dirty
+	// Force ObjectA to replicate
 	Server->ReplicationSystem->ForceNetUpdate(ServerObjectA->NetRefHandle);
 
-	auto PreUpdate = [&](FNetRefHandle NetHandle, UObject* ReplicatedObject, const UReplicationBridge* ReplicationBridge)
+	auto PreUpdateObjectA = [&](FNetRefHandle NetHandle, UObject* ReplicatedObject, const UReplicationBridge* ReplicationBridge)
 	{
-		// When ObjectA is updated, make ObjectB dirty
+		// Inside ObjectA PreReplicationUpdate, force ObjectB to be replicated
 		if (ServerObjectA == ReplicatedObject)
 		{
-			// This will cause an ensure now
-			Server->ReplicationSystem->MarkDirty(ServerObjectB->NetRefHandle);
+			Server->ReplicationSystem->ForceNetUpdate(ServerObjectB->NetRefHandle);
 		}
 	};
 
 	// Now add a dependency where the poll of the first object makes the second one dirty
-	Server->GetReplicationBridge()->SetExternalPreUpdateFunctor(PreUpdate);
+	Server->GetReplicationBridge()->SetExternalPreUpdateFunctor(PreUpdateObjectA);
 
 	// Send and deliver packet
 	Server->UpdateAndSend({ Client });
 
-	// Client property should have changed now
+	// ObjectA should be replicated now
 	UE_NET_ASSERT_EQ(ClientObjectA->IntA, ServerObjectA->IntA);
 
-	// This fails because MarkDirty of a different object in PreUpdate is ignored and flushed.
+	// But not ObjectB because it was called inside PreUpdate
+	UE_NET_ASSERT_NE(ClientObjectB->IntA, ServerObjectB->IntA);
+
+	// Send and deliver packet
+	Server->UpdateAndSend({ Client });
+
+	// Now the ForceNetUpdate on ObjectB is applied and it is replicated
 	UE_NET_ASSERT_EQ(ClientObjectB->IntA, ServerObjectB->IntA);
 
+	auto PreUpdateObjectB = [&](FNetRefHandle NetHandle, UObject* ReplicatedObject, const UReplicationBridge* ReplicationBridge)
+	{
+		// Inside ObjectB PreReplicationUpdate, force ObjectA to be replicated
+		if (ServerObjectB == ReplicatedObject)
+		{
+			Server->ReplicationSystem->ForceNetUpdate(ServerObjectA->NetRefHandle);
+		}
+	};
+	Server->GetReplicationBridge()->SetExternalPreUpdateFunctor(PreUpdateObjectB);
+
+	// Dirty both Objects
+	ServerObjectA->IntA = 0xBB;
+	ServerObjectB->IntA = 0xBB;
+
+	// But only force update ObjectB
+	Server->ReplicationSystem->ForceNetUpdate(ServerObjectB->NetRefHandle);
+
+	// Send and deliver packet
+	Server->UpdateAndSend({ Client });
+
+	// ObjectB should have replicated
+	UE_NET_ASSERT_EQ(ClientObjectB->IntA, ServerObjectB->IntA);
+
+	// But not ObjectA because it was forced inside PreReplicate()
+	UE_NET_ASSERT_NE(ClientObjectA->IntA, ServerObjectA->IntA);
+
+	// Send and deliver packet
+	Server->UpdateAndSend({ Client });
+
+	// Now ObjectA's force net update is applied and ObjectA is replicated
+	UE_NET_ASSERT_EQ(ClientObjectA->IntA, ServerObjectA->IntA);
+
 	Server->DestroyObject(ServerObjectA);
 	Server->DestroyObject(ServerObjectB);
-}*/
-
+}
 
 }
diff --git a/Engine/Source/Runtime/Experimental/Iris/Core/Private/Iris/ReplicationSystem/DirtyNetObjectTracker.cpp b/Engine/Source/Runtime/Experimental/Iris/Core/Private/Iris/ReplicationSystem/DirtyNetObjectTracker.cpp
--- a/Engine/Source/Runtime/Experimental/Iris/Core/Private/Iris/ReplicationSystem/DirtyNetObjectTracker.cpp
+++ b/Engine/Source/Runtime/Experimental/Iris/Core/Private/Iris/ReplicationSystem/DirtyNetObjectTracker.cpp
@@ -30,9 +30,6 @@ namespace UE::Net::Private
 FDirtyNetObjectTracker::FDirtyNetObjectTracker()
 : ReplicationSystemId(InvalidReplicationSystemId)
 {
-#if UE_NET_IRIS_VALIDATE_POLLED_OBJECT
-	CurrentPolledObject = FNetRefHandleManager::InvalidInternalIndex;
-#endif
 }
 
 FDirtyNetObjectTracker::~FDirtyNetObjectTracker()
@@ -138,12 +135,6 @@ void FDirtyNetObjectTracker::MarkNetObjectDirty(FInternalNetRefIndex NetObjectIn
 
 	if ((NetObjectIndex >= NetObjectIdRangeStart) & (NetObjectIndex <= NetObjectIdRangeEnd))
 	{
-#if UE_NET_IRIS_VALIDATE_POLLED_OBJECT
-		ensureMsgf(FNetRefHandleManager::InvalidInternalIndex == CurrentPolledObject || CurrentPolledObject == NetObjectIndex, 
-			TEXT("While calling NetUpdate on %s, a different replicated object %s was dirtied. Only the updated object can be dirtied for now."),
-			*GetNameSafe(NetRefHandleManager->GetReplicatedObjectInstance(CurrentPolledObject)), *GetNameSafe(NetRefHandleManager->GetReplicatedObjectInstance(NetObjectIndex)));
-#endif
-
 		const uint32 BitOffset = NetObjectIndex;
 		const StorageType BitMask = StorageType(1) << (BitOffset & (StorageTypeBitCount - 1));
 
@@ -186,7 +177,7 @@ FNetBitArrayView FDirtyNetObjectTracker::GetDirtyNetObjectsThisFrame()
 	return FNetBitArrayView(DirtyNetObjectContainer, NetObjectIdCount);
 }
 
-void FDirtyNetObjectTracker::ClearDirtyNetObjects(const FNetBitArrayView& CleanNetObjects)
+void FDirtyNetObjectTracker::ReconcilePolledList(const FNetBitArrayView& ObjectsPolled)
 {
 	LockExternalAccess();
 
@@ -196,18 +187,15 @@ void FDirtyNetObjectTracker::ClearDirtyNetObjects(const FNetBitArrayView& CleanN
 		FGlobalDirtyNetObjectTracker::ResetDirtyNetObjects(GlobalDirtyTrackerPollHandle);
 	}
 
-	// Keep dirty flags for objects not cleaned this frame.
-	MakeNetBitArrayView(AccumulatedDirtyNetObjects).Combine(CleanNetObjects, FNetBitArrayView::AndNotOp);
+	// Clear ForceNetUpdate from every object that were polled.
+	MakeNetBitArrayView(ForceNetUpdateObjects).Combine(ObjectsPolled, FNetBitArrayView::AndNotOp);
+
+	// Clear dirty flags for objects that were polled
+	MakeNetBitArrayView(AccumulatedDirtyNetObjects).Combine(ObjectsPolled, FNetBitArrayView::AndNotOp);
 
 	// Clear the current frame dirty objects
 	FMemory::Memzero(DirtyNetObjectContainer, DirtyNetObjectWordCount*sizeof(StorageType));
 
-	ForceNetUpdateObjects.Reset();
-
-#if UE_NET_IRIS_VALIDATE_POLLED_OBJECT
-	CurrentPolledObject = 0;
-#endif
-
 	AllowExternalAccess();
 
 	std::atomic_thread_fence(std::memory_order_seq_cst);
diff --git a/Engine/Source/Runtime/Experimental/Iris/Core/Private/Iris/ReplicationSystem/DirtyNetObjectTracker.h b/Engine/Source/Runtime/Experimental/Iris/Core/Private/Iris/ReplicationSystem/DirtyNetObjectTracker.h
--- a/Engine/Source/Runtime/Experimental/Iris/Core/Private/Iris/ReplicationSystem/DirtyNetObjectTracker.h
+++ b/Engine/Source/Runtime/Experimental/Iris/Core/Private/Iris/ReplicationSystem/DirtyNetObjectTracker.h
@@ -15,11 +15,6 @@ namespace UE::Net::Private
 	typedef uint32 FInternalNetRefIndex;
 }
 
-// When enabled will detect if a polled object is dirtying another object. This is an unsupported behavior
-#ifndef UE_NET_IRIS_VALIDATE_POLLED_OBJECT
-#define UE_NET_IRIS_VALIDATE_POLLED_OBJECT !UE_BUILD_SHIPPING
-#endif
-
 namespace UE::Net::Private
 {
 
@@ -55,11 +50,8 @@ public:
 	/** Release safety permissions and allow to write in the bit array via the public methods */
 	void AllowExternalAccess();
 
-	/** Reset the global and local dirty objects lists for those objects that are now clean */
-	void ClearDirtyNetObjects(const FNetBitArrayView& CleanNetObjects);
-
-	/** Track which object is currently being polled */
-	inline void SetCurrentPolledObject(FInternalNetRefIndex PolledObject);
+	/** Reset the global list and look at the final polled list and clear any flags for objects that got polled */
+	void ReconcilePolledList(const FNetBitArrayView& ObjectsPolled);
 
 	/** Returns the list of objects that are dirty this frame or were dirty in previous frames but not cleaned up at that time. */
 	const FNetBitArrayView GetAccumulatedDirtyNetObjects() const { return MakeNetBitArrayView(AccumulatedDirtyNetObjects); }
@@ -99,10 +91,6 @@ private:
 	
 	FGlobalDirtyNetObjectTracker::FPollHandle GlobalDirtyTrackerPollHandle;
 
-#if UE_NET_IRIS_VALIDATE_POLLED_OBJECT
-	FInternalNetRefIndex CurrentPolledObject;
-#endif
-
 	uint32 ReplicationSystemId;
 
 	uint32 DirtyNetObjectWordCount = 0;
@@ -117,14 +105,6 @@ private:
 #endif
 };
 
-
-void FDirtyNetObjectTracker::SetCurrentPolledObject(FInternalNetRefIndex PolledObject)
-{ 
-#if UE_NET_IRIS_VALIDATE_POLLED_OBJECT
-	CurrentPolledObject = PolledObject; 
-#endif
-}
-
 /**
  * Gives access to the list of dirty objects while detecting non-thread safe access to it.
  */
diff --git a/Engine/Source/Runtime/Experimental/Iris/Core/Private/Iris/ReplicationSystem/ObjectReplicationBridge.cpp b/Engine/Source/Runtime/Experimental/Iris/Core/Private/Iris/ReplicationSystem/ObjectReplicationBridge.cpp
--- a/Engine/Source/Runtime/Experimental/Iris/Core/Private/Iris/ReplicationSystem/ObjectReplicationBridge.cpp
+++ b/Engine/Source/Runtime/Experimental/Iris/Core/Private/Iris/ReplicationSystem/ObjectReplicationBridge.cpp
@@ -771,6 +771,7 @@ void UObjectReplicationBridge::PreUpdateAndPoll()
 	// Update every relevant objects from here
 	FReplicationSystemInternal* ReplicationSystemInternal = GetReplicationSystem()->GetReplicationSystemInternal();
 	const FNetRefHandleManager& LocalNetRefHandleManager = ReplicationSystemInternal->GetNetRefHandleManager();
+	FDirtyNetObjectTracker& DirtyNetObjectTracker = ReplicationSystemInternal->GetDirtyNetObjectTracker();
 	const FNetBitArrayView RelevantObjects = LocalNetRefHandleManager.GetRelevantObjectsInternalIndices();
 	const FNetBitArrayView WantToBeDormantObjects = MakeNetBitArrayView(LocalNetRefHandleManager.GetWantToBeDormantInternalIndices());
 
@@ -778,18 +779,18 @@ void UObjectReplicationBridge::PreUpdateAndPoll()
 	FNetBitArrayView ObjectsConsideredForPolling = LocalNetRefHandleManager.GetPolledObjectsInternalIndices();
 	ObjectsConsideredForPolling.Reset();
 
-	// We always want to consider objects marked as dirty and their subobjects
-	const FNetBitArrayView AccumulatedDirtyObjects = ReplicationSystemInternal->GetDirtyNetObjectTracker().GetAccumulatedDirtyNetObjects();
-
 	if (bUseFrequencyBasedPolling)
 	{
 		if (bEnableForceNetUpdate)
 		{
-			const FNetBitArrayView ForceNetUpdateObjects = ReplicationSystemInternal->GetDirtyNetObjectTracker().GetForceNetUpdateObjects();
+			// Find objects ready to be polled and add objects that called ForceNetupdate
+			const FNetBitArrayView ForceNetUpdateObjects = DirtyNetObjectTracker.GetForceNetUpdateObjects();
 			PollFrequencyLimiter->Update(RelevantObjects, ForceNetUpdateObjects, ObjectsConsideredForPolling);
 		}
 		else
 		{
+			// Find objects ready to be polled and add objects that were flagged Dirty.
+			const FNetBitArrayView AccumulatedDirtyObjects = DirtyNetObjectTracker.GetAccumulatedDirtyNetObjects();
 			PollFrequencyLimiter->Update(RelevantObjects, AccumulatedDirtyObjects, ObjectsConsideredForPolling);
 		}
 	}
@@ -804,6 +805,7 @@ void UObjectReplicationBridge::PreUpdateAndPoll()
 		IRIS_PROFILER_SCOPE(PreUpdateAndPoll_Dormancy);
 
 		// Mask off objects pending dormancy that are not dirty
+		const FNetBitArrayView AccumulatedDirtyObjects = DirtyNetObjectTracker.GetAccumulatedDirtyNetObjects();
 		ObjectsConsideredForPolling.CombineMultiple(FNetBitArrayView::AndNotOp, WantToBeDormantObjects, FNetBitArrayView::AndNotOp, AccumulatedDirtyObjects);
 
 		FNetBitArrayView ForceNetUpdateObjects = ReplicationSystemInternal->GetDirtyNetObjectTracker().GetForceNetUpdateObjects();
@@ -850,7 +852,7 @@ void UObjectReplicationBridge::PreUpdateAndPoll()
 
 		// Update subobjects' owner first and owners' subobjects second. It's the only way to properly mark all groups of objects in two passes.
 		const FNetBitArrayView SubObjects = MakeNetBitArrayView(LocalNetRefHandleManager.GetSubObjectInternalIndices());
-		const FNetBitArrayView ForceNetUpdateObjects = ReplicationSystemInternal->GetDirtyNetObjectTracker().GetForceNetUpdateObjects();
+		const FNetBitArrayView ForceNetUpdateObjects = DirtyNetObjectTracker.GetForceNetUpdateObjects();
 
 		if (bEnableForceNetUpdate)
 		{
@@ -899,23 +901,28 @@ void UObjectReplicationBridge::PreUpdateAndPoll()
 		}
 	}
 	
-	{
-		IRIS_PROFILER_SCOPE(PreUpdateAndPoll_Poll);
+	FObjectPoller::FInitParams PollerInitParams;
+	PollerInitParams.ObjectReplicationBridge = this;
+	PollerInitParams.ReplicationSystemInternal = GetReplicationSystem()->GetReplicationSystemInternal();
 
-		FObjectPoller::FInitParams PollerInitParams;
-		PollerInitParams.ObjectReplicationBridge = this;
-		PollerInitParams.ReplicationSystemInternal = GetReplicationSystem()->GetReplicationSystemInternal();
+	FObjectPoller Poller(PollerInitParams);
 
-		FObjectPoller Poller(PollerInitParams);
+	{
+		IRIS_PROFILER_SCOPE(PreUpdateAndPoll_PreUpdate);
+		Poller.PreUpdatePass(ObjectsConsideredForPolling);
+	}
+
+	{
+		IRIS_PROFILER_SCOPE(PreUpdateAndPoll_Poll);
 		Poller.PollObjects(ObjectsConsideredForPolling);
+	}
 	
-		FObjectPoller::FPreUpdateAndPollStats Stats = Poller.GetPollStats();
+	FObjectPoller::FPreUpdateAndPollStats Stats = Poller.GetPollStats();
 
-		// Report stats
-		UE_NET_TRACE_FRAME_STATSCOUNTER(ReplicationSystem->GetId(), ReplicationSystem.PreUpdatedObjectCount, Stats.PreUpdatedObjectCount, ENetTraceVerbosity::Trace);
-		UE_NET_TRACE_FRAME_STATSCOUNTER(ReplicationSystem->GetId(), ReplicationSystem.PolledObjectCount, Stats.PolledObjectCount, ENetTraceVerbosity::Trace);
-		UE_NET_TRACE_FRAME_STATSCOUNTER(ReplicationSystem->GetId(), ReplicationSystem.PolledReferencesObjectCount, Stats.PolledReferencesObjectCount, ENetTraceVerbosity::Trace);
-	}
+	// Report stats
+	UE_NET_TRACE_FRAME_STATSCOUNTER(ReplicationSystem->GetId(), ReplicationSystem.PreUpdatedObjectCount, Stats.PreUpdatedObjectCount, ENetTraceVerbosity::Trace);
+	UE_NET_TRACE_FRAME_STATSCOUNTER(ReplicationSystem->GetId(), ReplicationSystem.PolledObjectCount, Stats.PolledObjectCount, ENetTraceVerbosity::Trace);
+	UE_NET_TRACE_FRAME_STATSCOUNTER(ReplicationSystem->GetId(), ReplicationSystem.PolledReferencesObjectCount, Stats.PolledReferencesObjectCount, ENetTraceVerbosity::Trace);
 }
 
 void UObjectReplicationBridge::UpdateInstancesWorldLocation()
diff --git a/Engine/Source/Runtime/Experimental/Iris/Core/Private/Iris/ReplicationSystem/Polling/ObjectPoller.cpp b/Engine/Source/Runtime/Experimental/Iris/Core/Private/Iris/ReplicationSystem/Polling/ObjectPoller.cpp
--- a/Engine/Source/Runtime/Experimental/Iris/Core/Private/Iris/ReplicationSystem/Polling/ObjectPoller.cpp
+++ b/Engine/Source/Runtime/Experimental/Iris/Core/Private/Iris/ReplicationSystem/Polling/ObjectPoller.cpp
@@ -30,16 +30,42 @@ FObjectPoller::FObjectPoller(const FInitParams& InitParams)
 	// DirtyObjectsThisFrame is acquired only during polling 
 }
 
+void FObjectPoller::PreUpdatePass(const FNetBitArrayView& ObjectsConsideredForPolling)
+{
+	IRIS_PROFILER_SCOPE_VERBOSE(PreUpdatePass);
+
+	ObjectsConsideredForPolling.ForAllSetBits([this](FInternalNetRefIndex Objectindex)
+	{
+		CallPreUpdate(Objectindex);
+	});
+}
+
+void FObjectPoller::CallPreUpdate(FInternalNetRefIndex ObjectIndex)
+{
+	FNetRefHandleManager::FReplicatedObjectData& ObjectData = LocalNetRefHandleManager.GetReplicatedObjectDataNoCheck(ObjectIndex);
+	if (UNLIKELY(ObjectData.InstanceProtocol == nullptr))
+	{
+		return;
+	}
+
+	IRIS_PROFILER_PROTOCOL_NAME(ObjectData.Protocol->DebugName->Name);
+
+	// Call per-instance PreUpdate function
+	if (ObjectReplicationBridge->PreUpdateInstanceFunction && EnumHasAnyFlags(ObjectData.InstanceProtocol->InstanceTraits, EReplicationInstanceProtocolTraits::NeedsPreSendUpdate))
+	{
+		ObjectReplicationBridge->PreUpdateInstanceFunction(ObjectData.RefHandle, ReplicatedInstances[ObjectIndex], ObjectReplicationBridge);
+		++PollStats.PreUpdatedObjectCount;
+	}
+}
+
 void FObjectPoller::PollObjects(const FNetBitArrayView& ObjectsConsideredForPolling)
 {
 	FDirtyObjectsAccessor DirtyObjectsAccessor(ReplicationSystemInternal->GetDirtyNetObjectTracker());
 	DirtyObjectsThisFrame = DirtyObjectsAccessor.GetDirtyNetObjects();
 
-	// From here we call into user code via PreUpdateInstanceFunction, so allow external code to set dirty flags since DirtyObjects is not read anymore.
-	ReplicationSystemInternal->GetDirtyNetObjectTracker().AllowExternalAccess();
-
 	if (IsIrisPushModelEnabled())
 	{
+		IRIS_PROFILER_SCOPE_VERBOSE(PollPushBased);
 		ObjectsConsideredForPolling.ForAllSetBits([this](FInternalNetRefIndex Objectindex)
 		{
 			PushModelPollObject(Objectindex);
@@ -47,29 +73,28 @@ void FObjectPoller::PollObjects(const FNetBitArrayView& ObjectsConsideredForPoll
 	}
 	else
 	{
+		IRIS_PROFILER_SCOPE_VERBOSE(ForcePoll);
 		ObjectsConsideredForPolling.ForAllSetBits([this](FInternalNetRefIndex Objectindex)
 		{
 			ForcePollObject(Objectindex);
 		});
 	}
-
-	// Clear ref to locked dirty bit array
-	DirtyObjectsThisFrame = FNetBitArrayView();
-
-	ReplicationSystemInternal->GetDirtyNetObjectTracker().SetCurrentPolledObject(FNetRefHandleManager::InvalidInternalIndex);
 }
 
 void FObjectPoller::PollSingleObject(FNetRefHandle Handle)
 {
 	if (uint32 InternalObjectIndex = LocalNetRefHandleManager.GetInternalIndex(Handle))
 	{
+		CallPreUpdate(InternalObjectIndex);
+
 		FDirtyObjectsAccessor DirtyObjectsAccessor(ReplicationSystemInternal->GetDirtyNetObjectTracker());
 		DirtyObjectsThisFrame = DirtyObjectsAccessor.GetDirtyNetObjects();
-
 		// From here we call into user code via PreUpdateInstanceFunction, so allow external code to set dirty flags since DirtyObjects is not read anymore.
 		ReplicationSystemInternal->GetDirtyNetObjectTracker().AllowExternalAccess();
 
-		ForcePollObject(InternalObjectIndex);
+		{
+			ForcePollObject(InternalObjectIndex);
+		}
 
 		// Clear ref to locked dirty bit array
 		DirtyObjectsThisFrame = FNetBitArrayView();
@@ -79,7 +104,7 @@ void FObjectPoller::PollSingleObject(FNetRefHandle Handle)
 void FObjectPoller::ForcePollObject(FInternalNetRefIndex ObjectIndex)
 {
 	FNetRefHandleManager::FReplicatedObjectData& ObjectData = LocalNetRefHandleManager.GetReplicatedObjectDataNoCheck(ObjectIndex);
-	if (ObjectData.InstanceProtocol == nullptr)
+	if (UNLIKELY(ObjectData.InstanceProtocol == nullptr))
 	{
 		return;
 	}
@@ -89,14 +114,6 @@ void FObjectPoller::ForcePollObject(FInternalNetRefIndex ObjectIndex)
 	// We always poll all states here.
 	ObjectData.bWantsFullPoll = 0U;
 
-	// Call per-instance PreUpdate function
-	if (ObjectReplicationBridge->PreUpdateInstanceFunction && EnumHasAnyFlags(ObjectData.InstanceProtocol->InstanceTraits, EReplicationInstanceProtocolTraits::NeedsPreSendUpdate))
-	{
-		IRIS_PROFILER_SCOPE_VERBOSE(PreReplicationUpdate);
-		ObjectReplicationBridge->PreUpdateInstanceFunction(ObjectData.RefHandle, ReplicatedInstances[ObjectIndex], ObjectReplicationBridge);
-		++PollStats.PreUpdatedObjectCount;
-	}
-
 	// Poll properties if the instance protocol requires it
 	if (EnumHasAnyFlags(ObjectData.InstanceProtocol->InstanceTraits, EReplicationInstanceProtocolTraits::NeedsPoll))
 	{
@@ -128,44 +145,21 @@ void FObjectPoller::ForcePollObject(FInternalNetRefIndex ObjectIndex)
 void FObjectPoller::PushModelPollObject(FInternalNetRefIndex ObjectIndex)
 {
 	FNetRefHandleManager::FReplicatedObjectData& ObjectData = LocalNetRefHandleManager.GetReplicatedObjectDataNoCheck(ObjectIndex);
-	if (ObjectData.InstanceProtocol == nullptr)
+	if (UNLIKELY(ObjectData.InstanceProtocol == nullptr))
 	{
 		return;
 	}
 
 	IRIS_PROFILER_PROTOCOL_NAME(ObjectData.Protocol->DebugName->Name);
 
-	ReplicationSystemInternal->GetDirtyNetObjectTracker().SetCurrentPolledObject(ObjectIndex);
-
 	const FReplicationInstanceProtocol* InstanceProtocol = ObjectData.InstanceProtocol;
 
 	const EReplicationInstanceProtocolTraits InstanceTraits = InstanceProtocol->InstanceTraits;
 	const bool bNeedsPoll = EnumHasAnyFlags(InstanceTraits, EReplicationInstanceProtocolTraits::NeedsPoll);
 
-	bool bIsDirtyObject = AccumulatedDirtyObjects.GetBit(ObjectIndex);
+	bool bIsDirtyObject = AccumulatedDirtyObjects.GetBit(ObjectIndex) || DirtyObjectsThisFrame.GetBit(ObjectIndex);
 
-	// Call per-instance PreUpdate function
-	if (ObjectReplicationBridge->PreUpdateInstanceFunction && EnumHasAnyFlags(InstanceTraits, EReplicationInstanceProtocolTraits::NeedsPreSendUpdate))
-	{
-		IRIS_PROFILER_SCOPE_VERBOSE(PreReplicationUpdate);
-
-		ObjectReplicationBridge->PreUpdateInstanceFunction(ObjectData.RefHandle, ReplicatedInstances[ObjectIndex], ObjectReplicationBridge);
-		++PollStats.PreUpdatedObjectCount;
-
-		// Pre update may have called MarkDirty. Detect it.
-		bIsDirtyObject = bIsDirtyObject || DirtyObjectsThisFrame.GetBit(ObjectIndex);
-		if (bNeedsPoll && !bIsDirtyObject)
-		{
-			bIsDirtyObject = FGlobalDirtyNetObjectTracker::IsNetObjectStateDirty(ObjectData.NetHandle);
-		}
-
-		if (bIsDirtyObject)
-		{
-			DirtyObjectsToCopy.SetBit(ObjectIndex);
-			DirtyObjectsThisFrame.SetBit(ObjectIndex);
-		}
-	}
-	else if (bIsDirtyObject)
+	if (bIsDirtyObject)
 	{
 		DirtyObjectsToCopy.SetBit(ObjectIndex);
 		DirtyObjectsThisFrame.SetBit(ObjectIndex);
diff --git a/Engine/Source/Runtime/Experimental/Iris/Core/Private/Iris/ReplicationSystem/Polling/ObjectPoller.h b/Engine/Source/Runtime/Experimental/Iris/Core/Private/Iris/ReplicationSystem/Polling/ObjectPoller.h
--- a/Engine/Source/Runtime/Experimental/Iris/Core/Private/Iris/ReplicationSystem/Polling/ObjectPoller.h
+++ b/Engine/Source/Runtime/Experimental/Iris/Core/Private/Iris/ReplicationSystem/Polling/ObjectPoller.h
@@ -48,6 +48,12 @@ public:
 
 	const FPreUpdateAndPollStats& GetPollStats() const { return PollStats; }
 
+	/**
+	 * Call the PreUpdateFunction, generally AActor::PreReplication, on polled objects that need it.
+	 * @param ObjectsConsideredForPolling The list of actors set to be polled this frame
+	 */
+	void PreUpdatePass(const FNetBitArrayView& ObjectsConsideredForPolling);
+
 	/** Poll all the objects whose bit index is set in the array */
 	void PollObjects(const FNetBitArrayView& ObjectsConsideredForPolling);
 
@@ -56,6 +62,9 @@ public:
 
 private:
 
+	/** Calls the PreUpdate function if the object needs it */
+	void CallPreUpdate(FInternalNetRefIndex ObjectIndex);
+
 	/** Polls an object in every circumstance */
 	void ForcePollObject(FInternalNetRefIndex ObjectIndex);
 
diff --git a/Engine/Source/Runtime/Experimental/Iris/Core/Private/Iris/ReplicationSystem/ReplicationSystem.cpp b/Engine/Source/Runtime/Experimental/Iris/Core/Private/Iris/ReplicationSystem/ReplicationSystem.cpp
--- a/Engine/Source/Runtime/Experimental/Iris/Core/Private/Iris/ReplicationSystem/ReplicationSystem.cpp
+++ b/Engine/Source/Runtime/Experimental/Iris/Core/Private/Iris/ReplicationSystem/ReplicationSystem.cpp
@@ -379,17 +379,16 @@ public:
 
 		FNetRefHandleManager& NetRefHandleManager = ReplicationSystemInternal.GetNetRefHandleManager();
 
-		// Clean the objects that got polled this frame
-		const FNetBitArrayView ObjectsToClean = NetRefHandleManager.GetPolledObjectsInternalIndices();
+		// Clear the objects that got polled this frame
+		const FNetBitArrayView PolledObjects = NetRefHandleManager.GetPolledObjectsInternalIndices();
 
 		// Reset object dirtyness
-		ObjectsToClean.ForAllSetBits([&NetRefHandleManager](uint32 DirtyIndex)
+		PolledObjects.ForAllSetBits([&NetRefHandleManager](uint32 DirtyIndex)
 		{
 			FReplicationInstanceOperationsInternal::ResetObjectStateDirtiness(NetRefHandleManager, DirtyIndex);
 		});
 
-		// Reset cleaned objects in the tracker
-		ReplicationSystemInternal.GetDirtyNetObjectTracker().ClearDirtyNetObjects(ObjectsToClean);
+		ReplicationSystemInternal.GetDirtyNetObjectTracker().ReconcilePolledList(PolledObjects);
 	}
 
 	void ProcessNetObjectAttachmentSendQueue(FNetBlobManager::EProcessMode ProcessMode)
diff --git a/Engine/Source/Runtime/Net/Core/Private/Net/Core/DirtyNetObjectTracker/GlobalDirtyNetObjectTracker.cpp b/Engine/Source/Runtime/Net/Core/Private/Net/Core/DirtyNetObjectTracker/GlobalDirtyNetObjectTracker.cpp
--- a/Engine/Source/Runtime/Net/Core/Private/Net/Core/DirtyNetObjectTracker/GlobalDirtyNetObjectTracker.cpp
+++ b/Engine/Source/Runtime/Net/Core/Private/Net/Core/DirtyNetObjectTracker/GlobalDirtyNetObjectTracker.cpp
@@ -99,6 +99,8 @@ const TSet<FNetHandle>& FGlobalDirtyNetObjectTracker::GetDirtyNetObjects(const F
 {
 	if (Instance && Handle.IsValid())
 	{
+		//$IRIS TODO: Once a poller reads the dirty list, we should make this list immutable and prevent any changes to it.
+		//			  Otherwise any changes between the first Get and the final Reset will get discarded and never seen by all the Pollers.
 		Instance->Pollers.SetBit(Handle.Index);
 		return Instance->DirtyObjects;
 	}
@@ -111,6 +113,7 @@ void FGlobalDirtyNetObjectTracker::ResetDirtyNetObjects(const FPollHandle& Handl
 	if (Instance && Handle.IsValid())
 	{
 		Instance->Pollers.ClearBit(Handle.Index);
+		//$IRIS TODO: This assumes every Poller called Get before anyone called Reset. We should make sure this is the case all the time.
 		if (Instance->Pollers.IsNoBitSet())
 		{
 			Instance->DirtyObjects.Reset();
