Iris - Implement support for NetSerializer::Apply which if implemented will be called when pushing received state to the target state. RepMovementNetSerializer makes use of the new functionality to avoid overwriting quantization levels.

diff --git a/Engine/Plugins/Runtime/ReplicationSystemTestPlugin/Source/Private/Tests/ReplicationState/TestDerivedStruct.cpp b/Engine/Plugins/Runtime/ReplicationSystemTestPlugin/Source/Private/Tests/ReplicationState/TestDerivedStruct.cpp
--- a/Engine/Plugins/Runtime/ReplicationSystemTestPlugin/Source/Private/Tests/ReplicationState/TestDerivedStruct.cpp
+++ b/Engine/Plugins/Runtime/ReplicationSystemTestPlugin/Source/Private/Tests/ReplicationState/TestDerivedStruct.cpp
@@ -55,6 +55,49 @@ struct FTestDerivedStruct_Inherited_WithNetSerializer_NetSerializer
 };
 UE_NET_IMPLEMENT_SERIALIZER(FTestDerivedStruct_Inherited_WithNetSerializer_NetSerializer);
 
+struct FTestDerivedStruct_Inherited_WithNetSerializerWithApply_NetSerializer
+{
+	struct FQuantizedData
+	{
+		uint8 ByteMember0;
+		uint8 ByteMember1;
+		uint8 ByteMemberNotSetOnApply;
+	};
+
+	typedef FTestDerivedStruct_Inherited_WithNetSerializerWithApply SourceType;
+	typedef FQuantizedData QuantizedType;
+
+	typedef FNetSerializerConfig ConfigType;
+
+	static const uint32 Version = 0;
+	inline static const ConfigType DefaultConfig;
+
+	static void Serialize(FNetSerializationContext&, const FNetSerializeArgs&);
+	static void Deserialize(FNetSerializationContext&, const FNetDeserializeArgs&);
+
+	static void Quantize(FNetSerializationContext&, const FNetQuantizeArgs&);
+	static void Dequantize(FNetSerializationContext&, const FNetDequantizeArgs&);
+
+	static bool IsEqual(FNetSerializationContext&, const FNetIsEqualArgs&);
+
+	static void Apply(FNetSerializationContext&, const FNetApplyArgs&);
+
+	class FNetSerializerRegistryDelegates final : private UE::Net::FNetSerializerRegistryDelegates
+	{
+	public:
+		virtual ~FNetSerializerRegistryDelegates();
+
+	private:
+		virtual void OnPreFreezeNetSerializerRegistry() override;
+
+		static inline const FName PropertyNetSerializerRegistry_NAME_TestDerivedStruct_Inherited_WithNetSerializerWithApply = FName("TestDerivedStruct_Inherited_WithNetSerializerWithApply");
+		UE_NET_IMPLEMENT_NAMED_STRUCT_NETSERIALIZER_INFO(PropertyNetSerializerRegistry_NAME_TestDerivedStruct_Inherited_WithNetSerializerWithApply, FTestDerivedStruct_Inherited_WithNetSerializerWithApply_NetSerializer);
+	};
+	
+	static inline FTestDerivedStruct_Inherited_WithNetSerializerWithApply_NetSerializer::FNetSerializerRegistryDelegates NetSerializerRegistryDelegates;
+};
+UE_NET_IMPLEMENT_SERIALIZER(FTestDerivedStruct_Inherited_WithNetSerializerWithApply_NetSerializer);
+
 class FTestDerivedStructFixture : public FNetworkAutomationTestSuiteFixture
 {
 };
@@ -64,7 +107,6 @@ class FTestDerivedStructInObjectFixture : public FReplicationSystemServerClientT
 protected:
 	UTestDerivedStruct_TestObject_Member* CreateObjectWithMember();
 	UTestDerivedStruct_TestObject_Array* CreateObjectWithArray();
-
 };
 
 // We want the struct deriving from the struct with custom NetSerializer to have the same size. This allows us to detect if the struct with the NetSerializer overwrites any of the properties in the derived struct.
@@ -276,6 +318,186 @@ UE_NET_TEST_FIXTURE(FTestDerivedStructInObjectFixture, CanReplicateDerivedStruct
 	UE_NET_ASSERT_EQ(ClientObject->DerivedStruct.ByteMember1, ServerObject->DerivedStruct.ByteMember1);
 	UE_NET_ASSERT_EQ(ClientObject->DerivedStruct.ByteMember2, ServerObject->DerivedStruct.ByteMember2);
 }
+
+UE_NET_TEST_FIXTURE(FTestDerivedStructInObjectFixture, NotAppliedMemberInDerivedStructIsNotOverwritten)
+{
+	FReplicationSystemTestClient* Client = CreateClient();
+
+	UTestDerivedStructWithNetSerializerWithApply_TestObject_Member* ServerObject = Server->CreateObject<UTestDerivedStructWithNetSerializerWithApply_TestObject_Member>();
+
+	// Replicate
+	Server->PreSendUpdate();
+	Server->SendAndDeliverTo(Client, DeliverPacket);
+	Server->PostSendUpdate();
+
+	UTestDerivedStructWithNetSerializerWithApply_TestObject_Member* ClientObject = Cast<UTestDerivedStructWithNetSerializerWithApply_TestObject_Member>(Client->GetReplicationBridge()->GetReplicatedObject(ServerObject->NetRefHandle));
+	UE_NET_ASSERT_NE(ClientObject, nullptr);
+
+	// On the client modify the non-applied property
+	const uint8 ClientModifiedNotAppliedPropertyValue = ClientObject->DerivedStruct.ByteMemberNotSetOnApply ^ 123U;
+	ClientObject->DerivedStruct.ByteMemberNotSetOnApply = ClientModifiedNotAppliedPropertyValue;
+
+	// On the server modify all members
+	ServerObject->DerivedStruct.ByteMember0 ^= 1U;
+	ServerObject->DerivedStruct.ByteMember1 ^= 2U;
+	ServerObject->DerivedStruct.ByteMemberNotSetOnApply = ClientModifiedNotAppliedPropertyValue ^ 231U;
+
+	Server->PreSendUpdate();
+	Server->SendAndDeliverTo(Client, DeliverPacket);
+	Server->PostSendUpdate();
+
+	// Verify state, most importantly that the not applied member remains as is.
+	UE_NET_ASSERT_EQ(ClientObject->DerivedStruct.ByteMember0, ServerObject->DerivedStruct.ByteMember0);
+	UE_NET_ASSERT_EQ(ClientObject->DerivedStruct.ByteMember1, ServerObject->DerivedStruct.ByteMember1);
+	UE_NET_ASSERT_EQ(ClientObject->DerivedStruct.ByteMemberNotSetOnApply, ClientModifiedNotAppliedPropertyValue);
+}
+
+UE_NET_TEST_FIXTURE(FTestDerivedStructInObjectFixture, NotAppliedMemberInDerivedStructInArrayIsNotOverwritten)
+{
+	FReplicationSystemTestClient* Client = CreateClient();
+
+	UTestDerivedStructWithNetSerializerWithApply_TestObject_Array* ServerObject = Server->CreateObject<UTestDerivedStructWithNetSerializerWithApply_TestObject_Array>();
+
+	// Add items to the array
+	{
+		FTestDerivedStruct_Inherited_WithNetSerializerWithApply_Inherited_WithoutNetSerializer Item;
+		for (int It = 0; It < 3; ++It)
+		{
+			ServerObject->DerivedStructArray.Add(Item);
+		}
+	}
+
+	// Replicate
+	Server->PreSendUpdate();
+	Server->SendAndDeliverTo(Client, DeliverPacket);
+	Server->PostSendUpdate();
+
+	UTestDerivedStructWithNetSerializerWithApply_TestObject_Array* ClientObject = Cast<UTestDerivedStructWithNetSerializerWithApply_TestObject_Array>(Client->GetReplicationBridge()->GetReplicatedObject(ServerObject->NetRefHandle));
+	UE_NET_ASSERT_NE(ClientObject, nullptr);
+	UE_NET_ASSERT_GT(ClientObject->DerivedStructArray.Num(), 0);
+
+	// On the client modify some non-applied properties
+	const uint8 ClientModifiedNotAppliedPropertyBaseValue = ClientObject->DerivedStructArray[0].ByteMemberNotSetOnApply ^ 123U;
+	for (auto& Element : ClientObject->DerivedStructArray)
+	{
+		const uint8 Index = &Element - ClientObject->DerivedStructArray.GetData();
+		Element.ByteMemberNotSetOnApply = ClientModifiedNotAppliedPropertyBaseValue + Index;
+	}
+
+	// On the server modify all elements in the array
+	for (auto& Element : ServerObject->DerivedStructArray)
+	{
+		const uint8 Index = &Element - ServerObject->DerivedStructArray.GetData();
+		Element.ByteMember0 ^= 1U;
+		Element.ByteMember1 ^= 2U;
+		Element.ByteMemberNotSetOnApply ^= 231U;
+	}
+
+	Server->PreSendUpdate();
+	Server->SendAndDeliverTo(Client, DeliverPacket);
+	Server->PostSendUpdate();
+
+	// Verify that the not applied member remains as is.
+	for (const auto& Element : ClientObject->DerivedStructArray)
+	{
+		const uint8 Index = &Element - ClientObject->DerivedStructArray.GetData();
+		UE_NET_ASSERT_EQ(Element.ByteMemberNotSetOnApply, uint8(ClientModifiedNotAppliedPropertyBaseValue + Index));
+	}
+}
+
+UE_NET_TEST_FIXTURE(FTestDerivedStructInObjectFixture, NotAppliedMemberInStructDerivedFromDerivedStructIsNotOverwritten)
+{
+	FReplicationSystemTestClient* Client = CreateClient();
+
+	UTestDerivedStructWithNetSerializerWithApply_Inherited_TestObject_Member* ServerObject = Server->CreateObject<UTestDerivedStructWithNetSerializerWithApply_Inherited_TestObject_Member>();
+
+	// Replicate
+	Server->PreSendUpdate();
+	Server->SendAndDeliverTo(Client, DeliverPacket);
+	Server->PostSendUpdate();
+
+	UTestDerivedStructWithNetSerializerWithApply_Inherited_TestObject_Member* ClientObject = Cast<UTestDerivedStructWithNetSerializerWithApply_Inherited_TestObject_Member>(Client->GetReplicationBridge()->GetReplicatedObject(ServerObject->NetRefHandle));
+	UE_NET_ASSERT_NE(ClientObject, nullptr);
+
+	// On the client modify the non-applied property
+	const uint8 ClientModifiedNotAppliedPropertyValue = ClientObject->DerivedStruct.ByteMemberNotSetOnApply ^ 123U;
+	ClientObject->DerivedStruct.ByteMemberNotSetOnApply = ClientModifiedNotAppliedPropertyValue;
+
+	// On the server modify all members
+	ServerObject->DerivedStruct.ByteMember0 ^= 1U;
+	ServerObject->DerivedStruct.ByteMember2 ^= 2U;
+	ServerObject->DerivedStruct.ByteMemberNotSetOnApply = ClientModifiedNotAppliedPropertyValue ^ 231U;
+
+	Server->PreSendUpdate();
+	Server->SendAndDeliverTo(Client, DeliverPacket);
+	Server->PostSendUpdate();
+
+	// Verify state, most importantly that the not applied member remains as is.
+	UE_NET_ASSERT_EQ(ClientObject->DerivedStruct.ByteMember0, ServerObject->DerivedStruct.ByteMember0);
+	UE_NET_ASSERT_EQ(ClientObject->DerivedStruct.ByteMember2, ServerObject->DerivedStruct.ByteMember2);
+	UE_NET_ASSERT_EQ(ClientObject->DerivedStruct.ByteMemberNotSetOnApply, ClientModifiedNotAppliedPropertyValue);
+}
+
+UE_NET_TEST_FIXTURE(FTestDerivedStructInObjectFixture, NotAppliedMemberInStructDerivedFromDerivedStructInArrayIsNotOverwritten)
+{
+	FReplicationSystemTestClient* Client = CreateClient();
+
+	UTestDerivedStructWithNetSerializerWithApply_Inherited_TestObject_Array* ServerObject = Server->CreateObject<UTestDerivedStructWithNetSerializerWithApply_Inherited_TestObject_Array>();
+
+	// Add items to the array
+	{
+		FTestDerivedStruct_Inherited_WithNetSerializerWithApply_Inherited_WithoutNetSerializer Item;
+		for (int It = 0; It < 3; ++It)
+		{
+			ServerObject->DerivedStructArray.Add(Item);
+		}
+	}
+
+	// Replicate
+	Server->PreSendUpdate();
+	Server->SendAndDeliverTo(Client, DeliverPacket);
+	Server->PostSendUpdate();
+
+	UTestDerivedStructWithNetSerializerWithApply_Inherited_TestObject_Array* ClientObject = Cast<UTestDerivedStructWithNetSerializerWithApply_Inherited_TestObject_Array>(Client->GetReplicationBridge()->GetReplicatedObject(ServerObject->NetRefHandle));
+	UE_NET_ASSERT_NE(ClientObject, nullptr);
+	UE_NET_ASSERT_GT(ClientObject->DerivedStructArray.Num(), 0);
+
+	// On the client modify some non-applied properties
+	// On the client modify the non-replicated property
+	constexpr uint8 ClientModifiedNotReplicatedPropertyValue = 255U;
+	const uint8 ClientModifiedNotAppliedPropertyBaseValue = ClientObject->DerivedStructArray[0].ByteMemberNotSetOnApply ^ 123U;
+	for (auto& Element : ClientObject->DerivedStructArray)
+	{
+		const uint8 Index = &Element - ClientObject->DerivedStructArray.GetData();
+		Element.ByteMemberNotSetOnApply = ClientModifiedNotAppliedPropertyBaseValue + Index;
+		Element.ByteMember3_NotReplicated = ClientModifiedNotReplicatedPropertyValue;
+	}
+
+	// On the server modify all elements in the array
+	for (auto& Element : ServerObject->DerivedStructArray)
+	{
+		const uint8 Index = &Element - ServerObject->DerivedStructArray.GetData();
+		Element.ByteMember0 ^= 1U;
+		Element.ByteMemberNotSetOnApply ^= 231U;
+		Element.ByteMember2 ^= 2U;
+		Element.ByteMember3_NotReplicated = ClientModifiedNotReplicatedPropertyValue ^ 0x7FU;
+	}
+
+	Server->PreSendUpdate();
+	Server->SendAndDeliverTo(Client, DeliverPacket);
+	Server->PostSendUpdate();
+
+	// Verify that the not applied member remains as is.
+	for (const auto& Element : ClientObject->DerivedStructArray)
+	{
+		const uint8 Index = &Element - ClientObject->DerivedStructArray.GetData();
+		UE_NET_ASSERT_EQ(Element.ByteMember0, ServerObject->DerivedStructArray[Index].ByteMember0);
+		UE_NET_ASSERT_EQ(Element.ByteMemberNotSetOnApply, uint8(ClientModifiedNotAppliedPropertyBaseValue + Index));
+		UE_NET_ASSERT_EQ(Element.ByteMember2, ServerObject->DerivedStructArray[Index].ByteMember2);
+		UE_NET_ASSERT_EQ(Element.ByteMember3_NotReplicated, ClientModifiedNotReplicatedPropertyValue);
+	}
+}
+
 // FTestDerivedStruct_Inherited_WithNetSerializer_NetSerializer implementation
 void FTestDerivedStruct_Inherited_WithNetSerializer_NetSerializer::Serialize(FNetSerializationContext& Context, const FNetSerializeArgs& Args)
 {
@@ -340,6 +562,85 @@ void FTestDerivedStruct_Inherited_WithNetSerializer_NetSerializer::FNetSerialize
 
 }
 
+// FTestDerivedStruct_Inherited_WithNetSerializerWithApply_NetSerializer implementation
+void FTestDerivedStruct_Inherited_WithNetSerializerWithApply_NetSerializer::Serialize(FNetSerializationContext& Context, const FNetSerializeArgs& Args)
+{
+	const QuantizedType& Value = *reinterpret_cast<const QuantizedType*>(Args.Source);
+
+	FNetBitStreamWriter* Writer = Context.GetBitStreamWriter();
+	Writer->WriteBits(Value.ByteMember0, 8U);
+	Writer->WriteBits(Value.ByteMember1, 8U);
+	// All serialization, dequantization etc is done on ByteMemberNotSetOnApply so that we can detect it's not applied.
+	Writer->WriteBits(Value.ByteMemberNotSetOnApply, 8U);
+}
+
+void FTestDerivedStruct_Inherited_WithNetSerializerWithApply_NetSerializer::Deserialize(FNetSerializationContext& Context, const FNetDeserializeArgs& Args)
+{
+	QuantizedType& Target = *reinterpret_cast<QuantizedType*>(Args.Target);
+
+	FNetBitStreamReader* Reader = Context.GetBitStreamReader();
+	Target.ByteMember0 = Reader->ReadBits(8U);
+	Target.ByteMember1 = Reader->ReadBits(8U);
+	Target.ByteMemberNotSetOnApply = Reader->ReadBits(8U);
+}
+
+void FTestDerivedStruct_Inherited_WithNetSerializerWithApply_NetSerializer::Quantize(FNetSerializationContext& Context, const FNetQuantizeArgs& Args)
+{
+	const SourceType& Source = *reinterpret_cast<const SourceType*>(Args.Source);
+	QuantizedType& Target = *reinterpret_cast<QuantizedType*>(Args.Target);
+
+	Target.ByteMember0 = Source.ByteMember0;
+	Target.ByteMember1 = Source.ByteMember1;
+	Target.ByteMemberNotSetOnApply = Source.ByteMemberNotSetOnApply;
+}
+
+void FTestDerivedStruct_Inherited_WithNetSerializerWithApply_NetSerializer::Dequantize(FNetSerializationContext& Context, const FNetDequantizeArgs& Args)
+{
+	const QuantizedType& Source = *reinterpret_cast<const QuantizedType*>(Args.Source);
+	SourceType& Target = *reinterpret_cast<SourceType*>(Args.Target);
+
+	Target.ByteMember0 = Source.ByteMember0;
+	Target.ByteMember1 = Source.ByteMember1;
+	Target.ByteMemberNotSetOnApply = Source.ByteMemberNotSetOnApply;
+}
+
+bool FTestDerivedStruct_Inherited_WithNetSerializerWithApply_NetSerializer::IsEqual(FNetSerializationContext& Context, const FNetIsEqualArgs& Args)
+{
+	if (Args.bStateIsQuantized)
+	{
+		const QuantizedType& Value0 = *reinterpret_cast<const QuantizedType*>(Args.Source0);
+		const QuantizedType& Value1 = *reinterpret_cast<const QuantizedType*>(Args.Source1);
+		return Value0.ByteMember0 == Value1.ByteMember0 && Value0.ByteMember1 == Value1.ByteMember1 && Value0.ByteMemberNotSetOnApply == Value1.ByteMemberNotSetOnApply;
+	}
+	else
+	{
+		const SourceType& Value0 = *reinterpret_cast<const SourceType*>(Args.Source0);
+		const SourceType& Value1 = *reinterpret_cast<const SourceType*>(Args.Source1);
+		return Value0.ByteMember0 == Value1.ByteMember0 && Value0.ByteMember1 == Value1.ByteMember1 && Value0.ByteMemberNotSetOnApply == Value1.ByteMemberNotSetOnApply;
+	}
+}
+
+void FTestDerivedStruct_Inherited_WithNetSerializerWithApply_NetSerializer::Apply(FNetSerializationContext&, const FNetApplyArgs& Args)
+{
+	const SourceType& Source = *reinterpret_cast<const SourceType*>(Args.Source);
+	SourceType& Target = *reinterpret_cast<SourceType*>(Args.Target);
+
+	Target.ByteMember0 = Source.ByteMember0;
+	Target.ByteMember1 = Source.ByteMember1;
+	// Not setting ByteMemberNotSetOnApply
+}
+
+FTestDerivedStruct_Inherited_WithNetSerializerWithApply_NetSerializer::FNetSerializerRegistryDelegates::~FNetSerializerRegistryDelegates()
+{
+	UE_NET_UNREGISTER_NETSERIALIZER_INFO(PropertyNetSerializerRegistry_NAME_TestDerivedStruct_Inherited_WithNetSerializerWithApply);
+}
+
+void FTestDerivedStruct_Inherited_WithNetSerializerWithApply_NetSerializer::FNetSerializerRegistryDelegates::OnPreFreezeNetSerializerRegistry()
+{
+	UE_NET_REGISTER_NETSERIALIZER_INFO(PropertyNetSerializerRegistry_NAME_TestDerivedStruct_Inherited_WithNetSerializerWithApply);
+
+}
+
 // FTestDerivedStruct_Inherited_WithNetSerializer_DeepInheritanceOfStructWithNetSerializer will not be used for replication testing so it's ok to use the NopNetSerializer.
 UE_NET_IMPLEMENT_FORWARDING_NETSERIALIZER_AND_REGISTRY_DELEGATES(TestDerivedStruct_Inherited_WithNetSerializer_DeepInheritanceOfStructWithNetSerializer, FNopNetSerializer);
 
@@ -382,3 +683,55 @@ void UTestDerivedStruct_TestObject_Array::RegisterReplicationFragments(UE::Net::
 {
 	UE::Net::FReplicationFragmentUtil::CreateAndRegisterFragmentsForObject(this, Context, RegistrationFlags);
 }
+
+void UTestDerivedStructWithNetSerializerWithApply_TestObject_Member::GetLifetimeReplicatedProps(TArray<FLifetimeProperty>& OutLifetimeProps) const
+{
+	FDoRepLifetimeParams Params;
+	Params.bIsPushBased = false;
+
+	DOREPLIFETIME_WITH_PARAMS_FAST(ThisClass, DerivedStruct, Params);
+}
+
+void UTestDerivedStructWithNetSerializerWithApply_TestObject_Member::RegisterReplicationFragments(UE::Net::FFragmentRegistrationContext& Context, UE::Net::EFragmentRegistrationFlags RegistrationFlags)
+{
+	UE::Net::FReplicationFragmentUtil::CreateAndRegisterFragmentsForObject(this, Context, RegistrationFlags);
+}
+
+void UTestDerivedStructWithNetSerializerWithApply_TestObject_Array::GetLifetimeReplicatedProps(TArray<FLifetimeProperty>& OutLifetimeProps) const
+{
+	FDoRepLifetimeParams Params;
+	Params.bIsPushBased = false;
+
+	DOREPLIFETIME_WITH_PARAMS_FAST(ThisClass, DerivedStructArray, Params);
+}
+
+void UTestDerivedStructWithNetSerializerWithApply_TestObject_Array::RegisterReplicationFragments(UE::Net::FFragmentRegistrationContext& Context, UE::Net::EFragmentRegistrationFlags RegistrationFlags)
+{
+	UE::Net::FReplicationFragmentUtil::CreateAndRegisterFragmentsForObject(this, Context, RegistrationFlags);
+}
+
+void UTestDerivedStructWithNetSerializerWithApply_Inherited_TestObject_Member::GetLifetimeReplicatedProps(TArray<FLifetimeProperty>& OutLifetimeProps) const
+{
+	FDoRepLifetimeParams Params;
+	Params.bIsPushBased = false;
+
+	DOREPLIFETIME_WITH_PARAMS_FAST(ThisClass, DerivedStruct, Params);
+}
+
+void UTestDerivedStructWithNetSerializerWithApply_Inherited_TestObject_Member::RegisterReplicationFragments(UE::Net::FFragmentRegistrationContext& Context, UE::Net::EFragmentRegistrationFlags RegistrationFlags)
+{
+	UE::Net::FReplicationFragmentUtil::CreateAndRegisterFragmentsForObject(this, Context, RegistrationFlags);
+}
+
+void UTestDerivedStructWithNetSerializerWithApply_Inherited_TestObject_Array::GetLifetimeReplicatedProps(TArray<FLifetimeProperty>& OutLifetimeProps) const
+{
+	FDoRepLifetimeParams Params;
+	Params.bIsPushBased = false;
+
+	DOREPLIFETIME_WITH_PARAMS_FAST(ThisClass, DerivedStructArray, Params);
+}
+
+void UTestDerivedStructWithNetSerializerWithApply_Inherited_TestObject_Array::RegisterReplicationFragments(UE::Net::FFragmentRegistrationContext& Context, UE::Net::EFragmentRegistrationFlags RegistrationFlags)
+{
+	UE::Net::FReplicationFragmentUtil::CreateAndRegisterFragmentsForObject(this, Context, RegistrationFlags);
+}
diff --git a/Engine/Plugins/Runtime/ReplicationSystemTestPlugin/Source/Private/Tests/ReplicationState/TestDerivedStruct.h b/Engine/Plugins/Runtime/ReplicationSystemTestPlugin/Source/Private/Tests/ReplicationState/TestDerivedStruct.h
--- a/Engine/Plugins/Runtime/ReplicationSystemTestPlugin/Source/Private/Tests/ReplicationState/TestDerivedStruct.h
+++ b/Engine/Plugins/Runtime/ReplicationSystemTestPlugin/Source/Private/Tests/ReplicationState/TestDerivedStruct.h
@@ -36,7 +36,31 @@ struct FTestDerivedStruct_Inherited_WithNetSerializer_Inherited_WithoutNetSerial
 	UPROPERTY(Transient)
 	uint8 ByteMember2 = 125;
 
-	UPROPERTY(NotReplicated)
+	UPROPERTY(Transient, NotReplicated)
+	uint8 ByteMember3_NotReplicated = 0;
+};
+
+USTRUCT()
+struct FTestDerivedStruct_Inherited_WithNetSerializerWithApply : public FTestDerivedStruct_Base
+{
+	GENERATED_BODY()
+	
+	UPROPERTY(Transient)
+	uint8 ByteMember1 = 25;
+
+	UPROPERTY(Transient)
+	uint8 ByteMemberNotSetOnApply = 33;
+};
+
+USTRUCT()
+struct FTestDerivedStruct_Inherited_WithNetSerializerWithApply_Inherited_WithoutNetSerializer : public FTestDerivedStruct_Inherited_WithNetSerializerWithApply
+{
+	GENERATED_BODY()
+
+	UPROPERTY(Transient)
+	uint8 ByteMember2 = 125;
+
+	UPROPERTY(Transient, NotReplicated)
 	uint8 ByteMember3_NotReplicated = 0;
 };
 
@@ -96,9 +120,66 @@ protected:
 	virtual void RegisterReplicationFragments(UE::Net::FFragmentRegistrationContext& Context, UE::Net::EFragmentRegistrationFlags RegistrationFlags) override;
 };
 
+UCLASS()
+class UTestDerivedStructWithNetSerializerWithApply_TestObject_Member : public UReplicatedTestObject
+{
+	GENERATED_BODY()
+
+public:
+	UPROPERTY(Replicated, Transient)
+	FTestDerivedStruct_Inherited_WithNetSerializerWithApply DerivedStruct;
+
+protected:
+	virtual void GetLifetimeReplicatedProps(TArray<FLifetimeProperty>& OutLifetimeProps) const override;
+	virtual void RegisterReplicationFragments(UE::Net::FFragmentRegistrationContext& Context, UE::Net::EFragmentRegistrationFlags RegistrationFlags) override;
+};
+
+UCLASS()
+class UTestDerivedStructWithNetSerializerWithApply_TestObject_Array : public UReplicatedTestObject
+{
+	GENERATED_BODY()
+
+public:
+	UPROPERTY(Replicated, Transient)
+	TArray<FTestDerivedStruct_Inherited_WithNetSerializerWithApply> DerivedStructArray;
+
+protected:
+	virtual void GetLifetimeReplicatedProps(TArray<FLifetimeProperty>& OutLifetimeProps) const override;
+	virtual void RegisterReplicationFragments(UE::Net::FFragmentRegistrationContext& Context, UE::Net::EFragmentRegistrationFlags RegistrationFlags) override;
+};
+
+UCLASS()
+class UTestDerivedStructWithNetSerializerWithApply_Inherited_TestObject_Member : public UReplicatedTestObject
+{
+	GENERATED_BODY()
+
+public:
+	UPROPERTY(Replicated, Transient)
+	FTestDerivedStruct_Inherited_WithNetSerializerWithApply_Inherited_WithoutNetSerializer DerivedStruct;
+
+protected:
+	virtual void GetLifetimeReplicatedProps(TArray<FLifetimeProperty>& OutLifetimeProps) const override;
+	virtual void RegisterReplicationFragments(UE::Net::FFragmentRegistrationContext& Context, UE::Net::EFragmentRegistrationFlags RegistrationFlags) override;
+};
+
+UCLASS()
+class UTestDerivedStructWithNetSerializerWithApply_Inherited_TestObject_Array : public UReplicatedTestObject
+{
+	GENERATED_BODY()
+
+public:
+	UPROPERTY(Replicated, Transient)
+	TArray<FTestDerivedStruct_Inherited_WithNetSerializerWithApply_Inherited_WithoutNetSerializer> DerivedStructArray;
+
+protected:
+	virtual void GetLifetimeReplicatedProps(TArray<FLifetimeProperty>& OutLifetimeProps) const override;
+	virtual void RegisterReplicationFragments(UE::Net::FFragmentRegistrationContext& Context, UE::Net::EFragmentRegistrationFlags RegistrationFlags) override;
+};
+
 namespace UE::Net::Private
 {
 
 UE_NET_DECLARE_SERIALIZER(FTestDerivedStruct_Inherited_WithNetSerializer_NetSerializer, REPLICATIONSYSTEMTESTPLUGIN_API);
+UE_NET_DECLARE_SERIALIZER(FTestDerivedStruct_Inherited_WithNetSerializerWithApply_NetSerializer, REPLICATIONSYSTEMTESTPLUGIN_API);
 
 }
diff --git a/Engine/Plugins/Runtime/ReplicationSystemTestPlugin/Source/Private/Tests/Serialization/TestNetSerializer.cpp b/Engine/Plugins/Runtime/ReplicationSystemTestPlugin/Source/Private/Tests/Serialization/TestNetSerializer.cpp
--- a/Engine/Plugins/Runtime/ReplicationSystemTestPlugin/Source/Private/Tests/Serialization/TestNetSerializer.cpp
+++ b/Engine/Plugins/Runtime/ReplicationSystemTestPlugin/Source/Private/Tests/Serialization/TestNetSerializer.cpp
@@ -32,7 +32,7 @@ struct FFullNetSerializer
 	typedef int SourceType;
 	typedef FNetSerializerConfig ConfigType;
 
-	static const ConfigType DefaultConfig;
+	inline static const ConfigType DefaultConfig;
 
 	static void Serialize(FNetSerializationContext&, const FNetSerializeArgs&);
 	static void Deserialize(FNetSerializationContext&, const FNetDeserializeArgs&);
@@ -48,16 +48,23 @@ struct FFullNetSerializer
 };
 UE_NET_IMPLEMENT_SERIALIZER(FFullNetSerializer);
 
-const FFullNetSerializer::ConfigType FFullNetSerializer::DefaultConfig;
-
 UE_NET_DECLARE_SERIALIZER(FNetSerializerWithHasConnectionSpecificSerializationTrait, REPLICATIONSYSTEMTESTPLUGIN_API);
-
 struct FNetSerializerWithHasConnectionSpecificSerializationTrait : public FMinimalNetSerializer
 {
 	static constexpr bool bHasConnectionSpecificSerialization = true;
 };
 UE_NET_IMPLEMENT_SERIALIZER(FNetSerializerWithHasConnectionSpecificSerializationTrait);
 
+UE_NET_DECLARE_SERIALIZER(FNetSerializerWithApply, REPLICATIONSYSTEMTESTPLUGIN_API);
+struct FNetSerializerWithApply : public FMinimalNetSerializer
+{
+	typedef int SourceType;
+
+	static void Apply(FNetSerializationContext&, const FNetApplyArgs&);
+};
+UE_NET_IMPLEMENT_SERIALIZER(FNetSerializerWithApply);
+
+
 }
 
 namespace UE::Net::Private
@@ -296,6 +303,34 @@ UE_NET_TEST(FNetSerializerWithHasConnectionSpecificSerializationTrait, HasHasCon
 	UE_NET_ASSERT_TRUE(EnumHasAnyFlags(Serializer->Traits, ENetSerializerTraits::HasConnectionSpecificSerialization));
 }
 
+//
+UE_NET_TEST(FNetSerializerWithApply, HasApplyTrait)
+{
+	const FNetSerializer& Serializer = UE_NET_GET_SERIALIZER(FNetSerializerWithApply);
+	UE_NET_ASSERT_TRUE(EnumHasAnyFlags(Serializer.Traits, ENetSerializerTraits::HasApply));
+}
+
+UE_NET_TEST(FNetSerializerWithApply, HasApplyFunction)
+{
+	const FNetSerializer& Serializer = UE_NET_GET_SERIALIZER(FNetSerializerWithApply);
+	UE_NET_ASSERT_NE(Serializer.Apply, NetApplyFunction(0));
+}
+
+UE_NET_TEST(FNetSerializerWithApply, ApplyGetsCalled)
+{
+	const FNetSerializer& Serializer = UE_NET_GET_SERIALIZER(FNetSerializerWithApply);
+
+	FNetSerializerWithApply::SourceType CallCount = 0;
+
+	FNetSerializationContext Context;
+	FNetApplyArgs Args;
+	Args.Version = Serializer.Version;
+	Args.NetSerializerConfig = NetSerializerConfigParam(nullptr);
+	Args.Source = NetSerializerValuePointer(&CallCount);
+	Serializer.Apply(Context, Args);
+	UE_NET_ASSERT_EQ(CallCount, FNetSerializerWithApply::SourceType(1));
+}
+
 }
 
 namespace UE::Net
@@ -370,4 +405,10 @@ void FFullNetSerializer::CollectNetReferences(FNetSerializationContext&, const F
 	*reinterpret_cast<SourceType*>(Args.Source) += 1;
 }
 
+// FNetSerializerWithApply
+void FNetSerializerWithApply::Apply(FNetSerializationContext&, const FNetApplyArgs& Args)
+{
+	*reinterpret_cast<SourceType*>(Args.Source) += 1;
+}
+
 }
diff --git a/Engine/Source/Runtime/Engine/Private/RepMovementNetSerializer.cpp b/Engine/Source/Runtime/Engine/Private/RepMovementNetSerializer.cpp
--- a/Engine/Source/Runtime/Engine/Private/RepMovementNetSerializer.cpp
+++ b/Engine/Source/Runtime/Engine/Private/RepMovementNetSerializer.cpp
@@ -61,6 +61,9 @@ struct FRepMovementNetSerializer
 	static bool IsEqual(FNetSerializationContext&, const FNetIsEqualArgs&);
 	static bool Validate(FNetSerializationContext&, const FNetValidateArgs&);
 
+	// Avoid overwriting some members in the target FRepMovement instance.
+	static void Apply(FNetSerializationContext&, const FNetApplyArgs&);
+
 private:
 	enum EFlags : uint32
 	{
@@ -567,7 +570,7 @@ void FRepMovementNetSerializer::Dequantize(FNetSerializationContext& Context, co
 	Target.VelocityQuantizationLevel = EVectorQuantization(Source.VelocityQuantizationLevel);
 	Target.RotationQuantizationLevel = ERotatorQuantization(Source.RotationQuantizationLevel);
 
-	// We do not overwrite AngularVelocity unless we're replicating it. This is consistent with the FRepMovement serialization method.
+	// There's no need to dequantize angular velocity unless we're replicating it.
 	if (Source.Flags & Flag_RepPhysics)
 	{
 		const FNetSerializer* Serializer = VectorNetQuantizeNetSerializers[Source.VelocityQuantizationLevel];
@@ -694,6 +697,33 @@ bool FRepMovementNetSerializer::Validate(FNetSerializationContext& Context, cons
 	return true;
 }
 
+void FRepMovementNetSerializer::Apply(FNetSerializationContext& Context, const FNetApplyArgs& Args)
+{
+	const SourceType& Source = *reinterpret_cast<const SourceType*>(Args.Source);
+	SourceType& Target = *reinterpret_cast<SourceType*>(Args.Target);
+
+	Target.bSimulatedPhysicSleep = Source.bSimulatedPhysicSleep;
+	Target.bRepPhysics = Source.bRepPhysics;
+	Target.ServerFrame = Source.ServerFrame;
+	// FRepMovement::NetSerialize does not overwrite ServerPhysicsHandle if it wasn't replicated. Despite the member name we mimic the behavior.
+	if (Source.ServerPhysicsHandle != INDEX_NONE)
+	{
+		Target.ServerPhysicsHandle = Source.ServerPhysicsHandle;
+	}
+
+	Target.LinearVelocity = Source.LinearVelocity;
+	Target.Location = Source.Location;
+	Target.Rotation = Source.Rotation;
+
+	// Explicitly do not overwrite AngularVelocity unless it was replicated.
+	if (Source.bRepPhysics)
+	{
+		Target.AngularVelocity = Source.AngularVelocity;
+	}
+
+	// Explicitly do not overwrite quantization levels.
+}
+
 // FNetSerializerRegistryDelegates
 FRepMovementNetSerializer::FNetSerializerRegistryDelegates::~FNetSerializerRegistryDelegates()
 {
diff --git a/Engine/Source/Runtime/Experimental/Iris/Core/Private/Iris/ReplicationState/InternalPropertyReplicationState.cpp b/Engine/Source/Runtime/Experimental/Iris/Core/Private/Iris/ReplicationState/InternalPropertyReplicationState.cpp
--- a/Engine/Source/Runtime/Experimental/Iris/Core/Private/Iris/ReplicationState/InternalPropertyReplicationState.cpp
+++ b/Engine/Source/Runtime/Experimental/Iris/Core/Private/Iris/ReplicationState/InternalPropertyReplicationState.cpp
@@ -220,11 +220,134 @@ bool InternalCompareStructProperty(const FReplicationStateDescriptor* StructDesc
 	return true;
 }
 
+void InternalApplyStructProperty(const FReplicationStateDescriptor* StructDescriptor, void* RESTRICT Dst, const void* RESTRICT Src)
+{
+	uint32 FirstStructMemberForApply = 0U;
+	if (EnumHasAnyFlags(StructDescriptor->Traits, EReplicationStateTraits::IsDerivedStruct))
+	{
+		// The first member in a derived struct descriptor is the base struct. We can skip after we've applied the value.
+		FirstStructMemberForApply = 1U;
+
+		// If derived struct has custom Apply we need to use it.
+		const FReplicationStateMemberSerializerDescriptor& MemberSerializerDescriptor = StructDescriptor->MemberSerializerDescriptors[0];
+		if (EnumHasAnyFlags(MemberSerializerDescriptor.Serializer->Traits, ENetSerializerTraits::HasApply))
+		{
+			FNetSerializationContext Context;
+			FNetApplyArgs ApplyArgs;
+			ApplyArgs.NetSerializerConfig = MemberSerializerDescriptor.SerializerConfig;
+			ApplyArgs.Source = NetSerializerValuePointer(Src);
+			ApplyArgs.Target = NetSerializerValuePointer(Dst);
+			MemberSerializerDescriptor.Serializer->Apply(Context, ApplyArgs);
+		}
+		else
+		{
+		// The base struct is the closest parent with a custom serializer. It was determined above that it did not have a custom apply so we should copy it in its entirety.
+			const UScriptStruct* BaseStruct = StructDescriptor->BaseStruct;
+			BaseStruct->CopyScriptStruct(Dst, Src, 1);
+		}
+	}
+
+	const FReplicationStateMemberDescriptor* MemberDescriptors = StructDescriptor->MemberDescriptors;
+	const FProperty** MemberProperties = StructDescriptor->MemberProperties;
+	const FReplicationStateMemberPropertyDescriptor* MemberPropertyDescriptors = StructDescriptor->MemberPropertyDescriptors;
+	for (uint32 StructMemberIt = FirstStructMemberForApply, StructMemberEndIt = StructDescriptor->MemberCount; StructMemberIt < StructMemberEndIt; ++StructMemberIt)
+	{
+		const FReplicationStateMemberDescriptor& MemberDescriptor = MemberDescriptors[StructMemberIt];
+		const FReplicationStateMemberPropertyDescriptor& MemberPropertyDescriptor = MemberPropertyDescriptors[StructMemberIt];
+		const FProperty* MemberProperty = MemberProperties[StructMemberIt];
+		const SIZE_T MemberOffset = MemberProperty->GetOffset_ForGC() + MemberProperty->ElementSize*MemberPropertyDescriptor.ArrayIndex;
+		InternalApplyPropertyValue(StructDescriptor, StructMemberIt, static_cast<uint8*>(Dst) + MemberOffset, static_cast<const uint8*>(Src) + MemberOffset);
+	}
+}
+
+void InternalApplyPropertyValue(const FReplicationStateDescriptor* Descriptor, uint32 MemberIndex, void* RESTRICT Dst, const void* RESTRICT Src)
+{
+	const FReplicationStateMemberSerializerDescriptor& MemberSerializerDescriptor = Descriptor->MemberSerializerDescriptors[MemberIndex];
+	
+	// If member has serializer with Apply then use it.
+	if (EnumHasAnyFlags(MemberSerializerDescriptor.Serializer->Traits, ENetSerializerTraits::HasApply))
+	{
+		FNetSerializationContext Context;
+		FNetApplyArgs ApplyArgs;
+		ApplyArgs.NetSerializerConfig = MemberSerializerDescriptor.SerializerConfig;
+		ApplyArgs.Source = NetSerializerValuePointer(Src);
+		ApplyArgs.Target = NetSerializerValuePointer(Dst);
+		MemberSerializerDescriptor.Serializer->Apply(Context, ApplyArgs);
+		return;
+	}
+
+	if (IsUsingStructNetSerializer(MemberSerializerDescriptor))
+	{
+		const FStructNetSerializerConfig* StructConfig = static_cast<const FStructNetSerializerConfig*>(MemberSerializerDescriptor.SerializerConfig);
+		const FReplicationStateDescriptor* StructDescriptor = StructConfig->StateDescriptor;
+		if (!EnumHasAnyFlags(StructDescriptor->Traits, EReplicationStateTraits::AllMembersAreReplicated))
+		{
+			InternalApplyStructProperty(StructDescriptor, Dst, Src);
+			return;
+		}
+	}
+	else if (IsUsingArrayPropertyNetSerializer(MemberSerializerDescriptor))
+	{
+		const FArrayPropertyNetSerializerConfig* ArrayConfig = static_cast<const FArrayPropertyNetSerializerConfig*>(MemberSerializerDescriptor.SerializerConfig);
+		const FReplicationStateDescriptor* ElementStateDescriptor = ArrayConfig->StateDescriptor;
+		const FReplicationStateMemberSerializerDescriptor& ElementSerializerDescriptor = ElementStateDescriptor->MemberSerializerDescriptors[0];
+
+		const bool bSerializerHasApply = EnumHasAnyFlags(ElementSerializerDescriptor.Serializer->Traits, ENetSerializerTraits::HasApply);
+		const bool bIsStructWithNotReplicatedProps = IsUsingStructNetSerializer(ElementSerializerDescriptor) && !EnumHasAnyFlags(ElementStateDescriptor->Traits, EReplicationStateTraits::AllMembersAreReplicated);
+		if (bSerializerHasApply || bIsStructWithNotReplicatedProps)
+		{
+			const FReplicationStateDescriptor* StructDescriptor = static_cast<const FStructNetSerializerConfig*>(ElementSerializerDescriptor.SerializerConfig)->StateDescriptor;
+	
+			// Need to explicitly iterate over array members to be able to only copy data that we should copy
+			FScriptArrayHelper ScriptArrayHelperSrc(ArrayConfig->Property.Get(), reinterpret_cast<const void*>(Src));
+			FScriptArrayHelper ScriptArrayHelperDst(ArrayConfig->Property.Get(), reinterpret_cast<void*>(Dst));
+
+			// First we must resize target to match size of source data
+			const uint32 ElementCount = ScriptArrayHelperSrc.Num();
+			ScriptArrayHelperDst.Resize(ElementCount);
+		
+			// Iterate over array entries and copy the values
+			if (bSerializerHasApply)
+			{
+				for (uint32 ElementIt = 0, ElementEndIt = ElementCount; ElementIt < ElementEndIt; ++ElementIt)
+				{
+					const uint8* ArraySrc = ScriptArrayHelperSrc.GetRawPtr(ElementIt);
+					uint8* ArrayDst = ScriptArrayHelperDst.GetRawPtr(ElementIt);
+
+					FNetSerializationContext Context;
+					FNetApplyArgs ApplyArgs;
+					ApplyArgs.NetSerializerConfig = ElementSerializerDescriptor.SerializerConfig;
+					ApplyArgs.Source = NetSerializerValuePointer(ArraySrc);
+					ApplyArgs.Target = NetSerializerValuePointer(ArrayDst);
+					ElementSerializerDescriptor.Serializer->Apply(Context, ApplyArgs);
+				}
+			}
+			else
+			{
+				for (uint32 ElementIt = 0, ElementEndIt = ElementCount; ElementIt < ElementEndIt; ++ElementIt)
+				{
+					const uint8* ArraySrc = ScriptArrayHelperSrc.GetRawPtr(ElementIt);
+					uint8* ArrayDst = ScriptArrayHelperDst.GetRawPtr(ElementIt);
+
+					InternalApplyStructProperty(StructDescriptor, ArrayDst, ArraySrc);
+				}
+			}
+
+			return;
+		}
+	}
+
+	// Default handling for all properties except for structs and arrays that have some non-replicated members or a serializer with custom Apply.
+	const FProperty* Property = Descriptor->MemberProperties[MemberIndex];
+	Property->CopySingleValue(Dst, Src);
+}
+
 void InternalCopyStructProperty(const FReplicationStateDescriptor* StructDescriptor, void* RESTRICT Dst, const void* RESTRICT Src)
 {
 	uint32 FirstStructMemberForCopy = 0U;
 	if (EnumHasAnyFlags(StructDescriptor->Traits, EReplicationStateTraits::IsDerivedStruct))
 	{
+		// The base struct is the closest parent with a custom serializer, for which we always copy the entire state.
 		const UScriptStruct* BaseStruct = StructDescriptor->BaseStruct;
 		BaseStruct->CopyScriptStruct(Dst, Src, 1);
 
diff --git a/Engine/Source/Runtime/Experimental/Iris/Core/Private/Iris/ReplicationState/InternalPropertyReplicationState.h b/Engine/Source/Runtime/Experimental/Iris/Core/Private/Iris/ReplicationState/InternalPropertyReplicationState.h
--- a/Engine/Source/Runtime/Experimental/Iris/Core/Private/Iris/ReplicationState/InternalPropertyReplicationState.h
+++ b/Engine/Source/Runtime/Experimental/Iris/Core/Private/Iris/ReplicationState/InternalPropertyReplicationState.h
@@ -22,6 +22,12 @@ IRISCORE_API void DestructPropertyReplicationState(uint8* StateBuffer, const FRe
 /** Copy external state in SrcStateBuffer to external state DstStateBuffer including changemask */
 IRISCORE_API void CopyPropertyReplicationState(uint8* RESTRICT DstStateBuffer, uint8* RESTRICT SrcStateBuffer, const FReplicationStateDescriptor* Descriptor);
 
+/** Copy property value to target. If the property is fully replicated we use the properties copy function, otherwise we do a per member copy of all replicated data. Members with NetSerializers with an Apply() function will call that. */
+IRISCORE_API void InternalApplyPropertyValue(const FReplicationStateDescriptor* Descriptor, uint32 MemberIndex, void* RESTRICT Dst, const void* RESTRICT Src);
+
+/** Copy struct members to target using InternalApplyPropertyValue */
+IRISCORE_API void InternalApplyStructProperty(const FReplicationStateDescriptor* StructDescriptor, void* RESTRICT Dst, const void* RESTRICT Src);
+
 /** Copy property value, if the property is fully replicated we use the properties copy function, otherwise we do a per member copy of all replicated data */
 IRISCORE_API void InternalCopyPropertyValue(const FReplicationStateDescriptor* Descriptor, uint32 MemberIndex, void* RESTRICT Dst, const void* RESTRICT Src);
 
diff --git a/Engine/Source/Runtime/Experimental/Iris/Core/Private/Iris/ReplicationState/PropertyReplicationState.cpp b/Engine/Source/Runtime/Experimental/Iris/Core/Private/Iris/ReplicationState/PropertyReplicationState.cpp
--- a/Engine/Source/Runtime/Experimental/Iris/Core/Private/Iris/ReplicationState/PropertyReplicationState.cpp
+++ b/Engine/Source/Runtime/Experimental/Iris/Core/Private/Iris/ReplicationState/PropertyReplicationState.cpp
@@ -144,7 +144,7 @@ void FPropertyReplicationState::Set(const FPropertyReplicationState& Other)
 	}
 }
 
-void FPropertyReplicationState::SetPropertyValue(uint32 Index, const void* SrcValue)
+void FPropertyReplicationState::PollPropertyValue(uint32 Index, const void* SrcValue)
 {
 	const FReplicationStateDescriptor* Descriptor = ReplicationStateDescriptor;
 	void* DstValue = StateBuffer + Descriptor->MemberDescriptors[Index].ExternalMemberOffset;
@@ -176,6 +176,18 @@ void FPropertyReplicationState::SetPropertyValue(uint32 Index, const void* SrcVa
 	Private::InternalCopyPropertyValue(Descriptor, Index, DstValue, SrcValue);
 }
 
+void FPropertyReplicationState::SetPropertyValue(uint32 Index, const void* SrcValue)
+{
+	// We can perform the same operation as normal polling of the state does.
+	PollPropertyValue(Index, SrcValue);
+}
+
+void FPropertyReplicationState::PushPropertyValue(uint32 Index, void* DstValue) const
+{
+	void* SrcValue = StateBuffer + ReplicationStateDescriptor->MemberDescriptors[Index].ExternalMemberOffset;
+	Private::InternalApplyPropertyValue(ReplicationStateDescriptor, Index, DstValue, SrcValue);
+}
+
 void FPropertyReplicationState::GetPropertyValue(uint32 Index, void* DstValue) const
 {
 	void* SrcValue = StateBuffer + ReplicationStateDescriptor->MemberDescriptors[Index].ExternalMemberOffset;
@@ -270,7 +282,7 @@ bool FPropertyReplicationState::PollPropertyReplicationState(const void* RESTRIC
 			const FProperty* Property = MemberProperties[MemberIt];
 
 			//$TODO: make special version to avoid unnecessary overhead.
-			SetPropertyValue(MemberIt, SrcBuffer + Property->GetOffset_ForGC() + Property->ElementSize*MemberPropertyDescriptor.ArrayIndex);
+			PollPropertyValue(MemberIt, SrcBuffer + Property->GetOffset_ForGC() + Property->ElementSize*MemberPropertyDescriptor.ArrayIndex);
 		}
 	}
 
@@ -299,7 +311,7 @@ bool FPropertyReplicationState::PollPropertyReplicationStateForRepNotifies(const
 			if (MemberPropertyDescriptor.RepNotifyFunction)
 			{
 				const FProperty* Property = MemberProperties[MemberIt];
-				SetPropertyValue(MemberIt, SrcBuffer + Property->GetOffset_ForGC() + Property->ElementSize*MemberPropertyDescriptor.ArrayIndex);
+				PollPropertyValue(MemberIt, SrcBuffer + Property->GetOffset_ForGC() + Property->ElementSize*MemberPropertyDescriptor.ArrayIndex);
 			}
 		}
 	}
@@ -334,7 +346,7 @@ void FPropertyReplicationState::PushPropertyReplicationState(void* RESTRICT DstD
 				const FReplicationStateMemberPropertyDescriptor& MemberPropertyDescriptor = MemberPropertyDescriptors[MemberIt];
 				const FProperty* Property = MemberProperties[MemberIt];
 
-				GetPropertyValue(MemberIt, DstBuffer + Property->GetOffset_ForGC() + Property->ElementSize*MemberPropertyDescriptor.ArrayIndex);
+				PushPropertyValue(MemberIt, DstBuffer + Property->GetOffset_ForGC() + Property->ElementSize*MemberPropertyDescriptor.ArrayIndex);
 			}
 		}
 	}
@@ -361,7 +373,7 @@ bool FPropertyReplicationState::PollObjectReferences(const void* RESTRICT SrcSta
 				const FReplicationStateMemberPropertyDescriptor& MemberPropertyDescriptor = MemberPropertyDescriptors[MemberIt];
 				const FProperty* Property = MemberProperties[MemberIt];
 
-				SetPropertyValue(MemberIt, SrcBuffer + Property->GetOffset_ForGC() + Property->ElementSize*MemberPropertyDescriptor.ArrayIndex);
+				PollPropertyValue(MemberIt, SrcBuffer + Property->GetOffset_ForGC() + Property->ElementSize*MemberPropertyDescriptor.ArrayIndex);
 			}
 		}
 	}
@@ -499,28 +511,4 @@ bool FPropertyReplicationState::IsCustomConditionEnabled(uint32 Index) const
 	return MemberConditionalChangeMask.GetBit(ChangeMaskInfo.BitOffset);
 }
 
-void FPropertyReplicationState::PollProperty(const void* SrcData, uint32 MemberIndex)
-{
-	if (IsValid())
-	{
-		const FReplicationStateDescriptor* Descriptor = ReplicationStateDescriptor;
-		const uint8* SrcBuffer = static_cast<const uint8*>(SrcData);
-
-		const FReplicationStateMemberDescriptor* MemberDescriptors = Descriptor->MemberDescriptors;
-		const FProperty** MemberProperties = Descriptor->MemberProperties;
-		const FReplicationStateMemberPropertyDescriptor* MemberPropertyDescriptors = Descriptor->MemberPropertyDescriptors;
-		const uint32 MemberCount = Descriptor->MemberCount;
-
-		if (MemberIndex < MemberCount)
-		{
-			const FReplicationStateMemberDescriptor& MemberDescriptor = MemberDescriptors[MemberIndex];
-			const FReplicationStateMemberPropertyDescriptor& MemberPropertyDescriptor = MemberPropertyDescriptors[MemberIndex];
-			const FProperty* Property = MemberProperties[MemberIndex];
-
-			//$TODO: make special version to avoid unnecessary overhead.
-			SetPropertyValue(MemberIndex, SrcBuffer + Property->GetOffset_ForGC() + Property->ElementSize*MemberPropertyDescriptor.ArrayIndex);
-		}
-	}
-}
-
 }
diff --git a/Engine/Source/Runtime/Experimental/Iris/Core/Private/Iris/ReplicationSystem/FastArrayReplicationFragment.cpp b/Engine/Source/Runtime/Experimental/Iris/Core/Private/Iris/ReplicationSystem/FastArrayReplicationFragment.cpp
--- a/Engine/Source/Runtime/Experimental/Iris/Core/Private/Iris/ReplicationSystem/FastArrayReplicationFragment.cpp
+++ b/Engine/Source/Runtime/Experimental/Iris/Core/Private/Iris/ReplicationSystem/FastArrayReplicationFragment.cpp
@@ -280,6 +280,11 @@ void FNativeFastArrayReplicationFragmentBase::CollectOwner(FReplicationStateOwne
 ////////////////////////////////////////////////////////////////////////////////////////////////////////////////////////////////////////////////////
 // FFastArrayReplicationFragmentHelper
 ////////////////////////////////////////////////////////////////////////////////////////////////////////////////////////////////////////////////////
+void FFastArrayReplicationFragmentHelper::InternalApplyArrayElement(const FReplicationStateDescriptor* ArrayElementDescriptor, void* RESTRICT Dst, const void* RESTRICT Src)
+{
+	InternalApplyStructProperty(ArrayElementDescriptor, Dst, Src);
+}
+
 void FFastArrayReplicationFragmentHelper::InternalCopyArrayElement(const FReplicationStateDescriptor* ArrayElementDescriptor, void* RESTRICT Dst, const void* RESTRICT Src)
 {
 	InternalCopyStructProperty(ArrayElementDescriptor, Dst, Src);
diff --git a/Engine/Source/Runtime/Experimental/Iris/Core/Public/Iris/ReplicationState/PropertyReplicationState.h b/Engine/Source/Runtime/Experimental/Iris/Core/Public/Iris/ReplicationState/PropertyReplicationState.h
--- a/Engine/Source/Runtime/Experimental/Iris/Core/Public/Iris/ReplicationState/PropertyReplicationState.h
+++ b/Engine/Source/Runtime/Experimental/Iris/Core/Public/Iris/ReplicationState/PropertyReplicationState.h
@@ -63,12 +63,11 @@ public:
 	*/
 	IRISCORE_API void SetPropertyValue(uint32 Index, const void* SrcValue);
 
-	/** Set the value at the provided Index, written into DstValue. 
-		The UProperty is looked up from the descriptor using the index, if the value differs the statemask is updated
-		Mostly intended for test code
-		Normal use of this class is through the poll layer
-		TODO: It would be nice if we could provide some validation on property types
-	*/
+	/**
+	 * Retrieve the value at the provided Index by writing it to DstValue. 
+	 * The property is looked up from the descriptor using the index.
+	 * Mostly intended for test code. Normal use of this class is through the poll layer.
+	 */
 	IRISCORE_API void GetPropertyValue(uint32 Index, void* DstValue) const;
 
 	/** Is the property at the given index dirty, for properties with multiple bits in the statemask this will return true if any of those bits are set */
@@ -149,9 +148,6 @@ public:
 		TArrayElementChangeMaskBitOffset = 1U,
 	};
 
-	// Polls source data for a single property. If the property has changed it will be marked as dirty.
-	void PollProperty(const void* SrcData, uint32 Index);
-
 	/**
 	 * If state has custom conditionals then this function will check whether the condition is enabled or not.
 	 * If state doesn't have custom conditionals it will return true.
@@ -164,6 +160,9 @@ private:
 	void ConstructStateInternal();
 	void DestructStateInternal();
 	void InjectState(const FReplicationStateDescriptor* Descriptor, uint8* InStateBuffer);
+	void PollPropertyValue(uint32 Index, const void* SrcValue);
+	void PushPropertyValue(uint32 Index, void* DstValue) const;
+
 
 	TRefCountPtr<const FReplicationStateDescriptor> ReplicationStateDescriptor;
 	uint8* StateBuffer;
diff --git a/Engine/Source/Runtime/Experimental/Iris/Core/Public/Iris/ReplicationSystem/Private/FastArrayReplicationFragmentInternal.h b/Engine/Source/Runtime/Experimental/Iris/Core/Public/Iris/ReplicationSystem/Private/FastArrayReplicationFragmentInternal.h
--- a/Engine/Source/Runtime/Experimental/Iris/Core/Public/Iris/ReplicationSystem/Private/FastArrayReplicationFragmentInternal.h
+++ b/Engine/Source/Runtime/Experimental/Iris/Core/Public/Iris/ReplicationSystem/Private/FastArrayReplicationFragmentInternal.h
@@ -64,6 +64,9 @@ struct FFastArrayReplicationFragmentHelper
 	template <typename FastArrayType, typename ItemArrayType>
 	static void ApplyReplicatedState(FastArrayType* DstFastArray, ItemArrayType* DstWrappedArray, FastArrayType* SrcFastArray, const ItemArrayType* SrcWrappedArray, const FReplicationStateDescriptor* ArrayElementDescriptor, FReplicationStateApplyContext& Context);
 
+	/** Apply array element, only replicated items will be applied, using the serializers' Apply function if present  */
+	IRISCORE_API static void InternalApplyArrayElement(const FReplicationStateDescriptor* ArrayElementDescriptor, void* RESTRICT Dst, const void* RESTRICT Src);
+
 	/** Copy array element, only replicated items will be copied */
 	IRISCORE_API static void InternalCopyArrayElement(const FReplicationStateDescriptor* ArrayElementDescriptor, void* RESTRICT Dst, const void* RESTRICT Src);
 
@@ -215,8 +218,8 @@ void FFastArrayReplicationFragmentHelper::ApplyReplicatedState(FastArrayType* Ds
 
 					ModifiedIndices.Add(*ExistingIndex);
 
-					// We use per element copy since we do not want to overwrite data that is not replicated
-					InternalCopyArrayElement(ArrayElementDescriptor, &(*DstWrappedArray)[*ExistingIndex], &SrcItems[It]);
+					// We use per element apply since we do not want to overwrite data that is not replicated
+					InternalApplyArrayElement(ArrayElementDescriptor, &(*DstWrappedArray)[*ExistingIndex], &SrcItems[It]);
 				}
 			}
 			else
diff --git a/Engine/Source/Runtime/Experimental/Iris/Core/Public/Iris/Serialization/NetSerializer.h b/Engine/Source/Runtime/Experimental/Iris/Core/Public/Iris/Serialization/NetSerializer.h
--- a/Engine/Source/Runtime/Experimental/Iris/Core/Public/Iris/Serialization/NetSerializer.h
+++ b/Engine/Source/Runtime/Experimental/Iris/Core/Public/Iris/Serialization/NetSerializer.h
@@ -114,6 +114,9 @@ struct FExampleNetSerializer
 
 	/** Required if bHasCustomNetReference is true. Add object references to a collector. */
 	static void CollectNetReferences(FNetSerializationContext&, const FNetCollectReferencesArgs&);
+
+	/** Serializers that want to be selective about which members to modify in the target instance when applying state should implement Apply where the serializer is responsible for setting the members of the target instance. The function operates on non-quantized state. */
+	static void Apply(FNetSerializationContext&, const FNetApplyArgs&);
 };
 UE_NET_IMPLEMENT_SERIALIZER(FExampleNetSerializer);
 
@@ -325,6 +328,16 @@ struct FNetFreeDynamicStateArgs : FNetSerializerBaseArgs
 };
 typedef void(*NetFreeDynamicStateFunction)(FNetSerializationContext&, const FNetFreeDynamicStateArgs&);
 
+/** Serializers that want to be selective about which members to modify in the target instance when applying state should implement Apply where the serializer is responsible for setting the members of the target instance. The function operates on non-quantized state. */
+struct FNetApplyArgs : FNetSerializerBaseArgs
+{
+	/** A pointer to the non-quantized source data. */
+	NetSerializerValuePointer Source;
+	/** A pointer to the non-quantized target data. */
+	NetSerializerValuePointer Target;
+};
+typedef void(*NetApplyFunction)(FNetSerializationContext&, const FNetApplyArgs&);
+
 /**
  * Various traits that can be set for a FNetSerializer.
  * These traits are typically set via constexpr bool in the declaration of the serializer.
@@ -349,6 +362,9 @@ enum class ENetSerializerTraits : uint32
 
 	/** Data replicated using this serializer should use the IsEqual implementation in order to determine whether the data is dirty or not. */
 	UseSerializerIsEqual = HasCustomNetReference << 1U,
+
+	/** Has an Apply function which should be used when applying its dequantized data to another instance. Useful for custom struct serializers where not all of the struct properties are replicated. Without a custom Apply all values will be overwritten. */
+	HasApply = UseSerializerIsEqual << 1U,
 };
 ENUM_CLASS_FLAGS(ENetSerializerTraits);
 
@@ -373,6 +389,7 @@ struct FNetSerializer
 	NetCloneDynamicStateFunction CloneDynamicState;
 	NetFreeDynamicStateFunction FreeDynamicState;
 	NetCollectNetReferencesFunction CollectNetReferences;
+	NetApplyFunction Apply;
 	const FNetSerializerConfig* DefaultConfig;
 	uint16 QuantizedTypeSize;
 	uint16 QuantizedTypeAlignment;
@@ -413,6 +430,7 @@ public:
 		Serializer.CloneDynamicState = Builder.GetCloneDynamicStateFunction();
 		Serializer.FreeDynamicState = Builder.GetFreeDynamicStateFunction();
 		Serializer.CollectNetReferences = Builder.GetCollectNetReferencesFunction();
+		Serializer.Apply = Builder.GetApplyFunction();
 
 		Serializer.DefaultConfig = Builder.GetDefaultConfig();
 
diff --git a/Engine/Source/Runtime/Experimental/Iris/Core/Public/Iris/Serialization/NetSerializerBuilder.inl b/Engine/Source/Runtime/Experimental/Iris/Core/Public/Iris/Serialization/NetSerializerBuilder.inl
--- a/Engine/Source/Runtime/Experimental/Iris/Core/Public/Iris/Serialization/NetSerializerBuilder.inl
+++ b/Engine/Source/Runtime/Experimental/Iris/Core/Public/Iris/Serialization/NetSerializerBuilder.inl
@@ -128,10 +128,10 @@ private:
 
 	struct FTraits
 	{
-		static constexpr bool bIsForwardingSerializer = true;
+		static constexpr bool bIsForwardingSerializer = false;
 		static constexpr bool bHasConnectionSpecificSerialization = false;
-		static constexpr bool bHasCustomNetReference = true;
-		static constexpr bool bHasDynamicState = true;
+		static constexpr bool bHasCustomNetReference = false;
+		static constexpr bool bHasDynamicState = false;
 		static constexpr bool bUseDefaultDelta = true;
 		static constexpr bool bUseSerializerIsEqual = false;
 	};
@@ -234,6 +234,9 @@ private:
 	template<typename U> static ETrueType TestHasCollectNetReferences(FSignatureCheck<NetCollectNetReferencesFunction, &U::CollectNetReferences>*);
 	template<typename> static EFalseType TestHasCollectNetReferences(...);
 
+	template<typename U> static ETrueType TestHasApply(FSignatureCheck<NetApplyFunction, &U::Apply>*);
+	template<typename> static EFalseType TestHasApply(...);
+
 	enum ETraits : unsigned
 	{
 		HasVersion = unsigned(decltype(TestHasVersion<NetSerializerImpl>(nullptr))::Value),
@@ -274,6 +277,7 @@ private:
 		HasFreeDynamicState = unsigned(decltype(TestHasFreeDynamicState<NetSerializerImpl>(nullptr))::Value),
 		HasCloneDynamicState = unsigned(decltype(TestHasCloneDynamicState<NetSerializerImpl>(nullptr))::Value),
 		HasCollectNetReferences = unsigned(decltype(TestHasCollectNetReferences<NetSerializerImpl>(nullptr))::Value),
+		HasApply = unsigned(decltype(TestHasApply<NetSerializerImpl>(nullptr))::Value),
 	};
 
 public:
@@ -392,6 +396,12 @@ public:
 	template<typename T = void, typename U = typename TEnableIf<!HasCollectNetReferences, T>::Type, char V = 0>
 	static NetCollectNetReferencesFunction GetCollectNetReferencesFunction() { return NetCollectNetReferencesFunction(nullptr); }
 
+	template<typename T = void, typename U = typename TEnableIf<HasApply, T>::Type, bool V = true>
+	static NetApplyFunction GetApplyFunction() { return NetSerializerImpl::Apply; }
+
+	template<typename T = void, typename U = typename TEnableIf<!HasApply, T>::Type, char V = 0>
+	static NetApplyFunction GetApplyFunction() { return NetApplyFunction(nullptr); }
+
 	// CloneDynamicState
 	template<typename T = void, typename U = typename TEnableIf<HasCloneDynamicState && (IsForwardingSerializer() || HasDynamicState()), T>::Type, bool V = true>
 	static NetCloneDynamicStateFunction GetCloneDynamicStateFunction() { return NetSerializerImpl::CloneDynamicState; }
@@ -452,6 +462,7 @@ public:
 		Traits |= (HasCustomNetReference() ? ENetSerializerTraits::HasCustomNetReference : ENetSerializerTraits::None);
 		Traits |= (HasDynamicState() ? ENetSerializerTraits::HasDynamicState : ENetSerializerTraits::None);
 		Traits |= (UseSerializerIsEqual() ? ENetSerializerTraits::UseSerializerIsEqual : ENetSerializerTraits::None);
+		Traits |= (HasApply ? ENetSerializerTraits::HasApply : ENetSerializerTraits::None);
 
 		return Traits;
 	}
