Iris:
* Added unit test calling RPCs on subobjects

#rnx

diff --git a/Engine/Plugins/Runtime/ReplicationSystemTestPlugin/Source/Private/Tests/ReplicationSystem/RPC/ReplicatedTestObjectWithRPC.cpp b/Engine/Plugins/Runtime/ReplicationSystemTestPlugin/Source/Private/Tests/ReplicationSystem/RPC/ReplicatedTestObjectWithRPC.cpp
--- a/Engine/Plugins/Runtime/ReplicationSystemTestPlugin/Source/Private/Tests/ReplicationSystem/RPC/ReplicatedTestObjectWithRPC.cpp
+++ b/Engine/Plugins/Runtime/ReplicationSystemTestPlugin/Source/Private/Tests/ReplicationSystem/RPC/ReplicatedTestObjectWithRPC.cpp
@@ -24,6 +24,18 @@ void UTestReplicatedObjectWithRPC::RegisterReplicationFragments(UE::Net::FFragme
 	UE::Net::FReplicationFragmentUtil::CreateAndRegisterFragmentsForObject(this, Context, RegistrationFlags, &this->ReplicationFragments);
 }
 
+void UTestReplicatedObjectWithRPC::Init(UReplicationSystem* InRepSystem)
+{
+	bIsServerObject = InRepSystem->IsServer();
+	ReplicationSystem = InRepSystem;
+}
+
+void UTestReplicatedObjectWithRPC::SetRootObject(UTestReplicatedObjectWithRPC* InRootObject)
+{
+	check(InRootObject);
+	RootObject = InRootObject;
+}
+
 int32 UTestReplicatedObjectWithRPC::GetFunctionCallspace(UFunction* Function, FFrame* Stack)
 {
 	check(!(Function->FunctionFlags & FUNC_Static));
@@ -84,7 +96,14 @@ int32 UTestReplicatedObjectWithRPC::GetFunctionCallspace(UFunction* Function, FF
 
 bool UTestReplicatedObjectWithRPC::CallRemoteFunction(UFunction* Function, void* Parameters, FOutParmRec* OutParms, FFrame* Stack)
 {
-	return ReplicationSystem->SendRPC(this, nullptr, Function, Parameters);
+	if (bIsSubObject)
+	{
+		return ReplicationSystem->SendRPC(RootObject, this, Function, Parameters);
+	}
+	else
+	{
+		return ReplicationSystem->SendRPC(this, nullptr, Function, Parameters);
+	}
 }
 
 void UTestReplicatedObjectWithRPC::ClientRPC_Implementation()
diff --git a/Engine/Plugins/Runtime/ReplicationSystemTestPlugin/Source/Private/Tests/ReplicationSystem/RPC/ReplicatedTestObjectWithRPC.h b/Engine/Plugins/Runtime/ReplicationSystemTestPlugin/Source/Private/Tests/ReplicationSystem/RPC/ReplicatedTestObjectWithRPC.h
--- a/Engine/Plugins/Runtime/ReplicationSystemTestPlugin/Source/Private/Tests/ReplicationSystem/RPC/ReplicatedTestObjectWithRPC.h
+++ b/Engine/Plugins/Runtime/ReplicationSystemTestPlugin/Source/Private/Tests/ReplicationSystem/RPC/ReplicatedTestObjectWithRPC.h
@@ -17,19 +17,25 @@ class UTestReplicatedObjectWithRPC : public UReplicatedTestObject
 public:
 	UTestReplicatedObjectWithRPC();
 
+	void Init(UReplicationSystem* InRepSystem);
+	void SetRootObject(UTestReplicatedObjectWithRPC* InRootObject);
+
 	virtual void RegisterReplicationFragments(UE::Net::FFragmentRegistrationContext& Context, UE::Net::EFragmentRegistrationFlags RegistrationFlags) override;
 
 	virtual int32 GetFunctionCallspace(UFunction* Function, FFrame* Stack) override;
 	virtual bool CallRemoteFunction(UFunction* Function, void* Parameters, FOutParmRec* OutParms, FFrame* Stack) override;
 
-
 	UReplicationSystem* ReplicationSystem = nullptr;
 
 	// Network data only for test
 	TArray<UE::Net::FReplicationFragment*> ReplicationFragments;
 
+	// To determine if this object is located on the server or client
 	bool bIsServerObject = false;
 
+	// Our owner when the object is a subobject
+	UTestReplicatedObjectWithRPC* RootObject = nullptr;
+
 	// RPC test functions
 	UFUNCTION(Reliable, Client)
 	void ClientRPC();
diff --git a/Engine/Plugins/Runtime/ReplicationSystemTestPlugin/Source/Private/Tests/ReplicationSystem/RPC/TestObjectRPC.cpp b/Engine/Plugins/Runtime/ReplicationSystemTestPlugin/Source/Private/Tests/ReplicationSystem/RPC/TestObjectRPC.cpp
--- a/Engine/Plugins/Runtime/ReplicationSystemTestPlugin/Source/Private/Tests/ReplicationSystem/RPC/TestObjectRPC.cpp
+++ b/Engine/Plugins/Runtime/ReplicationSystemTestPlugin/Source/Private/Tests/ReplicationSystem/RPC/TestObjectRPC.cpp
@@ -62,10 +62,7 @@ namespace UE::Net::Private
 		ClientObject->ServerRPC();
 
 		// Send and deliver client packet
-		Client->PreSendUpdate();
-		Client->SendUpdate();
-		Client->DeliverTo(*Server, 0x01, 0x01, true);
-		Client->PostSendUpdate();
+		Client->UpdateAndSend(Server);
 
 		// Verify RPC reception
 		UE_NET_ASSERT_TRUE(ServerObject->bServerRPCCalled);
@@ -74,10 +71,7 @@ namespace UE::Net::Private
 		ClientObject->ServerRPCWithParam(IntParam);
 
 		// Send and deliver client packet
-		Client->PreSendUpdate();
-		Client->SendUpdate();
-		Client->DeliverTo(*Server, 0x01, 0x01, true);
-		Client->PostSendUpdate();
+		Client->UpdateAndSend(Server);
 
 		// Verify RPC reception
 		UE_NET_ASSERT_TRUE(ServerObject->ServerRPCWithParamCalled == IntParam);
@@ -136,4 +130,79 @@ namespace UE::Net::Private
 		UE_NET_ASSERT_EQ(ClientObject->NetMulticast_MultiCastRPCCallOrder, 2);
 	}
 
+	UE_NET_TEST_FIXTURE(FRPCTestFixture, TestSubObjectRPC)
+	{
+		// Add a client
+		FReplicationSystemTestClient* Client = CreateClient();
+
+		// Spawn object on server
+		UTestReplicatedObjectWithRPC* ServerRootObject = Server->CreateObject<UTestReplicatedObjectWithRPC>();
+		ServerRootObject->Init(Server->GetReplicationSystem());
+
+		const FNetRefHandle ServerRootObjectHandle = ServerRootObject->NetRefHandle;
+		Server->ReplicationSystem->SetOwningNetConnection(ServerRootObjectHandle, 0x01);
+
+		UTestReplicatedObjectWithRPC* ServerSubObject = Server->CreateSubObject<UTestReplicatedObjectWithRPC>(ServerRootObjectHandle);
+		ServerSubObject->Init(Server->GetReplicationSystem());
+		ServerSubObject->SetRootObject(ServerRootObject);
+
+		const FNetRefHandle ServerSubObjectHandle = ServerSubObject->NetRefHandle;
+
+		// Send and deliver packet
+		Server->PreSendUpdate();
+		Server->SendAndDeliverTo(Client, true);
+		Server->PostSendUpdate();
+
+		UTestReplicatedObjectWithRPC* ClientRootObject = Cast<UTestReplicatedObjectWithRPC>(Client->GetReplicationBridge()->GetReplicatedObject(ServerRootObjectHandle));
+
+		// Verify that the root object exists on the client
+		UE_NET_ASSERT_TRUE(ClientRootObject != nullptr);
+
+		UTestReplicatedObjectWithRPC* ClientSubObject = Cast<UTestReplicatedObjectWithRPC>(Client->GetReplicationBridge()->GetReplicatedObject(ServerSubObjectHandle));
+
+		// Verify that the subobject also exists on the client
+		UE_NET_ASSERT_TRUE(ClientSubObject != nullptr);
+
+		ClientRootObject->Init(Client->GetReplicationSystem());
+		ClientSubObject->Init(Client->GetReplicationSystem());
+		ClientSubObject->SetRootObject(ClientRootObject);
+
+		// Call an RPC server->client on the subobject
+		ServerSubObject->ClientRPC();
+
+		// Send and deliver packet
+		Server->UpdateAndSend({ Client });
+
+		// Verify RPC reception
+		UE_NET_ASSERT_TRUE(ClientSubObject->bClientRPCCalled);
+
+		// Call an RPC server->client on the subobject
+		const int32 IntParam = 0xBABA;
+		ServerSubObject->ClientRPCWithParam(0xBABA);
+
+		// Send and deliver packet
+		Server->UpdateAndSend({ Client });
+
+		// Verify RPC reception
+		UE_NET_ASSERT_TRUE(ClientSubObject->ClientRPCWithParamCalled == IntParam);
+
+		// Call an RPC client->server
+		ClientSubObject->ServerRPC();
+
+		// Send and deliver client packet
+		Client->UpdateAndSend(Server);
+
+		// Verify RPC reception
+		UE_NET_ASSERT_TRUE(ServerSubObject->bServerRPCCalled);
+
+		// Call an RPC client->server
+		ClientSubObject->ServerRPCWithParam(IntParam);
+
+		// Send and deliver client packet
+		Client->UpdateAndSend(Server);
+
+		// Verify RPC reception
+		UE_NET_ASSERT_TRUE(ServerSubObject->ServerRPCWithParamCalled == IntParam);
+	}
+
 }
\ No newline at end of file
diff --git a/Engine/Plugins/Runtime/ReplicationSystemTestPlugin/Source/Private/Tests/ReplicationSystem/ReplicatedTestObject.cpp b/Engine/Plugins/Runtime/ReplicationSystemTestPlugin/Source/Private/Tests/ReplicationSystem/ReplicatedTestObject.cpp
--- a/Engine/Plugins/Runtime/ReplicationSystemTestPlugin/Source/Private/Tests/ReplicationSystem/ReplicatedTestObject.cpp
+++ b/Engine/Plugins/Runtime/ReplicationSystemTestPlugin/Source/Private/Tests/ReplicationSystem/ReplicatedTestObject.cpp
@@ -77,19 +77,22 @@ UE::Net::FNetRefHandle UReplicatedTestObjectBridge::BeginReplication(UReplicated
 	return Handle;
 }
 
-UE::Net::FNetRefHandle UReplicatedTestObjectBridge::BeginReplication(FNetRefHandle OwnerHandle, UReplicatedTestObject* Instance, FNetRefHandle InsertRelativeToSubObjectHandle, ESubObjectInsertionOrder InsertionOrder)
+UE::Net::FNetRefHandle UReplicatedTestObjectBridge::BeginReplication(FNetRefHandle OwnerHandle, UReplicatedTestObject* SubObjectInstance, FNetRefHandle InsertRelativeToSubObjectHandle, ESubObjectInsertionOrder InsertionOrder)
 {
+	check(OwnerHandle.IsValid());
+
 	// Create NetRefHandle for the registered fragments
 	Super::FCreateNetRefHandleParams Params = Super::DefaultCreateNetRefHandleParams;
 	Params.bCanReceive = true;
 
 	// Create NetRefHandle for the registered fragments
-	FNetRefHandle Handle = Super::BeginReplication(OwnerHandle, Instance, InsertRelativeToSubObjectHandle, InsertionOrder, Params);
+	FNetRefHandle Handle = Super::BeginReplication(OwnerHandle, SubObjectInstance, InsertRelativeToSubObjectHandle, InsertionOrder, Params);
 
 	if (Handle.IsValid())
 	{
 		// This is optional but typically we want to cache at least the NetRefHandle in the game instance to avoid doing map lookups to find it
-		Instance->NetRefHandle = Handle;
+		SubObjectInstance->NetRefHandle = Handle;
+		SubObjectInstance->bIsSubObject = true;
 	}
 	
 	return Handle;
@@ -171,6 +174,11 @@ FObjectReplicationBridgeInstantiateResult UReplicatedTestObjectBridge::BeginInst
 
 	FStaticConstructObjectParameters ConstructObjectParameters(ArcheType->GetClass());
 	UObject* CreatedObject = StaticConstructObject_Internal(ConstructObjectParameters);
+
+	if (UReplicatedTestObject* BaseTestObject = Cast<UReplicatedTestObject>(CreatedObject))
+	{
+		BaseTestObject->bIsSubObject = RootObjectOfSubObject.IsValid();
+	}
 	
 	if (UTestReplicatedIrisObject* CreatedTestObject = Cast<UTestReplicatedIrisObject>(CreatedObject))
 	{
diff --git a/Engine/Plugins/Runtime/ReplicationSystemTestPlugin/Source/Private/Tests/ReplicationSystem/ReplicatedTestObject.h b/Engine/Plugins/Runtime/ReplicationSystemTestPlugin/Source/Private/Tests/ReplicationSystem/ReplicatedTestObject.h
--- a/Engine/Plugins/Runtime/ReplicationSystemTestPlugin/Source/Private/Tests/ReplicationSystem/ReplicatedTestObject.h
+++ b/Engine/Plugins/Runtime/ReplicationSystemTestPlugin/Source/Private/Tests/ReplicationSystem/ReplicatedTestObject.h
@@ -39,6 +39,9 @@ public:
 
 	/** If this is set to true, this instance will fail to instantiate on remote end. */
 	bool bForceFailToInstantiateOnRemote = false;
+
+	// To determine if this object is a root object or a subobject
+	bool bIsSubObject = false;
 };
 
 USTRUCT()
diff --git a/Engine/Plugins/Runtime/ReplicationSystemTestPlugin/Source/Private/Tests/ReplicationSystem/ReplicationSystemServerClientTestFixture.cpp b/Engine/Plugins/Runtime/ReplicationSystemTestPlugin/Source/Private/Tests/ReplicationSystem/ReplicationSystemServerClientTestFixture.cpp
--- a/Engine/Plugins/Runtime/ReplicationSystemTestPlugin/Source/Private/Tests/ReplicationSystem/ReplicationSystemServerClientTestFixture.cpp
+++ b/Engine/Plugins/Runtime/ReplicationSystemTestPlugin/Source/Private/Tests/ReplicationSystem/ReplicationSystemServerClientTestFixture.cpp
@@ -352,16 +352,40 @@ float FReplicationSystemTestNode::ConvertPollPeriodIntoFrequency(uint32 PollPeri
 	return PollFrequency;
 }
 
-// FReplicationSystemTestClient implementation
+//*****************************************************************************
+// Class FReplicationSystemTestClient
+//*****************************************************************************
+
 FReplicationSystemTestClient::FReplicationSystemTestClient(const TCHAR* Name)
 : FReplicationSystemTestNode(false, Name)
 , ConnectionIdOnServer(~0U)
 {
 }
 
-// FReplicationSystemTestServer implementation
+bool FReplicationSystemTestClient::UpdateAndSend(FReplicationSystemTestServer* Server, bool bDeliver)
+{
+	bool bSuccess = false;
+
+	PreSendUpdate();
+
+	if (SendUpdate())
+	{
+		constexpr uint32 ServerRemoteConnectionId = 0x01;
+		DeliverTo(*Server, LocalConnectionId, ServerRemoteConnectionId, bDeliver);
+		bSuccess = true;
+	}
+
+	PostSendUpdate();
+
+	return bSuccess;
+}
+
+//*****************************************************************************
+// Class FReplicationSystemTestServer
+//*****************************************************************************
+
 FReplicationSystemTestServer::FReplicationSystemTestServer(const TCHAR* Name)
-: FReplicationSystemTestNode(true, Name)
+	: FReplicationSystemTestNode(true, Name)
 {
 }
 
@@ -406,7 +430,10 @@ bool FReplicationSystemTestServer::UpdateAndSend(const TArray<FReplicationSystem
 	return bSuccess;
 }
 
-// FReplicationSystemServerClientTestFixture implementation
+//*****************************************************************************
+// Class FReplicationSystemServerClientTestFixture
+//*****************************************************************************
+
 void FReplicationSystemServerClientTestFixture::SetUp()
 {
 	FNetworkAutomationTestSuiteFixture::SetUp();
diff --git a/Engine/Plugins/Runtime/ReplicationSystemTestPlugin/Source/Private/Tests/ReplicationSystem/ReplicationSystemServerClientTestFixture.h b/Engine/Plugins/Runtime/ReplicationSystemTestPlugin/Source/Private/Tests/ReplicationSystem/ReplicationSystemServerClientTestFixture.h
--- a/Engine/Plugins/Runtime/ReplicationSystemTestPlugin/Source/Private/Tests/ReplicationSystem/ReplicationSystemServerClientTestFixture.h
+++ b/Engine/Plugins/Runtime/ReplicationSystemTestPlugin/Source/Private/Tests/ReplicationSystem/ReplicationSystemServerClientTestFixture.h
@@ -169,6 +169,9 @@ class FReplicationSystemTestClient : public FReplicationSystemTestNode
 public:
 	FReplicationSystemTestClient(const TCHAR* Name);
 
+	// Tick and send packets to the server
+	bool UpdateAndSend(class FReplicationSystemTestServer* Server, bool bDeliver = true);
+
 	uint32 ConnectionIdOnServer;
 	uint32 LocalConnectionId;
 };
