Iris
- Implemented optimized path only dequantizing dirty members of received statedata during apply

diff --git a/Engine/Plugins/Runtime/ReplicationSystemTestPlugin/Source/Private/Tests/ReplicationSystem/ReplicatedTestObject.cpp b/Engine/Plugins/Runtime/ReplicationSystemTestPlugin/Source/Private/Tests/ReplicationSystem/ReplicatedTestObject.cpp
--- a/Engine/Plugins/Runtime/ReplicationSystemTestPlugin/Source/Private/Tests/ReplicationSystem/ReplicatedTestObject.cpp
+++ b/Engine/Plugins/Runtime/ReplicationSystemTestPlugin/Source/Private/Tests/ReplicationSystem/ReplicatedTestObject.cpp
@@ -517,6 +517,40 @@ void UReplicatedSubObjectOrderObject::RegisterReplicationFragments(UE::Net::FFra
 	}
 }
 
+//////////////////////////////////////////////////////////////////////////
+// Implementation for UTestReplicatedObjectWithRepNotifies
+//////////////////////////////////////////////////////////////////////////
+void UTestReplicatedObjectWithRepNotifies::GetLifetimeReplicatedProps(TArray< class FLifetimeProperty >& OutLifetimeProps) const
+{
+	DOREPLIFETIME_CONDITION_NOTIFY(UTestReplicatedObjectWithRepNotifies, IntA, COND_None, REPNOTIFY_Always);
+	DOREPLIFETIME_CONDITION_NOTIFY(UTestReplicatedObjectWithRepNotifies, IntB, COND_None, REPNOTIFY_Always);
+	DOREPLIFETIME(UTestReplicatedObjectWithRepNotifies, IntC);
+}
+
+UTestReplicatedObjectWithRepNotifies::UTestReplicatedObjectWithRepNotifies()
+	: UReplicatedTestObject()
+{
+}
+
+void UTestReplicatedObjectWithRepNotifies::RegisterReplicationFragments(UE::Net::FFragmentRegistrationContext& Context, UE::Net::EFragmentRegistrationFlags RegistrationFlags)
+{
+	// Base object owns the fragment in this case
+	{
+		this->ReplicationFragments.Reset();
+		UE::Net::FReplicationFragmentUtil::CreateAndRegisterFragmentsForObject(this, Context, RegistrationFlags, &this->ReplicationFragments);
+	}
+}
+
+void UTestReplicatedObjectWithRepNotifies::OnRep_IntA(int32 OldInt)
+{
+	PrevIntAStoredInOnRep = OldInt;
+}
+
+void UTestReplicatedObjectWithRepNotifies::OnRep_IntB(int32 OldInt)
+{
+	PrevIntBStoredInOnRep = OldInt;
+}
+
 //////////////////////////////////////////////////////////////////////////
 // Implementation for UTestReplicatedIrisObjectWithDynamicCondition
 //////////////////////////////////////////////////////////////////////////
diff --git a/Engine/Plugins/Runtime/ReplicationSystemTestPlugin/Source/Private/Tests/ReplicationSystem/ReplicatedTestObject.h b/Engine/Plugins/Runtime/ReplicationSystemTestPlugin/Source/Private/Tests/ReplicationSystem/ReplicatedTestObject.h
--- a/Engine/Plugins/Runtime/ReplicationSystemTestPlugin/Source/Private/Tests/ReplicationSystem/ReplicatedTestObject.h
+++ b/Engine/Plugins/Runtime/ReplicationSystemTestPlugin/Source/Private/Tests/ReplicationSystem/ReplicatedTestObject.h
@@ -428,6 +428,43 @@ public:
 	TArray<UE::Net::FReplicationFragment*> ReplicationFragments;
 };
 
+UCLASS()
+class UTestReplicatedObjectWithRepNotifies : public UReplicatedTestObject
+{
+	GENERATED_BODY()
+
+public:
+
+	UTestReplicatedObjectWithRepNotifies();
+
+	// Network interface must be part of base.
+	virtual void RegisterReplicationFragments(UE::Net::FFragmentRegistrationContext& Fragments, UE::Net::EFragmentRegistrationFlags RegistrationFlags) override;
+
+public:
+
+	UPROPERTY(Transient, ReplicatedUsing=OnRep_IntA)
+	int32 IntA = -1;
+	int32 PrevIntAStoredInOnRep = -1;
+
+	UPROPERTY(Transient, ReplicatedUsing=OnRep_IntB)
+	int32 IntB = -1;
+	int32 PrevIntBStoredInOnRep = -1;
+
+	UPROPERTY(Transient, Replicated)
+	int8 IntC = -1;
+
+public:
+	UFUNCTION()
+	void OnRep_IntA(int32 OldInt);
+
+	UFUNCTION()
+	void OnRep_IntB(int32 OldInt);
+
+	// Network data only for test
+	TArray<UE::Net::FReplicationFragment*> ReplicationFragments;
+};
+
+
 
 /**
  *  A test class for Replication on an object with no replicated members
diff --git a/Engine/Plugins/Runtime/ReplicationSystemTestPlugin/Source/Private/Tests/ReplicationSystem/TestReplicationSystem.cpp b/Engine/Plugins/Runtime/ReplicationSystemTestPlugin/Source/Private/Tests/ReplicationSystem/TestReplicationSystem.cpp
--- a/Engine/Plugins/Runtime/ReplicationSystemTestPlugin/Source/Private/Tests/ReplicationSystem/TestReplicationSystem.cpp
+++ b/Engine/Plugins/Runtime/ReplicationSystemTestPlugin/Source/Private/Tests/ReplicationSystem/TestReplicationSystem.cpp
@@ -1801,4 +1801,107 @@ UE_NET_TEST_FIXTURE(FReplicationSystemServerClientTestFixture, TestClientCanSkip
 }
 
 
+// Test that PropertyReplication properly handles partial states during Apply
+UE_NET_TEST_FIXTURE(FReplicationSystemServerClientTestFixture, TestPartialDequantize)
+{
+	// Enable cvars to exercise path that store previous state for OnReps to make sure we exercise path that accumulate dirty changes so that we have a complete state.
+	IConsoleVariable* CVarUsePrevReceivedStateForOnReps = IConsoleManager::Get().FindConsoleVariable(TEXT("net.Iris.UsePrevReceivedStateForOnReps"));
+	check(CVarUsePrevReceivedStateForOnReps != nullptr && CVarUsePrevReceivedStateForOnReps->IsVariableBool());
+	const bool bUsePrevReceivedStateForOnReps = CVarUsePrevReceivedStateForOnReps->GetBool();
+	CVarUsePrevReceivedStateForOnReps->Set(true, ECVF_SetByCode);
+
+	// Make sure we allow partial dequantize
+	IConsoleVariable* CVarForceFullDequantizeAndApply = IConsoleManager::Get().FindConsoleVariable(TEXT("net.iris.ForceFullDequantizeAndApply"));
+	check(CVarForceFullDequantizeAndApply != nullptr && CVarForceFullDequantizeAndApply->IsVariableBool());
+	const bool bForceFullDequantizeAndApply = CVarForceFullDequantizeAndApply->GetBool();
+	CVarForceFullDequantizeAndApply->Set(false, ECVF_SetByCode);
+
+	// Add a client
+	FReplicationSystemTestClient* Client = CreateClient();
+
+	// Spawn objects on server
+	UTestReplicatedObjectWithRepNotifies* ServerObjectA = Server->CreateObject<UTestReplicatedObjectWithRepNotifies>();
+	
+	// Send and deliver packet
+	Server->PreSendUpdate();
+	Server->SendAndDeliverTo(Client, true);
+	Server->PostSendUpdate();
+
+	// Verify assumptions
+	// Object should exist on client and have default state
+	UTestReplicatedObjectWithRepNotifies* ClientObjectA = Cast<UTestReplicatedObjectWithRepNotifies>(Client->GetReplicationBridge()->GetReplicatedObject(ServerObjectA->NetRefHandle));
+	UE_NET_ASSERT_NE(ClientObjectA, nullptr);
+
+	UE_NET_ASSERT_EQ(ServerObjectA->IntA, ClientObjectA->IntA);
+	UE_NET_ASSERT_EQ(ClientObjectA->PrevIntAStoredInOnRep, -1);
+	UE_NET_ASSERT_EQ(ServerObjectA->IntB, ClientObjectA->IntB);
+	UE_NET_ASSERT_EQ(ClientObjectA->PrevIntBStoredInOnRep, -1);
+	UE_NET_ASSERT_EQ(ServerObjectA->IntC, ClientObjectA->IntC);
+
+	// Modify only IntA
+	ServerObjectA->IntA = 1;
+
+	// Send and deliver packet
+	Server->UpdateAndSend({Client});
+
+	// Verify assumptions
+	// Only IntA should have been modified
+	UE_NET_ASSERT_EQ(ServerObjectA->IntA, ClientObjectA->IntA);
+	UE_NET_ASSERT_EQ(ClientObjectA->PrevIntAStoredInOnRep, -1);
+	UE_NET_ASSERT_EQ(ServerObjectA->IntB, ClientObjectA->IntB);
+	UE_NET_ASSERT_EQ(ClientObjectA->PrevIntBStoredInOnRep, -1);
+	UE_NET_ASSERT_EQ(ServerObjectA->IntC, ClientObjectA->IntC);
+
+	// Modify only IntB
+	ServerObjectA->IntB = 1;
+
+	// Send and deliver packet
+	Server->UpdateAndSend({Client});
+
+	// Verify assumptions
+	// Only IntA should have been modified
+	UE_NET_ASSERT_EQ(ServerObjectA->IntA, ClientObjectA->IntA);
+	UE_NET_ASSERT_EQ(ClientObjectA->PrevIntAStoredInOnRep, -1);
+
+	UE_NET_ASSERT_EQ(ServerObjectA->IntB, ClientObjectA->IntB);
+	UE_NET_ASSERT_EQ(ClientObjectA->PrevIntBStoredInOnRep, -1);
+	UE_NET_ASSERT_EQ(ServerObjectA->IntC, ClientObjectA->IntC);
+
+	// Modify only IntA
+	ServerObjectA->IntA = 2;
+
+	// Send and deliver packet
+	Server->UpdateAndSend({Client});
+
+	// Verify assumptions
+	// IntA should have been modified, and if everything works correctly PrevIntAStoredInOnRep should be 1
+	UE_NET_ASSERT_EQ(ServerObjectA->IntA, ClientObjectA->IntA);
+	UE_NET_ASSERT_EQ(ClientObjectA->PrevIntAStoredInOnRep, 1);
+
+	UE_NET_ASSERT_EQ(ServerObjectA->IntB, ClientObjectA->IntB);
+	UE_NET_ASSERT_EQ(ClientObjectA->PrevIntBStoredInOnRep, -1);
+	UE_NET_ASSERT_EQ(ServerObjectA->IntC, ClientObjectA->IntC);
+
+	// Verify that we do not apply repnotifies if we do not receive data from server by modifying values on the client and verifying that they do not get overwritten
+	ServerObjectA->IntB = 2;
+	ClientObjectA->IntA = -1;
+	ClientObjectA->PrevIntAStoredInOnRep = -1;
+
+	// Send and deliver packet
+	Server->UpdateAndSend({Client});
+
+	// Verify assumptions, since we messed with IntA and PrevIntAStoredInOnRep locally they have the value we set but IntB should be updated according to replicated state
+	UE_NET_ASSERT_NE(ServerObjectA->IntA, ClientObjectA->IntA);
+	UE_NET_ASSERT_EQ(ClientObjectA->PrevIntAStoredInOnRep, -1);
+	UE_NET_ASSERT_EQ(ServerObjectA->IntB, ClientObjectA->IntB);
+	UE_NET_ASSERT_EQ(ClientObjectA->PrevIntBStoredInOnRep, 1);
+	UE_NET_ASSERT_EQ(ServerObjectA->IntC, ClientObjectA->IntC);
+
+	// Restore cvars
+	CVarUsePrevReceivedStateForOnReps->Set(bUsePrevReceivedStateForOnReps, ECVF_SetByCode);
+	CVarForceFullDequantizeAndApply->Set(bForceFullDequantizeAndApply, ECVF_SetByCode);
+}
+
+
+
 }
diff --git a/Engine/Source/Runtime/Experimental/Iris/Core/Private/Iris/ReplicationState/InternalPropertyReplicationState.cpp b/Engine/Source/Runtime/Experimental/Iris/Core/Private/Iris/ReplicationState/InternalPropertyReplicationState.cpp
--- a/Engine/Source/Runtime/Experimental/Iris/Core/Private/Iris/ReplicationState/InternalPropertyReplicationState.cpp
+++ b/Engine/Source/Runtime/Experimental/Iris/Core/Private/Iris/ReplicationState/InternalPropertyReplicationState.cpp
@@ -28,12 +28,19 @@ void InitReplicationStateInternals(uint8* StateBuffer, const FReplicationStateDe
 	}
 }
 
-void CopyPropertyReplicationStateInternals(uint8* RESTRICT DstStateBuffer, uint8* RESTRICT SrcStateBuffer, const FReplicationStateDescriptor* Descriptor)
+void CopyPropertyReplicationStateInternals(uint8* RESTRICT DstStateBuffer, uint8* RESTRICT SrcStateBuffer, const FReplicationStateDescriptor* Descriptor, bool bOverwriteChangeMask = true)
 {
 	FNetBitArrayView DstChangeMask = GetMemberChangeMask(DstStateBuffer, Descriptor);
 	FNetBitArrayView SrcChangeMask = GetMemberChangeMask(SrcStateBuffer, Descriptor);
 
-	DstChangeMask.Copy(SrcChangeMask);
+	if (bOverwriteChangeMask)
+	{
+		DstChangeMask.Copy(SrcChangeMask);
+	}
+	else
+	{
+		DstChangeMask.Combine(SrcChangeMask, FNetBitArrayView::OrOp);
+	}
 
 	// Copy optional conditional changemask
 	if (EnumHasAnyFlags(Descriptor->Traits, EReplicationStateTraits::HasLifetimeConditionals))
@@ -41,7 +48,14 @@ void CopyPropertyReplicationStateInternals(uint8* RESTRICT DstStateBuffer, uint8
 		FNetBitArrayView DstConditionalChangeMask = GetMemberConditionalChangeMask(DstStateBuffer, Descriptor);
 		FNetBitArrayView SrcConditionalChangeMask = GetMemberConditionalChangeMask(SrcStateBuffer, Descriptor);
 
-		DstConditionalChangeMask.Copy(SrcConditionalChangeMask);
+		if (bOverwriteChangeMask)
+		{
+			DstConditionalChangeMask.Copy(SrcConditionalChangeMask);
+		}
+		else
+		{
+			DstConditionalChangeMask.Combine(SrcConditionalChangeMask, FNetBitArrayView::OrOp);
+		}
 	}
 }
 
@@ -122,6 +136,39 @@ void CopyPropertyReplicationState(uint8* RESTRICT DstStateBuffer, uint8* RESTRIC
 	}
 }
 
+void CopyDirtyMembers(uint8* RESTRICT DstStateBuffer, uint8* RESTRICT SrcStateBuffer, const FReplicationStateDescriptor* Descriptor)
+{
+	check(IsAligned(DstStateBuffer, Descriptor->ExternalAlignment) && IsAligned(SrcStateBuffer, Descriptor->ExternalAlignment));
+
+	// Merge changemasks
+	const bool bOverwriteChangeMask = false;
+	CopyPropertyReplicationStateInternals(DstStateBuffer, SrcStateBuffer, Descriptor, bOverwriteChangeMask);
+
+	// copy dirty members
+	const FReplicationStateMemberDescriptor* MemberDescriptors = Descriptor->MemberDescriptors;
+	const FProperty** MemberProperties = Descriptor->MemberProperties;
+	const FReplicationStateMemberPropertyDescriptor* MemberPropertyDescriptors = Descriptor->MemberPropertyDescriptors;
+	FNetBitArrayView DirtyStates = GetMemberChangeMask(SrcStateBuffer, Descriptor);
+	
+	const uint32 MemberCount = Descriptor->MemberCount;
+
+	const bool bIsInitState = Descriptor->IsInitState();
+
+	for (uint32 MemberIt = 0; MemberIt < MemberCount; ++MemberIt)
+	{
+		const FReplicationStateMemberPropertyDescriptor& MemberPropertyDescriptor = MemberPropertyDescriptors[MemberIt];
+		const FReplicationStateMemberChangeMaskDescriptor& ChangeMaskInfo = Descriptor->MemberChangeMaskDescriptors[MemberIt];
+
+		const bool bShouldCopyProperty = bIsInitState || DirtyStates.IsAnyBitSet(ChangeMaskInfo.BitOffset, ChangeMaskInfo.BitCount);
+		if (bShouldCopyProperty && MemberPropertyDescriptor.ArrayIndex == 0)
+		{
+			const FReplicationStateMemberDescriptor& MemberDescriptor = MemberDescriptors[MemberIt];
+			const FProperty* Property = MemberProperties[MemberIt];
+			Property->CopyCompleteValue(DstStateBuffer + MemberDescriptor.ExternalMemberOffset, SrcStateBuffer + MemberDescriptor.ExternalMemberOffset);
+		}
+	}
+}
+
 /**
  * Structs and Arrays need to be handled carefully. Our intermediate representation need 
  * to match the actual native representation perfectly. We are careful not
diff --git a/Engine/Source/Runtime/Experimental/Iris/Core/Private/Iris/ReplicationState/InternalPropertyReplicationState.h b/Engine/Source/Runtime/Experimental/Iris/Core/Private/Iris/ReplicationState/InternalPropertyReplicationState.h
--- a/Engine/Source/Runtime/Experimental/Iris/Core/Private/Iris/ReplicationState/InternalPropertyReplicationState.h
+++ b/Engine/Source/Runtime/Experimental/Iris/Core/Private/Iris/ReplicationState/InternalPropertyReplicationState.h
@@ -22,6 +22,9 @@ IRISCORE_API void DestructPropertyReplicationState(uint8* StateBuffer, const FRe
 /** Copy external state in SrcStateBuffer to external state DstStateBuffer including changemask */
 IRISCORE_API void CopyPropertyReplicationState(uint8* RESTRICT DstStateBuffer, uint8* RESTRICT SrcStateBuffer, const FReplicationStateDescriptor* Descriptor);
 
+/** Copy only properties from external state in SrcStateBuffer to external state DstStateBufferif dirty, destination changemask will also be updated */
+IRISCORE_API void CopyDirtyMembers(uint8* RESTRICT DstStateBuffer, uint8* RESTRICT SrcStateBuffer, const FReplicationStateDescriptor* Descriptor);
+
 /** Copy property value to target. If the property is fully replicated we use the properties copy function, otherwise we do a per member copy of all replicated data. Members with NetSerializers with an Apply() function will call that. */
 IRISCORE_API void InternalApplyPropertyValue(const FReplicationStateDescriptor* Descriptor, uint32 MemberIndex, void* RESTRICT Dst, const void* RESTRICT Src);
 
diff --git a/Engine/Source/Runtime/Experimental/Iris/Core/Private/Iris/ReplicationState/PropertyReplicationState.cpp b/Engine/Source/Runtime/Experimental/Iris/Core/Private/Iris/ReplicationState/PropertyReplicationState.cpp
--- a/Engine/Source/Runtime/Experimental/Iris/Core/Private/Iris/ReplicationState/PropertyReplicationState.cpp
+++ b/Engine/Source/Runtime/Experimental/Iris/Core/Private/Iris/ReplicationState/PropertyReplicationState.cpp
@@ -375,6 +375,21 @@ void FPropertyReplicationState::PushPropertyReplicationState(const UObject* Owne
 	}
 }
 
+void FPropertyReplicationState::CopyDirtyProperties(const FPropertyReplicationState& Other)
+{
+	check(this != &Other && IsValid());
+	check(ReplicationStateDescriptor.GetReference() == Other.ReplicationStateDescriptor.GetReference());
+
+	if (!Private::IsReplicationStateBound(StateBuffer, ReplicationStateDescriptor.GetReference()))
+	{
+		Private::CopyDirtyMembers(StateBuffer, Other.StateBuffer, ReplicationStateDescriptor.GetReference());
+	}
+	else
+	{
+		Set(Other);
+	}
+}
+
 bool FPropertyReplicationState::PollObjectReferences(const void* RESTRICT SrcStateData)
 {
 	if (IsValid())
diff --git a/Engine/Source/Runtime/Experimental/Iris/Core/Private/Iris/ReplicationSystem/DequantizeAndApplyHelper.cpp b/Engine/Source/Runtime/Experimental/Iris/Core/Private/Iris/ReplicationSystem/DequantizeAndApplyHelper.cpp
--- a/Engine/Source/Runtime/Experimental/Iris/Core/Private/Iris/ReplicationSystem/DequantizeAndApplyHelper.cpp
+++ b/Engine/Source/Runtime/Experimental/Iris/Core/Private/Iris/ReplicationSystem/DequantizeAndApplyHelper.cpp
@@ -2,6 +2,7 @@
 
 #include "DequantizeAndApplyHelper.h"
 #include "Containers/ArrayView.h"
+#include "HAL/IConsoleManager.h"
 #include "Iris/Core/IrisProfiler.h"
 #include "Net/Core/NetBitArray.h"
 #include "Iris/ReplicationState/ReplicationStateDescriptor.h"
@@ -16,6 +17,12 @@
 namespace UE::Net::Private
 {
 
+static bool bCVarForceFullDequantizeAndApply = false;
+static FAutoConsoleVariableRef CVarForceFullDequantizeAndApply(
+	TEXT("net.iris.ForceFullDequantizeAndApply"),
+	bCVarForceFullDequantizeAndApply,
+	TEXT("When enabled a full dequantize of dirty states will be used when applying received statedata regardless of traits set in the fragments."));
+
 struct FDequantizeAndApplyHelper::FContext
 {
 	struct FStateData
@@ -93,7 +100,8 @@ FDequantizeAndApplyHelper::FContext* FDequantizeAndApplyHelper::Initialize(FNetS
 		ChangeMaskReader.ReadBitStream(ChangeMaskStorage, ChangeMaskBitCount);
 
 		// Cache all ReplicationStates with dirty changes
-		const bool bShouldDequantizeState = ChangeMask.IsAnyBitSet() || (NetSerializationContext.IsInitState() && CurrentDescriptor->IsInitState());
+		const bool bIsInitState = NetSerializationContext.IsInitState() && CurrentDescriptor->IsInitState();
+		const bool bShouldDequantizeState = ChangeMask.IsAnyBitSet() || bIsInitState;
 		if (bShouldDequantizeState)
 		{
 			FContext::FStateData& StateData = CachedStateData[CachedStateCount];
@@ -112,8 +120,17 @@ FDequantizeAndApplyHelper::FContext* FDequantizeAndApplyHelper::Initialize(FNetS
 				FNetBitArrayView DestChangeMask = Private::GetMemberChangeMask(StateBuffer, CurrentDescriptor);
 				DestChangeMask.Copy(ChangeMask);
 	
-				// Dequantize state data
-				FReplicationStateOperations::Dequantize(NetSerializationContext, StateBuffer, (uint8*)CurrentInternalStateBuffer, CurrentDescriptor);
+				// Dequantize state data, if the fragment supports partial dequantize we will only dequantize dirty members
+				const bool bUseFullDequantizeAndApply = bIsInitState || bCVarForceFullDequantizeAndApply || !EnumHasAnyFlags(StateData.Fragment->GetTraits(), EReplicationFragmentTraits::SupportsPartialDequantizedState);
+				if (bUseFullDequantizeAndApply)
+				{
+					FReplicationStateOperations::Dequantize(NetSerializationContext, StateBuffer, (uint8*)CurrentInternalStateBuffer, CurrentDescriptor);
+				}
+				else
+				{
+					NetSerializationContext.SetChangeMask(&ChangeMask);
+					FReplicationStateOperations::DequantizeWithMask(NetSerializationContext, ChangeMask, 0U, StateBuffer, (uint8*)CurrentInternalStateBuffer, CurrentDescriptor);
+				}
 
 				StateData.StateBufferData.ExternalStateBuffer = StateBuffer;
 			}
diff --git a/Engine/Source/Runtime/Experimental/Iris/Core/Private/Iris/ReplicationSystem/PropertyReplicationFragment.cpp b/Engine/Source/Runtime/Experimental/Iris/Core/Private/Iris/ReplicationSystem/PropertyReplicationFragment.cpp
--- a/Engine/Source/Runtime/Experimental/Iris/Core/Private/Iris/ReplicationSystem/PropertyReplicationFragment.cpp
+++ b/Engine/Source/Runtime/Experimental/Iris/Core/Private/Iris/ReplicationSystem/PropertyReplicationFragment.cpp
@@ -73,6 +73,8 @@ FPropertyReplicationFragment::FPropertyReplicationFragment(EReplicationFragmentT
 	}
 
 	Traits |= EReplicationFragmentTraits::HasPropertyReplicationState;
+	// We can handle partial state in all apply operations
+	Traits |= EReplicationFragmentTraits::SupportsPartialDequantizedState;
 }
 
 FPropertyReplicationFragment::~FPropertyReplicationFragment() = default;
@@ -106,7 +108,16 @@ void FPropertyReplicationFragment::CallRepNotifies(FReplicationStateApplyContext
 		// If we rely on received data for the onreps, we just copy the received state, otherwise we must store the local state before applying received data.
 		if (bUsePrevReceivedStateForOnReps && PrevReplicationState)
 		{
-			*PrevReplicationState = ReceivedState;
+			// Init is always a full state so we can just copy it
+			if (Context.bIsInit)
+			{
+				*PrevReplicationState = ReceivedState;
+			}
+			else
+			{
+				// As apply now might provide us with partial states we should only copy dirty members.
+				PrevReplicationState->CopyDirtyProperties(ReceivedState);
+			}
 		}
 	}
 	else
diff --git a/Engine/Source/Runtime/Experimental/Iris/Core/Public/Iris/ReplicationState/PropertyReplicationState.h b/Engine/Source/Runtime/Experimental/Iris/Core/Public/Iris/ReplicationState/PropertyReplicationState.h
--- a/Engine/Source/Runtime/Experimental/Iris/Core/Public/Iris/ReplicationState/PropertyReplicationState.h
+++ b/Engine/Source/Runtime/Experimental/Iris/Core/Public/Iris/ReplicationState/PropertyReplicationState.h
@@ -101,6 +101,11 @@ public:
 	 */
 	IRISCORE_API bool PollObjectReferences(const void* RESTRICT SrcData);
 
+	/**
+	 * Copy dirty properties from the other state including changemask
+	 */
+	IRISCORE_API void CopyDirtyProperties(const FPropertyReplicationState& Other);
+
 	struct FCallRepNotifiesParameters
 	{
 		// Previous state if requested
diff --git a/Engine/Source/Runtime/Experimental/Iris/Core/Public/Iris/ReplicationSystem/ReplicationFragment.h b/Engine/Source/Runtime/Experimental/Iris/Core/Public/Iris/ReplicationSystem/ReplicationFragment.h
--- a/Engine/Source/Runtime/Experimental/Iris/Core/Public/Iris/ReplicationSystem/ReplicationFragment.h
+++ b/Engine/Source/Runtime/Experimental/Iris/Core/Public/Iris/ReplicationSystem/ReplicationFragment.h
@@ -92,23 +92,55 @@ enum class EReplicationStateToStringFlags : uint32
 };
 ENUM_CLASS_FLAGS(EReplicationStateToStringFlags);
 
+/** Traits describing a ReplicationFragmet */
 enum class EReplicationFragmentTraits : uint32
 {
 	None								= 0,
+
+	// Not implemented
 	HasInterpolation					= 1,
+
+	// Fragment has rep notifies
 	HasRepNotifies						= HasInterpolation << 1,
+
+	// Save previous state before apply
 	KeepPreviousState					= HasRepNotifies << 1,
+
+	// Fragment is owned by instance protocol and will be destroyed when the instance protocol is destroyed
 	DeleteWithInstanceProtocol			= KeepPreviousState << 1,
+
+	// Pass raw quantized data when applying state data
 	HasPersistentTargetStateBuffer		= DeleteWithInstanceProtocol << 1,
+
+	// Fragment can be used to source replication data
 	CanReplicate						= HasPersistentTargetStateBuffer << 1,
+
+	// Fragment can receive replication data
 	CanReceive							= CanReplicate << 1,
+
+	// Fragment requires polling to detect dirtiness
 	NeedsPoll							= CanReceive << 1,
+
+	// Fragment requires legacy callbacks
 	NeedsLegacyCallbacks				= NeedsPoll << 1,
+
+	// Fragment requires PreSendUpdate to be called 
 	NeedsPreSendUpdate					= NeedsLegacyCallbacks << 1,
+
+	// Fragment requires world location to be updated before filtering and prioritization
 	NeedsWorldLocationUpdate			= NeedsPreSendUpdate << 1,
+
+	// Fragment supports push based dirtiness
 	HasPushBasedDirtiness				= NeedsWorldLocationUpdate << 1,
+
+	// Fragment is a PropertyReplication, or is using PropertyReplicationState
 	HasPropertyReplicationState			= HasPushBasedDirtiness << 1,
+
+	// Fragment has object references
 	HasObjectReference					= HasPropertyReplicationState << 1,
+
+	// Fragment supports partial dequantized state in apply
+	SupportsPartialDequantizedState		= HasObjectReference << 1,
 };
 ENUM_CLASS_FLAGS(EReplicationFragmentTraits);
 
