Iris - Don't apply state or call rep notify functions unless we've received the state or resolved object references. The behavior can be controlled via the cvar net.Iris.DispatchUnresolvedPreviouslyReceivedChanges which defaults to false.

diff --git a/Engine/Plugins/Runtime/ReplicationSystemTestPlugin/Source/Private/Tests/ReplicationState/TestPropertyReplicationState.cpp b/Engine/Plugins/Runtime/ReplicationSystemTestPlugin/Source/Private/Tests/ReplicationState/TestPropertyReplicationState.cpp
--- a/Engine/Plugins/Runtime/ReplicationSystemTestPlugin/Source/Private/Tests/ReplicationState/TestPropertyReplicationState.cpp
+++ b/Engine/Plugins/Runtime/ReplicationSystemTestPlugin/Source/Private/Tests/ReplicationState/TestPropertyReplicationState.cpp
@@ -7,11 +7,13 @@
 #include "Iris/ReplicationState/InternalPropertyReplicationState.h"
 #include "Iris/ReplicationState/ReplicationStateDescriptorBuilder.h"
 #include "Iris/ReplicationState/ReplicationStateUtil.h"
+#include "Iris/ReplicationSystem/ReplicationFragmentUtil.h"
 #include "Iris/Serialization/NetSerializers.h"
 #include "Misc/EnumClassFlags.h"
 #include "UObject/StrongObjectPtr.h"
 #include "Net/UnrealNetwork.h"
 #include "Net/Core/NetHandle/NetHandleManager.h"
+#include "Tests/ReplicationSystem/ReplicationSystemServerClientTestFixture.h"
 
 
 void UTestPropertyReplicationState_TestClass::GetLifetimeReplicatedProps( TArray< class FLifetimeProperty > & OutLifetimeProps ) const
@@ -44,10 +46,26 @@ void UTestPropertyReplicationState_TestClassWithInitAndCArrays::GetLifetimeRepli
 	DOREPLIFETIME(ThisClass, StructWithArrayOfNotFullyReplicatedStruct);
 }
 
+void UTestPropertyReplicationState_TestClassWithTArray::GetLifetimeReplicatedProps(TArray<class FLifetimeProperty>& OutLifetimeProps) const
+{
+	DOREPLIFETIME(ThisClass, ReferencedObjects);
+	DOREPLIFETIME(ThisClass, ForceReplication);
+}
+
+void UTestPropertyReplicationState_TestClassWithTArray::RegisterReplicationFragments(UE::Net::FFragmentRegistrationContext& Context, UE::Net::EFragmentRegistrationFlags RegistrationFlags)
+{
+	UE::Net::FReplicationFragmentUtil::CreateAndRegisterFragmentsForObject(this, Context, RegistrationFlags);
+}
+
+void UTestPropertyReplicationState_TestClassWithTArray::OnRep_ReferencedObjects()
+{
+	bOnRepWasCalled = true;
+}
+
 namespace UE::Net::Private
 {
 
-class FTestPropertyReplicationStateContext : public FNetworkAutomationTestSuiteFixture
+class FTestPropertyReplicationStateContext : public FReplicationSystemServerClientTestFixture
 {
 protected:
 	FReplicationStateDescriptorBuilder::FResult Descriptors;
@@ -481,4 +499,128 @@ UE_NET_TEST_FIXTURE(FTestPropertyReplicationStateContext, PropertyReplicationSta
 	}
 }
 
+// Test OnRep behavior on TArray when unresolved references are involved.
+UE_NET_TEST_FIXTURE(FTestPropertyReplicationStateContext, TestArrayOnRepWithUnresolvedReferences)
+{
+	IConsoleVariable* CVarbApplyPreviouslyReceivedState = IConsoleManager::Get().FindConsoleVariable(TEXT("net.Iris.DispatchUnresolvedPreviouslyReceivedChanges"), false);
+	UE_NET_ASSERT_NE(CVarbApplyPreviouslyReceivedState, nullptr);
+	UE_NET_ASSERT_TRUE(CVarbApplyPreviouslyReceivedState->IsVariableBool());
+
+	const bool bOldApplyPreviouslyReceivedState = CVarbApplyPreviouslyReceivedState->GetBool();
+	ON_SCOPE_EXIT { CVarbApplyPreviouslyReceivedState->Set(bOldApplyPreviouslyReceivedState, ECVF_SetByCode); };
+
+	for (const bool bApplyPreviouslyReceivedState : {false, true})
+	{
+		CVarbApplyPreviouslyReceivedState->Set(bApplyPreviouslyReceivedState, ECVF_SetByCode);
+
+		// Add a client
+		FReplicationSystemTestClient* Client = CreateClient();
+
+		// Spawn objects on server
+		UTestPropertyReplicationState_TestClassWithTArray* ServerObject = Server->CreateObject<UTestPropertyReplicationState_TestClassWithTArray>();
+
+		// Spawn objects to reference
+		UReplicatedTestObject* ServerReferencedObjectA = Server->CreateObject(UTestReplicatedIrisObject::FComponents{});
+		UReplicatedTestObject* ServerReferencedObjectB = Server->CreateObject(UTestReplicatedIrisObject::FComponents{});
+		UReplicatedTestObject* ServerReferencedObjectC = Server->CreateObject(UTestReplicatedIrisObject::FComponents{});
+		UReplicatedTestObject* ServerReferencedObjectD = Server->CreateObject(UTestReplicatedIrisObject::FComponents{});
+
+		// Prevent objects C and D from being replicated.
+		FNetObjectGroupHandle NotReplicatedGroupHandle = Server->GetReplicationSystem()->GetNotReplicatedNetObjectGroup();
+		Server->ReplicationSystem->AddToGroup(NotReplicatedGroupHandle, ServerReferencedObjectC->NetRefHandle);
+		Server->ReplicationSystem->AddToGroup(NotReplicatedGroupHandle, ServerReferencedObjectD->NetRefHandle);
+
+		UTestPropertyReplicationState_TestClassWithTArray* ClientObject = nullptr;
+
+		// Step 1. Create a fully resolvable payload and make sure it got replicated properly.
+		{
+			ServerObject->ReferencedObjects.Add(ServerReferencedObjectA);
+
+			// Send and make sure everything is as expected on the client
+			Server->UpdateAndSend({ Client });
+
+			ClientObject = Cast<UTestPropertyReplicationState_TestClassWithTArray>(Client->GetReplicationBridge()->GetReplicatedObject(ServerObject->NetRefHandle));
+			UE_NET_ASSERT_NE(ClientObject, nullptr);
+			UE_NET_ASSERT_FALSE(ClientObject->ReferencedObjects.IsEmpty());
+			UE_NET_ASSERT_NE(ClientObject->ReferencedObjects[0], nullptr);
+		}
+
+		// Step 2. Modify the referenced objects array such that it contains both resolvable and unresolvable references.
+		{
+			// Add not yet replicated objects to the mix.
+			ServerObject->ReferencedObjects.Add(ServerReferencedObjectC);
+			ServerObject->ReferencedObjects.Add(ServerReferencedObjectD);
+
+			// Send and make sure everything is as expected on the client
+			Server->UpdateAndSend({ Client });
+
+			ClientObject = Cast<UTestPropertyReplicationState_TestClassWithTArray>(Client->GetReplicationBridge()->GetReplicatedObject(ServerObject->NetRefHandle));
+			UE_NET_ASSERT_EQ(ClientObject->ReferencedObjects.Num(), 3);
+			UE_NET_ASSERT_EQ(ClientObject->ReferencedObjects[0], Client->GetReplicationBridge()->GetReplicatedObject(ServerReferencedObjectA->NetRefHandle));
+			UE_NET_ASSERT_EQ(ClientObject->ReferencedObjects[1], nullptr);
+			UE_NET_ASSERT_EQ(ClientObject->ReferencedObjects[2], nullptr);
+		}
+
+		// Step 3. Modify non-array property on server and null out reference on client. OnRep call depending on whether we're applying previously received state with unresolved references or not.
+		{
+			++ServerObject->ForceReplication;
+			ClientObject->bOnRepWasCalled = false;
+			ClientObject->ReferencedObjects[0] = nullptr;
+
+			// Send and make sure everything is as expected on the client
+			Server->UpdateAndSend({ Client });
+
+			UE_NET_ASSERT_EQ(ClientObject->bOnRepWasCalled, bApplyPreviouslyReceivedState);
+		}
+
+		// Step 4. Allow object C to be replicated. Expecting OnRep.
+		{
+			Server->ReplicationSystem->RemoveFromGroup(NotReplicatedGroupHandle, ServerReferencedObjectC->NetRefHandle);
+
+			ClientObject->bOnRepWasCalled = false;
+
+			Server->UpdateAndSend({ Client });
+
+			UE_NET_ASSERT_TRUE(ClientObject->bOnRepWasCalled);
+			UE_NET_ASSERT_EQ(ClientObject->ReferencedObjects[1], Client->GetReplicationBridge()->GetReplicatedObject(ServerReferencedObjectC->NetRefHandle));
+		}
+
+		// Step 5. Allow object D to be replicated. Expecting OnRep.
+		{
+			Server->ReplicationSystem->RemoveFromGroup(NotReplicatedGroupHandle, ServerReferencedObjectD->NetRefHandle);
+
+			ClientObject->bOnRepWasCalled = false;
+
+			Server->UpdateAndSend({ Client });
+
+			UE_NET_ASSERT_TRUE(ClientObject->bOnRepWasCalled);
+			UE_NET_ASSERT_EQ(ClientObject->ReferencedObjects[2], Client->GetReplicationBridge()->GetReplicatedObject(ServerReferencedObjectD->NetRefHandle));
+		}
+
+		// Step 6. Resize array on server. Expecting OnRep.
+		{
+			ServerObject->ReferencedObjects.SetNum(1);
+
+			ClientObject->bOnRepWasCalled = false;
+
+			Server->UpdateAndSend({ Client });
+
+			UE_NET_ASSERT_TRUE(ClientObject->bOnRepWasCalled);
+		}
+
+		// Step 7. Finally modify non-array property on server and null out reference on client and make sure we don't get an OnRep call.
+		{
+			++ServerObject->ForceReplication;
+
+			ClientObject->bOnRepWasCalled = false;
+			ClientObject->ReferencedObjects[0] = nullptr;
+
+			// Send and make sure everything is as expected on the client
+			Server->UpdateAndSend({ Client });
+
+			UE_NET_ASSERT_FALSE(ClientObject->bOnRepWasCalled);
+		}
+	}
+}
+
 }
diff --git a/Engine/Plugins/Runtime/ReplicationSystemTestPlugin/Source/Private/Tests/ReplicationState/TestPropertyReplicationState.h b/Engine/Plugins/Runtime/ReplicationSystemTestPlugin/Source/Private/Tests/ReplicationState/TestPropertyReplicationState.h
--- a/Engine/Plugins/Runtime/ReplicationSystemTestPlugin/Source/Private/Tests/ReplicationState/TestPropertyReplicationState.h
+++ b/Engine/Plugins/Runtime/ReplicationSystemTestPlugin/Source/Private/Tests/ReplicationState/TestPropertyReplicationState.h
@@ -4,6 +4,7 @@
 
 #include "NetworkAutomationTest.h"
 #include "UObject/ObjectMacros.h"
+#include "Tests/ReplicationSystem/ReplicatedTestObject.h"
 #include "TestPropertyReplicationState.generated.h"
 
 UCLASS()
@@ -119,3 +120,27 @@ public:
 	UPROPERTY(Replicated)
 	FTestPropertyReplicationState_StructWithArrayOfNotFullyReplicatedStruct StructWithArrayOfNotFullyReplicatedStruct;
 };
+
+UCLASS()
+class UTestPropertyReplicationState_TestClassWithTArray : public UReplicatedTestObject
+{
+	GENERATED_BODY()
+
+public:
+	UTestPropertyReplicationState_TestClassWithTArray() : UReplicatedTestObject() {}
+
+	UPROPERTY(ReplicatedUsing=OnRep_ReferencedObjects)
+	TArray<UObject*> ReferencedObjects;
+
+	UPROPERTY(Replicated)
+	uint32 ForceReplication = 0;
+
+	bool bOnRepWasCalled = false;
+
+protected:
+	virtual void RegisterReplicationFragments(UE::Net::FFragmentRegistrationContext& Context, UE::Net::EFragmentRegistrationFlags RegistrationFlags) override;
+
+	UFUNCTION()
+	void OnRep_ReferencedObjects();
+
+};
diff --git a/Engine/Source/Runtime/Experimental/Iris/Core/Private/Iris/ReplicationSystem/ReplicationReader.cpp b/Engine/Source/Runtime/Experimental/Iris/Core/Private/Iris/ReplicationSystem/ReplicationReader.cpp
--- a/Engine/Source/Runtime/Experimental/Iris/Core/Private/Iris/ReplicationSystem/ReplicationReader.cpp
+++ b/Engine/Source/Runtime/Experimental/Iris/Core/Private/Iris/ReplicationSystem/ReplicationReader.cpp
@@ -55,6 +55,12 @@ static FAutoConsoleVariableRef CVarDeferEndReplication(
 	TEXT("bDeferEndReplication if true calls to EndReplication will be defered until after we have applied statedata. Default is true."
 	));
 
+static bool bDispatchUnresolvedPreviouslyReceivedChanges = false;
+static FAutoConsoleVariableRef CvarDispatchUnresolvedPreviouslyReceivedChanges(
+	TEXT("net.Iris.DispatchUnresolvedPreviouslyReceivedChanges"),
+	bDispatchUnresolvedPreviouslyReceivedChanges,
+	TEXT("Whether to include previously received changes with unresolved object references to data received this frame when applying state data. This can call rep notify functions to be called despite being unchanged. Default is false."));
+
 static const FName NetError_FailedToFindAttachmentQueue("Failed to find attachment queue");
 
 class FResolveAndCollectUnresolvedAndResolvedReferenceCollector
@@ -1472,20 +1478,64 @@ void FReplicationReader::DispatchStateData(FNetSerializationContext& Context)
 			// If we have pending unresolved changes we include them as well 
 			FNetBitArrayView UnresolvedChangeMask = FChangeMaskUtil::MakeChangeMask(ReplicationInfo->UnresolvedChangeMaskOrPointer, ChangeMaskBitCount);
 
-			if (ReplicationInfo->bHasUnresolvedReferences)
+			FNetBitArrayView ChangeMaskForResolve;
+			const bool bHadUnresolvedReferences = ReplicationInfo->bHasUnresolvedReferences;
+			if (bHadUnresolvedReferences)
+			{
+				if (bDispatchUnresolvedPreviouslyReceivedChanges)
+				{
+					// Combine the changemask with the unresolved changemask so that result is used for the apply operation as well.
+					ChangeMask.Combine(UnresolvedChangeMask, FNetBitArrayView::OrOp);
+					ChangeMaskForResolve = ChangeMask;
+				}
+				else
+				{
+					// Memory for the changemask allocation will be freed when the TempLinearAllocator is reset via FMemMark scope. TempChangeMaskAllocator uses TempLinearAllocator.
+					FChangeMaskStorageOrPointer ChangeMaskForResolveAllocation;
+					ChangeMaskForResolveAllocation.Alloc(ChangeMaskForResolveAllocation, ChangeMaskBitCount, TempChangeMaskAllocator);
+					ChangeMaskForResolve = MakeNetBitArrayView(ChangeMaskForResolveAllocation.GetPointer(ChangeMaskBitCount), ChangeMaskBitCount, FNetBitArrayView::NoResetNoValidate);
+					ChangeMaskForResolve.Set(ChangeMask, FNetBitArrayView::OrOp, UnresolvedChangeMask);
+				}
+			}
+			else
 			{
-				ChangeMask.Combine(UnresolvedChangeMask, FNetBitArrayView::OrOp);
+				ChangeMaskForResolve = ChangeMask;
 			}
 
 			// Collect all unresolvable references, including old pending references
 			FResolveAndCollectUnresolvedAndResolvedReferenceCollector Collector;
-			Collector.CollectReferences(*ObjectReferenceCache, ResolveContext, Info.bIsInitialState | ReplicationInfo->bHasUnresolvedInitialReferences, &ChangeMask, ObjectData.ReceiveStateBuffer, ObjectData.Protocol);
+			Collector.CollectReferences(*ObjectReferenceCache, ResolveContext, Info.bIsInitialState | ReplicationInfo->bHasUnresolvedInitialReferences, &ChangeMaskForResolve, ObjectData.ReceiveStateBuffer, ObjectData.Protocol);
 
-			// If we have any object references we need to track them
-			if (Collector.GetUnresolvedReferences().Num() > 0 || Collector.GetResolvedReferences().Num() > 0)
+			// If we have or had any object references we need to track them and update the unresolved mask
+			if (bHadUnresolvedReferences || Collector.GetUnresolvedReferences().Num() > 0 || Collector.GetResolvedReferences().Num() > 0)
 			{
-				BuildUnresolvedChangeMaskAndUpdateObjectReferenceTracking(Collector, ChangeMask, ReplicationInfo, UnresolvedChangeMask);
+				FNetBitArrayView PrevUnresolvedChangeMask;
+
+				// If we're avoiding dispatching state we didn't receive and we didn't resolve anything for we need to figure out what got resolves and combine that with the received changemask.
+				const bool bMergeResolvedReferencesWithChangeMask = bHadUnresolvedReferences && !bDispatchUnresolvedPreviouslyReceivedChanges;
+				if (bMergeResolvedReferencesWithChangeMask)
+				{
+					FChangeMaskStorageOrPointer ChangeMaskForPrevUnresolvedAllocation;
+					ChangeMaskForPrevUnresolvedAllocation.Alloc(ChangeMaskForPrevUnresolvedAllocation, ChangeMaskBitCount, TempChangeMaskAllocator);
+					PrevUnresolvedChangeMask = MakeNetBitArrayView(ChangeMaskForPrevUnresolvedAllocation.GetPointer(ChangeMaskBitCount), ChangeMaskBitCount, FNetBitArrayView::NoResetNoValidate);
+					PrevUnresolvedChangeMask.Copy(UnresolvedChangeMask);
+				}
+
+				BuildUnresolvedChangeMaskAndUpdateObjectReferenceTracking(Collector, ChangeMaskForResolve, ReplicationInfo, UnresolvedChangeMask);
 			
+				// Allow resolved changes to be part of the state to be applied.
+				if (bMergeResolvedReferencesWithChangeMask)
+				{
+					// Merge in no longer unresolved changes
+					ChangeMask.CombineMultiple(FNetBitArrayView::OrOp, PrevUnresolvedChangeMask, FNetBitArrayView::AndNotOp, UnresolvedChangeMask);
+
+					// Merge in partially resolved changes
+					for (const FNetReferenceCollector::FReferenceInfo& ReferenceInfo : Collector.GetResolvedReferences())
+					{
+						ChangeMask.SetBit(ReferenceInfo.ChangeMaskInfo.BitOffset);
+					}
+				}
+
 				// $IRIS: $TODO: For now we always apply, even if we cannot resolve all references for a property
 				// We need to investigate how this is handled best as we do not want to prevent arrays(fastarrays) from applying data just because a single element wont resolve?
 				// Mask off any unresolvable states before we dispatch state data
