Disable unsafe type case warnings in ReplicationSystemTestPlugin.

diff --git a/Engine/Plugins/Runtime/ReplicationSystemTestPlugin/Source/ReplicationSystemTestPlugin.Build.cs b/Engine/Plugins/Runtime/ReplicationSystemTestPlugin/Source/ReplicationSystemTestPlugin.Build.cs
--- a/Engine/Plugins/Runtime/ReplicationSystemTestPlugin/Source/ReplicationSystemTestPlugin.Build.cs
+++ b/Engine/Plugins/Runtime/ReplicationSystemTestPlugin/Source/ReplicationSystemTestPlugin.Build.cs
@@ -55,8 +55,6 @@ namespace UnrealBuildTool.Rules
 				PrivateIncludePathModuleNames.Add("Catch2");
 			}
 
-			UnsafeTypeCastWarningLevel = WarningLevel.Error;
-
 			PrivateDefinitions.Add(String.Format("UE_NET_WITH_LOW_LEVEL_TESTS={0}", Target.ExplicitTestsTarget ? "1" : "0"));
 
 			SetupIrisSupport(Target);
