Enable truncation warnings in ReplicationSystemTestPlugin.

diff --git a/Engine/Plugins/Runtime/ReplicationSystemTestPlugin/Source/Private/Tests/ReplicationState/TestDerivedStruct.cpp b/Engine/Plugins/Runtime/ReplicationSystemTestPlugin/Source/Private/Tests/ReplicationState/TestDerivedStruct.cpp
--- a/Engine/Plugins/Runtime/ReplicationSystemTestPlugin/Source/Private/Tests/ReplicationState/TestDerivedStruct.cpp
+++ b/Engine/Plugins/Runtime/ReplicationSystemTestPlugin/Source/Private/Tests/ReplicationState/TestDerivedStruct.cpp
@@ -236,7 +236,7 @@ UE_NET_TEST_FIXTURE(FTestDerivedStructInObjectFixture, CanReplicateDerivedStruct
 	// Check that the members have been replicated
 	{
 		UE_NET_ASSERT_EQ(ClientObject->DerivedStructArray.Num(), ServerObject->DerivedStructArray.Num());
-		SIZE_T MemberIt = 0;
+		int32 MemberIt = 0;
 		for (const FTestDerivedStruct_Inherited_WithNetSerializer_Inherited_WithoutNetSerializer& ServerItem : ServerObject->DerivedStructArray)
 		{
 			const FTestDerivedStruct_Inherited_WithNetSerializer_Inherited_WithoutNetSerializer& ClientItem = ClientObject->DerivedStructArray[MemberIt];
@@ -380,14 +380,13 @@ UE_NET_TEST_FIXTURE(FTestDerivedStructInObjectFixture, NotAppliedMemberInDerived
 	const uint8 ClientModifiedNotAppliedPropertyBaseValue = ClientObject->DerivedStructArray[0].ByteMemberNotSetOnApply ^ 123U;
 	for (auto& Element : ClientObject->DerivedStructArray)
 	{
-		const uint8 Index = &Element - ClientObject->DerivedStructArray.GetData();
+		const uint8 Index = IntCastChecked<uint8>(&Element - ClientObject->DerivedStructArray.GetData());
 		Element.ByteMemberNotSetOnApply = ClientModifiedNotAppliedPropertyBaseValue + Index;
 	}
 
 	// On the server modify all elements in the array
 	for (auto& Element : ServerObject->DerivedStructArray)
 	{
-		const uint8 Index = &Element - ServerObject->DerivedStructArray.GetData();
 		Element.ByteMember0 ^= 1U;
 		Element.ByteMember1 ^= 2U;
 		Element.ByteMemberNotSetOnApply ^= 231U;
@@ -400,7 +399,7 @@ UE_NET_TEST_FIXTURE(FTestDerivedStructInObjectFixture, NotAppliedMemberInDerived
 	// Verify that the not applied member remains as is.
 	for (const auto& Element : ClientObject->DerivedStructArray)
 	{
-		const uint8 Index = &Element - ClientObject->DerivedStructArray.GetData();
+		const uint8 Index = IntCastChecked<uint8>(&Element - ClientObject->DerivedStructArray.GetData());
 		UE_NET_ASSERT_EQ(Element.ByteMemberNotSetOnApply, uint8(ClientModifiedNotAppliedPropertyBaseValue + Index));
 	}
 }
@@ -468,7 +467,7 @@ UE_NET_TEST_FIXTURE(FTestDerivedStructInObjectFixture, NotAppliedMemberInStructD
 	const uint8 ClientModifiedNotAppliedPropertyBaseValue = ClientObject->DerivedStructArray[0].ByteMemberNotSetOnApply ^ 123U;
 	for (auto& Element : ClientObject->DerivedStructArray)
 	{
-		const uint8 Index = &Element - ClientObject->DerivedStructArray.GetData();
+		const uint8 Index = IntCastChecked<uint8>(&Element - ClientObject->DerivedStructArray.GetData());
 		Element.ByteMemberNotSetOnApply = ClientModifiedNotAppliedPropertyBaseValue + Index;
 		Element.ByteMember3_NotReplicated = ClientModifiedNotReplicatedPropertyValue;
 	}
@@ -476,7 +475,6 @@ UE_NET_TEST_FIXTURE(FTestDerivedStructInObjectFixture, NotAppliedMemberInStructD
 	// On the server modify all elements in the array
 	for (auto& Element : ServerObject->DerivedStructArray)
 	{
-		const uint8 Index = &Element - ServerObject->DerivedStructArray.GetData();
 		Element.ByteMember0 ^= 1U;
 		Element.ByteMemberNotSetOnApply ^= 231U;
 		Element.ByteMember2 ^= 2U;
@@ -490,7 +488,7 @@ UE_NET_TEST_FIXTURE(FTestDerivedStructInObjectFixture, NotAppliedMemberInStructD
 	// Verify that the not applied member remains as is.
 	for (const auto& Element : ClientObject->DerivedStructArray)
 	{
-		const uint8 Index = &Element - ClientObject->DerivedStructArray.GetData();
+		const uint8 Index = IntCastChecked<uint8>(&Element - ClientObject->DerivedStructArray.GetData());
 		UE_NET_ASSERT_EQ(Element.ByteMember0, ServerObject->DerivedStructArray[Index].ByteMember0);
 		UE_NET_ASSERT_EQ(Element.ByteMemberNotSetOnApply, uint8(ClientModifiedNotAppliedPropertyBaseValue + Index));
 		UE_NET_ASSERT_EQ(Element.ByteMember2, ServerObject->DerivedStructArray[Index].ByteMember2);
@@ -513,8 +511,8 @@ void FTestDerivedStruct_Inherited_WithNetSerializer_NetSerializer::Deserialize(F
 	QuantizedType& Target = *reinterpret_cast<QuantizedType*>(Args.Target);
 
 	FNetBitStreamReader* Reader = Context.GetBitStreamReader();
-	Target.ByteMember0 = Reader->ReadBits(8U);
-	Target.ByteMember1 = Reader->ReadBits(8U);
+	Target.ByteMember0 = IntCastChecked<uint8>(Reader->ReadBits(8U));
+	Target.ByteMember1 = IntCastChecked<uint8>(Reader->ReadBits(8U));
 }
 
 void FTestDerivedStruct_Inherited_WithNetSerializer_NetSerializer::Quantize(FNetSerializationContext& Context, const FNetQuantizeArgs& Args)
@@ -579,9 +577,9 @@ void FTestDerivedStruct_Inherited_WithNetSerializerWithApply_NetSerializer::Dese
 	QuantizedType& Target = *reinterpret_cast<QuantizedType*>(Args.Target);
 
 	FNetBitStreamReader* Reader = Context.GetBitStreamReader();
-	Target.ByteMember0 = Reader->ReadBits(8U);
-	Target.ByteMember1 = Reader->ReadBits(8U);
-	Target.ByteMemberNotSetOnApply = Reader->ReadBits(8U);
+	Target.ByteMember0 = IntCastChecked<uint8>(Reader->ReadBits(8U));
+	Target.ByteMember1 = IntCastChecked<uint8>(Reader->ReadBits(8U));
+	Target.ByteMemberNotSetOnApply = IntCastChecked<uint8>(Reader->ReadBits(8U));
 }
 
 void FTestDerivedStruct_Inherited_WithNetSerializerWithApply_NetSerializer::Quantize(FNetSerializationContext& Context, const FNetQuantizeArgs& Args)
diff --git a/Engine/Plugins/Runtime/ReplicationSystemTestPlugin/Source/Private/Tests/ReplicationState/TestPropertyReplicationState.cpp b/Engine/Plugins/Runtime/ReplicationSystemTestPlugin/Source/Private/Tests/ReplicationState/TestPropertyReplicationState.cpp
--- a/Engine/Plugins/Runtime/ReplicationSystemTestPlugin/Source/Private/Tests/ReplicationState/TestPropertyReplicationState.cpp
+++ b/Engine/Plugins/Runtime/ReplicationSystemTestPlugin/Source/Private/Tests/ReplicationState/TestPropertyReplicationState.cpp
@@ -382,7 +382,7 @@ UE_NET_TEST_FIXTURE(FTestPropertyReplicationStateContext, PropertyReplicationSta
 	ReplicationState.SetPropertyValue(NotFullyReplicatedArrayElementIndex, &NotFullyReplicatedArrayElementValue);
 
 	ReplicationStateB.Set(ReplicationState);
-	for (SIZE_T MemberIt = 0, MemberEndIt = Descriptor->MemberCount; MemberIt != MemberEndIt; ++MemberIt)
+	for (uint32 MemberIt = 0, MemberEndIt = Descriptor->MemberCount; MemberIt != MemberEndIt; ++MemberIt)
 	{
 		// Init state doesn't have changemasks
 		UE_NET_ASSERT_TRUE_MSG(ReplicationStateB.IsDirty(MemberIt), "Member '" << Descriptor->MemberProperties[MemberIt]->GetName() << "' index " << MemberIt);
@@ -413,7 +413,7 @@ UE_NET_TEST_FIXTURE(FTestPropertyReplicationStateContext, PropertyReplicationSta
 	ReplicationState.SetPropertyValue(NotFullyReplicatedArrayElementIndex, &NotFullyReplicatedArrayElementValue);
 
 	ReplicationStateB.Set(ReplicationState);
-	for (SIZE_T MemberIt = 0, MemberEndIt = Descriptor->MemberCount; MemberIt != MemberEndIt; ++MemberIt)
+	for (uint32 MemberIt = 0, MemberEndIt = Descriptor->MemberCount; MemberIt != MemberEndIt; ++MemberIt)
 	{
 		if (MemberIt == FullyReplicatedArrayElementIndex || (MemberIt == NotFullyReplicatedArrayElementIndex))
 		{
@@ -450,7 +450,7 @@ UE_NET_TEST_FIXTURE(FTestPropertyReplicationStateContext, PropertyReplicationSta
 	ReplicationState.SetPropertyValue(NotFullyReplicatedArrayElementIndex, &NotFullyReplicatedArrayElementValue);
 
 	ReplicationStateB.Set(ReplicationState);
-	for (SIZE_T MemberIt = 0, MemberEndIt = Descriptor->MemberCount; MemberIt != MemberEndIt; ++MemberIt)
+	for (uint32 MemberIt = 0, MemberEndIt = Descriptor->MemberCount; MemberIt != MemberEndIt; ++MemberIt)
 	{
 		UE_NET_ASSERT_FALSE_MSG(ReplicationStateB.IsDirty(MemberIt), "Member '" << Descriptor->MemberProperties[MemberIt]->GetName() << "' index " << MemberIt);
 	}
@@ -486,7 +486,7 @@ UE_NET_TEST_FIXTURE(FTestPropertyReplicationStateContext, PropertyReplicationSta
 	ReplicationState.SetPropertyValue(NotFullyReplicatedArrayElementIndex, &NotFullyReplicatedArrayElementValue1);
 
 	ReplicationStateB.Set(ReplicationState);
-	for (SIZE_T MemberIt = 0, MemberEndIt = Descriptor->MemberCount; MemberIt != MemberEndIt; ++MemberIt)
+	for (uint32 MemberIt = 0, MemberEndIt = Descriptor->MemberCount; MemberIt != MemberEndIt; ++MemberIt)
 	{
 		if (MemberIt == FullyReplicatedArrayElementIndex || (MemberIt == NotFullyReplicatedArrayElementIndex))
 		{
diff --git a/Engine/Plugins/Runtime/ReplicationSystemTestPlugin/Source/Private/Tests/ReplicationSystem/Conditionals/TestConditionals.cpp b/Engine/Plugins/Runtime/ReplicationSystemTestPlugin/Source/Private/Tests/ReplicationSystem/Conditionals/TestConditionals.cpp
--- a/Engine/Plugins/Runtime/ReplicationSystemTestPlugin/Source/Private/Tests/ReplicationSystem/Conditionals/TestConditionals.cpp
+++ b/Engine/Plugins/Runtime/ReplicationSystemTestPlugin/Source/Private/Tests/ReplicationSystem/Conditionals/TestConditionals.cpp
@@ -491,7 +491,7 @@ UE_NET_TEST_FIXTURE(FTestConditionalsFixture, CanMixConditions)
 	const UTestReplicatedIrisObject* ClientObjectArray[sizeof(ClientArray)/sizeof(ClientArray[0])];
 	for (FReplicationSystemTestClient*& Client : ClientArray)
 	{
-		const uint32 ClientIndex = &Client - ClientArray;
+		const SIZE_T ClientIndex = &Client - ClientArray;
 		ClientObjectArray[ClientIndex] = Cast<UTestReplicatedIrisObject>(Client->GetReplicationBridge()->GetReplicatedObject(ServerObject->NetRefHandle));
 		UE_NET_ASSERT_NE(ClientObjectArray[ClientIndex], nullptr);
 	}
@@ -499,7 +499,7 @@ UE_NET_TEST_FIXTURE(FTestConditionalsFixture, CanMixConditions)
 	// Validate all the members
 	for (FReplicationSystemTestClient*& Client : ClientArray)
 	{
-		const uint32 ClientIndex = &Client - ClientArray;
+		const SIZE_T ClientIndex = &Client - ClientArray;
 		const UTestReplicatedIrisObject* ClientObject = ClientObjectArray[ClientIndex];
 
 		// Everybody should get the physics members
@@ -555,7 +555,7 @@ UE_NET_TEST_FIXTURE(FTestConditionalsFixture, CanSwitchAutonmousConnection)
 	const UTestReplicatedIrisObject* ClientObjectArray[sizeof(ClientArray)/sizeof(ClientArray[0])];
 	for (FReplicationSystemTestClient*& Client : ClientArray)
 	{
-		const uint32 ClientIndex = &Client - ClientArray;
+		const SIZE_T ClientIndex = &Client - ClientArray;
 		ClientObjectArray[ClientIndex] = Cast<UTestReplicatedIrisObject>(Client->GetReplicationBridge()->GetReplicatedObject(ServerObject->NetRefHandle));
 		UE_NET_ASSERT_NE(ClientObjectArray[ClientIndex], nullptr);
 	}
@@ -583,7 +583,7 @@ UE_NET_TEST_FIXTURE(FTestConditionalsFixture, CanSwitchAutonmousConnection)
 	// Validate all the members
 	for (FReplicationSystemTestClient*& Client : ClientArray)
 	{
-		const uint32 ClientIndex = &Client - ClientArray;
+		const SIZE_T ClientIndex = &Client - ClientArray;
 		const UTestReplicatedIrisObject* ClientObject = ClientObjectArray[ClientIndex];
 
 		if (Client->ConnectionIdOnServer == AutonomousClientConnectionId)
@@ -655,7 +655,7 @@ UE_NET_TEST_FIXTURE(FTestConditionalsFixture, SubObjectConditionsMatchesParentOb
 	const UTestReplicatedIrisObject* ClientSubObjectArray[sizeof(ClientArray)/sizeof(ClientArray[0])];
 	for (FReplicationSystemTestClient*& Client : ClientArray)
 	{
-		const uint32 ClientIndex = &Client - ClientArray;
+		const SIZE_T ClientIndex = &Client - ClientArray;
 		ClientObjectArray[ClientIndex] = Cast<UTestReplicatedIrisObject>(Client->GetReplicationBridge()->GetReplicatedObject(ServerObject->NetRefHandle));
 		UE_NET_ASSERT_NE(ClientObjectArray[ClientIndex], nullptr);
 		ClientSubObjectArray[ClientIndex] = Cast<UTestReplicatedIrisObject>(Client->GetReplicationBridge()->GetReplicatedObject(ServerSubObject->NetRefHandle));
@@ -665,7 +665,7 @@ UE_NET_TEST_FIXTURE(FTestConditionalsFixture, SubObjectConditionsMatchesParentOb
 	// Validate all the members
 	for (FReplicationSystemTestClient*& Client : ClientArray)
 	{
-		const uint32 ClientIndex = &Client - ClientArray;
+		const SIZE_T ClientIndex = &Client - ClientArray;
 
 		for (const UTestReplicatedIrisObject* ClientObject : {ClientObjectArray[ClientIndex], ClientSubObjectArray[ClientIndex]})
 		{
diff --git a/Engine/Plugins/Runtime/ReplicationSystemTestPlugin/Source/Private/Tests/ReplicationSystem/Filtering/TestFiltering.cpp b/Engine/Plugins/Runtime/ReplicationSystemTestPlugin/Source/Private/Tests/ReplicationSystem/Filtering/TestFiltering.cpp
--- a/Engine/Plugins/Runtime/ReplicationSystemTestPlugin/Source/Private/Tests/ReplicationSystem/Filtering/TestFiltering.cpp
+++ b/Engine/Plugins/Runtime/ReplicationSystemTestPlugin/Source/Private/Tests/ReplicationSystem/Filtering/TestFiltering.cpp
@@ -2177,10 +2177,10 @@ UE_NET_TEST_FIXTURE(FTestFilteringFixture, InclusionGroupsWorksWithMultipleObjec
 	// Exactly one object should have been replicated to each
 	for (UReplicatedTestObject*& ServerObject : ServerObjects)
 	{
-		const SIZE_T Index = &ServerObject - &ServerObjects[0];
-		UE_NET_ASSERT_NE(Clients[Index]->GetReplicationBridge()->GetReplicatedObject(ServerObject->NetRefHandle), nullptr);
-		UE_NET_ASSERT_EQ(Clients[(Index + 1U) % UE_ARRAY_COUNT(ServerObjects)]->GetReplicationBridge()->GetReplicatedObject(ServerObject->NetRefHandle), nullptr);
-		UE_NET_ASSERT_EQ(Clients[(Index + 2U) % UE_ARRAY_COUNT(ServerObjects)]->GetReplicationBridge()->GetReplicatedObject(ServerObject->NetRefHandle), nullptr);
+		const SIZE_T Index = IntCastChecked<uint32>(&ServerObject - &ServerObjects[0]);
+		UE_NET_ASSERT_NE(ClientArray[Index]->GetReplicationBridge()->GetReplicatedObject(ServerObject->NetRefHandle), nullptr);
+		UE_NET_ASSERT_EQ(ClientArray[(Index + 1U) % UE_ARRAY_COUNT(ServerObjects)]->GetReplicationBridge()->GetReplicatedObject(ServerObject->NetRefHandle), nullptr);
+		UE_NET_ASSERT_EQ(ClientArray[(Index + 2U) % UE_ARRAY_COUNT(ServerObjects)]->GetReplicationBridge()->GetReplicatedObject(ServerObject->NetRefHandle), nullptr);
 	}
 }
 
@@ -2225,7 +2225,7 @@ UE_NET_TEST_FIXTURE(FTestFilteringFixture, InclusionGroupsAreCumulative)
 	// Exactly one object should have been replicated to each connection
 	for (UReplicatedTestObject*& ServerObject : ServerObjects)
 	{
-		const SIZE_T Index = &ServerObject - &ServerObjects[0];
+		const int32 Index = static_cast<int32>(&ServerObject - &ServerObjects[0]);
 		UE_NET_ASSERT_NE(Clients[Index]->GetReplicationBridge()->GetReplicatedObject(ServerObject->NetRefHandle), nullptr);
 		UE_NET_ASSERT_EQ(Clients[(Index + 1U) % UE_ARRAY_COUNT(ServerObjects)]->GetReplicationBridge()->GetReplicatedObject(ServerObject->NetRefHandle), nullptr);
 		UE_NET_ASSERT_EQ(Clients[(Index + 2U) % UE_ARRAY_COUNT(ServerObjects)]->GetReplicationBridge()->GetReplicatedObject(ServerObject->NetRefHandle), nullptr);
diff --git a/Engine/Plugins/Runtime/ReplicationSystemTestPlugin/Source/Private/Tests/ReplicationSystem/NetBlob/TestReliableNetBlobQueue.cpp b/Engine/Plugins/Runtime/ReplicationSystemTestPlugin/Source/Private/Tests/ReplicationSystem/NetBlob/TestReliableNetBlobQueue.cpp
--- a/Engine/Plugins/Runtime/ReplicationSystemTestPlugin/Source/Private/Tests/ReplicationSystem/NetBlob/TestReliableNetBlobQueue.cpp
+++ b/Engine/Plugins/Runtime/ReplicationSystemTestPlugin/Source/Private/Tests/ReplicationSystem/NetBlob/TestReliableNetBlobQueue.cpp
@@ -87,7 +87,7 @@ UE_NET_TEST_FIXTURE(FTestReliableNetBlobQueue, CanSendMoreReliableBlobsThanRelia
 			Server->SendTo(Client);
 			Server->PostSendUpdate();
 
-			const uint32 NewPacketCount = ConnectionInfo.WrittenPackets.Count();
+			const SIZE_T NewPacketCount = ConnectionInfo.WrittenPackets.Count();
 			if (NewPacketCount == PacketCount)
 			{
 				break;
diff --git a/Engine/Plugins/Runtime/ReplicationSystemTestPlugin/Source/Private/Tests/ReplicationSystem/Prioritization/TestSphereNetObjectPrioritizer.cpp b/Engine/Plugins/Runtime/ReplicationSystemTestPlugin/Source/Private/Tests/ReplicationSystem/Prioritization/TestSphereNetObjectPrioritizer.cpp
--- a/Engine/Plugins/Runtime/ReplicationSystemTestPlugin/Source/Private/Tests/ReplicationSystem/Prioritization/TestSphereNetObjectPrioritizer.cpp
+++ b/Engine/Plugins/Runtime/ReplicationSystemTestPlugin/Source/Private/Tests/ReplicationSystem/Prioritization/TestSphereNetObjectPrioritizer.cpp
@@ -53,7 +53,7 @@ void USphereNetObjectPrioritizerForTest::PrioritizeWithParams(USphereNetObjectPr
 			const FVector* FirstPos = Params.InPositions.GetData();
 			for (const FVector& Pos : MakeArrayView(Params.InPositions.GetData() + ObjectIt, CurrentBatchObjectCount))
 			{
-				const uint32 Index = &Pos - FirstPos;
+				const int32 Index = static_cast<int32>(&Pos - FirstPos);
 				BatchParams.Positions[Index] = VectorLoadFloat3_W0(&Params.InPositions[Index]);
 			}
 
@@ -176,16 +176,16 @@ UE_NET_TEST_FIXTURE(FTestSphereNetObjectPrioritizer, ViewPositionSameAsObjectGiv
 
 	// Test one view
 	{
-		PrioParams.View = MakeReplicationView(TestPosition, FVector(1, 0, 0), FMath::DegreesToRadians(60));
+		PrioParams.View = MakeReplicationView(TestPosition, FVector(1, 0, 0), FMath::DegreesToRadians(60.0f));
 		SphereNetObjectPrioritizer->PrioritizeWithParams(PrioParams);
 		UE_NET_ASSERT_EQ(OutPriority, SphereNetObjectPrioritizerConfig->InnerPriority);
 	}
 
 	// Test two views
 	{
-		PrioParams.View = MakeReplicationView(TestPosition, FVector(1, 0, 0), FMath::DegreesToRadians(60));
+		PrioParams.View = MakeReplicationView(TestPosition, FVector(1, 0, 0), FMath::DegreesToRadians(60.0f));
 		// Add second view at origin.
-		PrioParams.View.Views.Add(MakeReplicationView(FVector(0, 0, 0), FVector(1, 0, 0), FMath::DegreesToRadians(60)).Views[0]);
+		PrioParams.View.Views.Add(MakeReplicationView(FVector(0, 0, 0), FVector(1, 0, 0), FMath::DegreesToRadians(60.0f)).Views[0]);
 
 		SphereNetObjectPrioritizer->PrioritizeWithParams(PrioParams);
 		UE_NET_ASSERT_EQ(OutPriority, SphereNetObjectPrioritizerConfig->InnerPriority);
@@ -198,10 +198,10 @@ UE_NET_TEST_FIXTURE(FTestSphereNetObjectPrioritizer, ViewPositionSameAsObjectGiv
 
 	// Test more than two views.
 	{
-		PrioParams.View = MakeReplicationView(TestPosition, FVector(1, 0, 0), FMath::DegreesToRadians(60));
+		PrioParams.View = MakeReplicationView(TestPosition, FVector(1, 0, 0), FMath::DegreesToRadians(60.0f));
 		// Add a couple of views.
-		PrioParams.View.Views.Add(MakeReplicationView(FVector(0, 0, 0), FVector(1, 0, 0), FMath::DegreesToRadians(60)).Views[0]);
-		PrioParams.View.Views.Add(MakeReplicationView(-TestPosition, FVector(1, 0, 0), FMath::DegreesToRadians(60)).Views[0]);
+		PrioParams.View.Views.Add(MakeReplicationView(FVector(0, 0, 0), FVector(1, 0, 0), FMath::DegreesToRadians(60.0f)).Views[0]);
+		PrioParams.View.Views.Add(MakeReplicationView(-TestPosition, FVector(1, 0, 0), FMath::DegreesToRadians(60.0f)).Views[0]);
 
 		SphereNetObjectPrioritizer->PrioritizeWithParams(PrioParams);
 		UE_NET_ASSERT_EQ(OutPriority, SphereNetObjectPrioritizerConfig->InnerPriority);
@@ -225,16 +225,16 @@ UE_NET_TEST_FIXTURE(FTestSphereNetObjectPrioritizer, ViewPositionFarAwayFromObje
 
 	// Test one view
 	{
-		PrioParams.View = MakeReplicationView(TestPosition + FVector(0, 0, SphereNetObjectPrioritizerConfig->OuterRadius), FVector(1, 0, 0), FMath::DegreesToRadians(60));
+		PrioParams.View = MakeReplicationView(TestPosition + FVector(0, 0, SphereNetObjectPrioritizerConfig->OuterRadius), FVector(1, 0, 0), FMath::DegreesToRadians(60.0f));
 		SphereNetObjectPrioritizer->PrioritizeWithParams(PrioParams);
 		UE_NET_ASSERT_EQ(OutPriority, SphereNetObjectPrioritizerConfig->OuterPriority);
 	}
 
 	// Test two views
 	{
-		PrioParams.View = MakeReplicationView(TestPosition + FVector(0, 0, SphereNetObjectPrioritizerConfig->OuterRadius), FVector(1, 0, 0), FMath::DegreesToRadians(60));
+		PrioParams.View = MakeReplicationView(TestPosition + FVector(0, 0, SphereNetObjectPrioritizerConfig->OuterRadius), FVector(1, 0, 0), FMath::DegreesToRadians(60.0f));
 		// Add second view even further away
-		PrioParams.View.Views.Add(MakeReplicationView(TestPosition + FVector(0, 0, 2.0f*SphereNetObjectPrioritizerConfig->OuterRadius), FVector(1, 0, 0), FMath::DegreesToRadians(60)).Views[0]);
+		PrioParams.View.Views.Add(MakeReplicationView(TestPosition + FVector(0, 0, 2.0f*SphereNetObjectPrioritizerConfig->OuterRadius), FVector(1, 0, 0), FMath::DegreesToRadians(60.0f)).Views[0]);
 
 		SphereNetObjectPrioritizer->PrioritizeWithParams(PrioParams);
 		UE_NET_ASSERT_EQ(OutPriority, SphereNetObjectPrioritizerConfig->OuterPriority);
@@ -247,10 +247,10 @@ UE_NET_TEST_FIXTURE(FTestSphereNetObjectPrioritizer, ViewPositionFarAwayFromObje
 
 	// Test more than two views.
 	{
-		PrioParams.View = MakeReplicationView(TestPosition + FVector(0, 0, SphereNetObjectPrioritizerConfig->OuterRadius), FVector(1, 0, 0), FMath::DegreesToRadians(60));
+		PrioParams.View = MakeReplicationView(TestPosition + FVector(0, 0, SphereNetObjectPrioritizerConfig->OuterRadius), FVector(1, 0, 0), FMath::DegreesToRadians(60.0f));
 		// Add a couple of views
-		PrioParams.View.Views.Add(MakeReplicationView(TestPosition + FVector(0, 0, 2.0f*SphereNetObjectPrioritizerConfig->OuterRadius), FVector(1, 0, 0), FMath::DegreesToRadians(60)).Views[0]);
-		PrioParams.View.Views.Add(MakeReplicationView(TestPosition - FVector(0, 0, 4.0f*SphereNetObjectPrioritizerConfig->OuterRadius), FVector(1, 0, 0), FMath::DegreesToRadians(60)).Views[0]);
+		PrioParams.View.Views.Add(MakeReplicationView(TestPosition + FVector(0, 0, 2.0f*SphereNetObjectPrioritizerConfig->OuterRadius), FVector(1, 0, 0), FMath::DegreesToRadians(60.0f)).Views[0]);
+		PrioParams.View.Views.Add(MakeReplicationView(TestPosition - FVector(0, 0, 4.0f*SphereNetObjectPrioritizerConfig->OuterRadius), FVector(1, 0, 0), FMath::DegreesToRadians(60.0f)).Views[0]);
 
 		SphereNetObjectPrioritizer->PrioritizeWithParams(PrioParams);
 		UE_NET_ASSERT_EQ(OutPriority, SphereNetObjectPrioritizerConfig->OuterPriority);
@@ -274,16 +274,16 @@ UE_NET_TEST_FIXTURE(FTestSphereNetObjectPrioritizer, ViewPositionVeryFarAwayFrom
 
 	// Test one view
 	{
-		PrioParams.View = MakeReplicationView(TestPosition + SphereNetObjectPrioritizerConfig->OuterRadius, FVector(1, 0, 0), FMath::DegreesToRadians(60));
+		PrioParams.View = MakeReplicationView(TestPosition + SphereNetObjectPrioritizerConfig->OuterRadius, FVector(1, 0, 0), FMath::DegreesToRadians(60.0f));
 		SphereNetObjectPrioritizer->PrioritizeWithParams(PrioParams);
 		UE_NET_ASSERT_EQ(OutPriority, SphereNetObjectPrioritizerConfig->OutsidePriority);
 	}
 
 	// Test two views
 	{
-		PrioParams.View = MakeReplicationView(TestPosition + SphereNetObjectPrioritizerConfig->OuterRadius, FVector(1, 0, 0), FMath::DegreesToRadians(60));
+		PrioParams.View = MakeReplicationView(TestPosition + SphereNetObjectPrioritizerConfig->OuterRadius, FVector(1, 0, 0), FMath::DegreesToRadians(60.0f));
 		// Add second view even further away
-		PrioParams.View.Views.Add(MakeReplicationView(TestPosition + 2.0f*SphereNetObjectPrioritizerConfig->OuterRadius, FVector(1, 0, 0), FMath::DegreesToRadians(60)).Views[0]);
+		PrioParams.View.Views.Add(MakeReplicationView(TestPosition + 2.0f*SphereNetObjectPrioritizerConfig->OuterRadius, FVector(1, 0, 0), FMath::DegreesToRadians(60.0f)).Views[0]);
 
 		SphereNetObjectPrioritizer->PrioritizeWithParams(PrioParams);
 		UE_NET_ASSERT_EQ(OutPriority, SphereNetObjectPrioritizerConfig->OutsidePriority);
@@ -296,10 +296,10 @@ UE_NET_TEST_FIXTURE(FTestSphereNetObjectPrioritizer, ViewPositionVeryFarAwayFrom
 
 	// Test more than two views.
 	{
-		PrioParams.View = MakeReplicationView(TestPosition + SphereNetObjectPrioritizerConfig->OuterRadius, FVector(1, 0, 0), FMath::DegreesToRadians(60));
+		PrioParams.View = MakeReplicationView(TestPosition + SphereNetObjectPrioritizerConfig->OuterRadius, FVector(1, 0, 0), FMath::DegreesToRadians(60.0f));
 		// Add a couple of views
-		PrioParams.View.Views.Add(MakeReplicationView(TestPosition + 2.0f*SphereNetObjectPrioritizerConfig->OuterRadius, FVector(1, 0, 0), FMath::DegreesToRadians(60)).Views[0]);
-		PrioParams.View.Views.Add(MakeReplicationView(TestPosition - 4.0f*SphereNetObjectPrioritizerConfig->OuterRadius, FVector(1, 0, 0), FMath::DegreesToRadians(60)).Views[0]);
+		PrioParams.View.Views.Add(MakeReplicationView(TestPosition + 2.0f*SphereNetObjectPrioritizerConfig->OuterRadius, FVector(1, 0, 0), FMath::DegreesToRadians(60.0f)).Views[0]);
+		PrioParams.View.Views.Add(MakeReplicationView(TestPosition - 4.0f*SphereNetObjectPrioritizerConfig->OuterRadius, FVector(1, 0, 0), FMath::DegreesToRadians(60.0f)).Views[0]);
 
 		SphereNetObjectPrioritizer->PrioritizeWithParams(PrioParams);
 		UE_NET_ASSERT_EQ(OutPriority, SphereNetObjectPrioritizerConfig->OutsidePriority);
@@ -333,7 +333,7 @@ UE_NET_TEST_FIXTURE(FTestSphereNetObjectPrioritizer, ViewPositionInRadiusInterva
 
 	// Test single view
 	{
-		PrioParams.View = MakeReplicationView(ViewPosition, FVector(1, 0, 0), FMath::DegreesToRadians(60));
+		PrioParams.View = MakeReplicationView(ViewPosition, FVector(1, 0, 0), FMath::DegreesToRadians(60.0f));
 		SphereNetObjectPrioritizer->PrioritizeWithParams(PrioParams);
 
 		for (float Priority : OutPriorities)
diff --git a/Engine/Plugins/Runtime/ReplicationSystemTestPlugin/Source/Private/Tests/ReplicationSystem/ReplicatedTestObject.cpp b/Engine/Plugins/Runtime/ReplicationSystemTestPlugin/Source/Private/Tests/ReplicationSystem/ReplicatedTestObject.cpp
--- a/Engine/Plugins/Runtime/ReplicationSystemTestPlugin/Source/Private/Tests/ReplicationSystem/ReplicatedTestObject.cpp
+++ b/Engine/Plugins/Runtime/ReplicationSystemTestPlugin/Source/Private/Tests/ReplicationSystem/ReplicatedTestObject.cpp
@@ -116,10 +116,10 @@ bool UReplicatedTestObjectBridge::WriteCreationHeader(UE::Net::FNetSerialization
 	}
 	if (const UTestReplicatedIrisObject* TestReplicatedIrisObject = Cast<UTestReplicatedIrisObject>(Object))
 	{
-		NumComponentsToSpawn = TestReplicatedIrisObject->Components.Num();
-		NumIrisComponentsToSpawn = TestReplicatedIrisObject->IrisComponents.Num();
-		NumDynamicComponentsToSpawn = TestReplicatedIrisObject->DynamicStateComponents.Num();
-		NumConnectionFilteredComponentsToSpawn = TestReplicatedIrisObject->ConnectionFilteredComponents.Num();
+		NumComponentsToSpawn = IntCastChecked<uint16>(TestReplicatedIrisObject->Components.Num());
+		NumIrisComponentsToSpawn = IntCastChecked<uint16>(TestReplicatedIrisObject->IrisComponents.Num());
+		NumDynamicComponentsToSpawn = IntCastChecked<uint16>(TestReplicatedIrisObject->DynamicStateComponents.Num());
+		NumConnectionFilteredComponentsToSpawn = IntCastChecked<uint16>(TestReplicatedIrisObject->ConnectionFilteredComponents.Num());
 	}
 
 	UObject* Archetype = Object->GetArchetype();
diff --git a/Engine/Plugins/Runtime/ReplicationSystemTestPlugin/Source/Private/Tests/ReplicationSystem/TestDirtyNetObjectTracker.cpp b/Engine/Plugins/Runtime/ReplicationSystemTestPlugin/Source/Private/Tests/ReplicationSystem/TestDirtyNetObjectTracker.cpp
--- a/Engine/Plugins/Runtime/ReplicationSystemTestPlugin/Source/Private/Tests/ReplicationSystem/TestDirtyNetObjectTracker.cpp
+++ b/Engine/Plugins/Runtime/ReplicationSystemTestPlugin/Source/Private/Tests/ReplicationSystem/TestDirtyNetObjectTracker.cpp
@@ -52,7 +52,7 @@ UE_NET_TEST_FIXTURE(FDirtyNetObjectTrackerTestFixture, TestNoObjectIsDirtyFromSt
 
 UE_NET_TEST_FIXTURE(FDirtyNetObjectTrackerTestFixture, CannotMarkInvalidObjectAsDirty)
 {
-	MarkNetObjectStateDirty(ReplicationSystemId, FNetRefHandle().GetId());
+	MarkNetObjectStateDirty(ReplicationSystemId, FNetRefHandleManager::InvalidInternalIndex);
 	MarkNetObjectStateDirty(ReplicationSystemId, NetObjectIndexRangeEnd + 1);
 
 	FDirtyObjectsAccessor DirtyObjectsAccessor(*DirtyNetObjectTracker);
diff --git a/Engine/Plugins/Runtime/ReplicationSystemTestPlugin/Source/Private/Tests/ReplicationSystem/TestReplicationRecord.cpp b/Engine/Plugins/Runtime/ReplicationSystemTestPlugin/Source/Private/Tests/ReplicationSystem/TestReplicationRecord.cpp
--- a/Engine/Plugins/Runtime/ReplicationSystemTestPlugin/Source/Private/Tests/ReplicationSystem/TestReplicationRecord.cpp
+++ b/Engine/Plugins/Runtime/ReplicationSystemTestPlugin/Source/Private/Tests/ReplicationSystem/TestReplicationRecord.cpp
@@ -40,12 +40,12 @@ UE_NET_TEST(ReplicationRecord, PushInfoAndAddToList)
 	UE_NET_ASSERT_EQ(0U, (uint32)List.FirstRecordIndex);
 
 	// Iterate over entries from start to end
-	uint32 CurrentIndex = List.FirstRecordIndex;
-	uint32 CurrentExpectedIndex = 0;
+	uint16 CurrentIndex = List.FirstRecordIndex;
+	uint16 CurrentExpectedIndex = 0;
 	
 	while (CurrentIndex != InvalidReplicationRecordIndex)
 	{
-		UE_NET_ASSERT_EQ((uint32)CurrentExpectedIndex, CurrentIndex);
+		UE_NET_ASSERT_EQ(CurrentExpectedIndex, CurrentIndex);
 
 		const FReplicationRecord::FRecordInfo* CurrentRecordInfo = ReplicationRecord.GetInfoForIndex(CurrentIndex);
 		UE_NET_ASSERT_TRUE(CurrentRecordInfo != nullptr);
@@ -53,7 +53,7 @@ UE_NET_TEST(ReplicationRecord, PushInfoAndAddToList)
 		CurrentIndex = CurrentRecordInfo->NextIndex;
 		++CurrentExpectedIndex;
 	};
-	UE_NET_ASSERT_EQ(InfoCount, CurrentExpectedIndex);
+	UE_NET_ASSERT_EQ(InfoCount, static_cast<uint32>(CurrentExpectedIndex));
 }
 
 UE_NET_TEST(ReplicationRecord, PopInfoAndRemoveFromList)
diff --git a/Engine/Plugins/Runtime/ReplicationSystemTestPlugin/Source/Private/Tests/Serialization/TestArrayPropertyNetSerializer.cpp b/Engine/Plugins/Runtime/ReplicationSystemTestPlugin/Source/Private/Tests/Serialization/TestArrayPropertyNetSerializer.cpp
--- a/Engine/Plugins/Runtime/ReplicationSystemTestPlugin/Source/Private/Tests/Serialization/TestArrayPropertyNetSerializer.cpp
+++ b/Engine/Plugins/Runtime/ReplicationSystemTestPlugin/Source/Private/Tests/Serialization/TestArrayPropertyNetSerializer.cpp
@@ -495,7 +495,7 @@ void FTestArrayPropertyNetSerializerBase::FreeDynamicState(void* StateBuffer)
 
 void FTestArrayPropertyNetSerializerBase::Serialize(const void* StateBuffer, void* BitStreamBuffer, SIZE_T BitStreamBufferSize)
 {
-	Writer.InitBytes(BitStreamBuffer, BitStreamBufferSize);
+	Writer.InitBytes(BitStreamBuffer, IntCastChecked<uint32>(BitStreamBufferSize));
 
 	FNetSerializeArgs Args = {};
 	Args.Version = ArrayPropertyNetSerializer->Version;
@@ -508,7 +508,7 @@ void FTestArrayPropertyNetSerializerBase::Serialize(const void* StateBuffer, voi
 
 void FTestArrayPropertyNetSerializerBase::SerializeDelta(const void* StateBuffer, const void* PrevStateBuffer, void* BitStreamBuffer, SIZE_T BitStreamBufferSize)
 {
-	Writer.InitBytes(BitStreamBuffer, BitStreamBufferSize);
+	Writer.InitBytes(BitStreamBuffer, IntCastChecked<uint32>(BitStreamBufferSize));
 
 	FNetSerializeDeltaArgs Args = {};
 	Args.Version = ArrayPropertyNetSerializer->Version;
@@ -535,7 +535,7 @@ void FTestSimpleArrayPropertyNetSerializer::SetUp()
 
 	NonEmptyArrayInstance0.ArrayOfUint.SetNumZeroed(3);
 	NonEmptyArrayInstance1.ArrayOfUint.SetNumZeroed(3);
-	for (SIZE_T It = 0, EndIt = 3; It != EndIt; ++It)
+	for (int32 It = 0, EndIt = 3; It != EndIt; ++It)
 	{
 		NonEmptyArrayInstance0.ArrayOfUint[It] = It + 1;
 		NonEmptyArrayInstance1.ArrayOfUint[It] = It + 1;
@@ -618,7 +618,7 @@ void FTestSimpleArrayPropertyNetSerializer::TestDequantize()
 	UE_NET_ASSERT_EQ(NonEmptyArrayInstance0.ArrayOfUint.Num(), TargetStruct.ArrayOfUint.Num());
 	for (const auto& OriginalValue : MakeArrayView(NonEmptyArrayInstance0.ArrayOfUint))
 	{
-		const SIZE_T ElementIndex = &OriginalValue - NonEmptyArrayInstance0.ArrayOfUint.GetData();
+		const int32 ElementIndex = static_cast<int32>(&OriginalValue - NonEmptyArrayInstance0.ArrayOfUint.GetData());
 		const auto& DequantizedValue =  TargetStruct.ArrayOfUint[ElementIndex];
 		UE_NET_ASSERT_EQ(OriginalValue, DequantizedValue);
 	}
@@ -645,7 +645,7 @@ void FTestSimpleArrayPropertyNetSerializer::TestCloneDynamicState()
 		UE_NET_ASSERT_EQ(NonEmptyArrayInstance0.ArrayOfUint.Num(), TargetStruct.ArrayOfUint.Num());
 		for (const auto& OriginalValue : MakeArrayView(NonEmptyArrayInstance0.ArrayOfUint))
 		{
-			const SIZE_T ElementIndex = &OriginalValue - NonEmptyArrayInstance0.ArrayOfUint.GetData();
+			const int32 ElementIndex = static_cast<uint32>(&OriginalValue - NonEmptyArrayInstance0.ArrayOfUint.GetData());
 			const auto& DequantizedValue = TargetStruct.ArrayOfUint[ElementIndex];
 			UE_NET_ASSERT_EQ(OriginalValue, DequantizedValue);
 		}
@@ -723,7 +723,7 @@ void FTestSimpleArrayPropertyNetSerializer::TestDeserialize()
 		UE_NET_ASSERT_EQ(NonEmptyArrayInstance0.ArrayOfUint.Num(), TargetStruct.ArrayOfUint.Num());
 		for (const auto& OriginalValue : MakeArrayView(NonEmptyArrayInstance0.ArrayOfUint))
 		{
-			const SIZE_T ElementIndex = &OriginalValue - NonEmptyArrayInstance0.ArrayOfUint.GetData();
+			const int32 ElementIndex = static_cast<int32>(&OriginalValue - NonEmptyArrayInstance0.ArrayOfUint.GetData());
 			const auto& DequantizedValue = TargetStruct.ArrayOfUint[ElementIndex];
 			UE_NET_ASSERT_EQ(OriginalValue, DequantizedValue);
 		}
@@ -839,7 +839,7 @@ void FTestSimpleArrayPropertyNetSerializer::TestDeserializeDelta()
 			UE_NET_ASSERT_EQ(NonEmptyArrayInstance0.ArrayOfUint.Num(), TargetStruct.ArrayOfUint.Num());
 			for (const auto& OriginalValue : MakeArrayView(NonEmptyArrayInstance0.ArrayOfUint))
 			{
-				const SIZE_T ElementIndex = &OriginalValue - NonEmptyArrayInstance0.ArrayOfUint.GetData();
+				const int32 ElementIndex = static_cast<int32>(&OriginalValue - NonEmptyArrayInstance0.ArrayOfUint.GetData());
 				const auto& DequantizedValue = TargetStruct.ArrayOfUint[ElementIndex];
 				UE_NET_ASSERT_EQ(OriginalValue, DequantizedValue);
 			}
@@ -903,7 +903,7 @@ void FTestSimpleArrayPropertyNetSerializer::TestDeserializeDelta()
 			UE_NET_ASSERT_EQ(NonEmptyArrayInstance0.ArrayOfUint.Num(), TargetStruct.ArrayOfUint.Num());
 			for (const auto& OriginalValue : MakeArrayView(NonEmptyArrayInstance0.ArrayOfUint))
 			{
-				const SIZE_T ElementIndex = &OriginalValue - NonEmptyArrayInstance0.ArrayOfUint.GetData();
+				const int32 ElementIndex = static_cast<int32>(&OriginalValue - NonEmptyArrayInstance0.ArrayOfUint.GetData());
 				const auto& DequantizedValue = TargetStruct.ArrayOfUint[ElementIndex];
 				UE_NET_ASSERT_EQ(OriginalValue, DequantizedValue);
 			}
@@ -1131,7 +1131,7 @@ void FTestComplexArrayPropertyNetSerializer::SetUp()
 	NonEmptyArrayInstance0.ArrayOfStructWithArray.SetNum(3);
 	NonEmptyArrayInstance1.ArrayOfStructWithArray.Empty(3);
 	NonEmptyArrayInstance1.ArrayOfStructWithArray.SetNum(3);
-	for (SIZE_T It = 0, EndIt = 3; It != EndIt; ++It)
+	for (int32 It = 0, EndIt = 3; It != EndIt; ++It)
 	{
 		NonEmptyArrayInstance0.ArrayOfStructWithArray[It].ArrayOfUint.SetNumZeroed(It + 1);
 		NonEmptyArrayInstance1.ArrayOfStructWithArray[It].ArrayOfUint.SetNumZeroed(It + 1);
@@ -1210,13 +1210,13 @@ void FTestComplexArrayPropertyNetSerializer::TestDequantize()
 	UE_NET_ASSERT_EQ(NonEmptyArrayInstance0.ArrayOfStructWithArray.Num(), TargetStruct.ArrayOfStructWithArray.Num());
 	for (const auto& OriginalValue : MakeArrayView(NonEmptyArrayInstance0.ArrayOfStructWithArray))
 	{
-		const SIZE_T ElementIndex = &OriginalValue - NonEmptyArrayInstance0.ArrayOfStructWithArray.GetData();
+		const int32 ElementIndex = static_cast<int32>(&OriginalValue - NonEmptyArrayInstance0.ArrayOfStructWithArray.GetData());
 		const auto& DequantizedValue =  TargetStruct.ArrayOfStructWithArray[ElementIndex];
 		UE_NET_ASSERT_EQ(OriginalValue.ArrayOfUint.Num(), DequantizedValue.ArrayOfUint.Num());
 
 		for (const auto& OriginalInnerValue : MakeArrayView(OriginalValue.ArrayOfUint))
 		{
-			const SIZE_T InnerElementIndex = &OriginalInnerValue - OriginalValue.ArrayOfUint.GetData();
+			const int32 InnerElementIndex = static_cast<int32>(&OriginalInnerValue - OriginalValue.ArrayOfUint.GetData());
 			const auto& DequantizedInnerValue = DequantizedValue.ArrayOfUint[InnerElementIndex];
 			UE_NET_ASSERT_EQ(OriginalInnerValue, DequantizedInnerValue);
 		}
@@ -1244,13 +1244,13 @@ void FTestComplexArrayPropertyNetSerializer::TestCloneDynamicState()
 		UE_NET_ASSERT_EQ(NonEmptyArrayInstance0.ArrayOfStructWithArray.Num(), TargetStruct.ArrayOfStructWithArray.Num());
 		for (const auto& OriginalValue : MakeArrayView(NonEmptyArrayInstance0.ArrayOfStructWithArray))
 		{
-			const SIZE_T ElementIndex = &OriginalValue - NonEmptyArrayInstance0.ArrayOfStructWithArray.GetData();
+			const int32 ElementIndex = static_cast<int32>(&OriginalValue - NonEmptyArrayInstance0.ArrayOfStructWithArray.GetData());
 			const auto& DequantizedValue = TargetStruct.ArrayOfStructWithArray[ElementIndex];
 			UE_NET_ASSERT_EQ(OriginalValue.ArrayOfUint.Num(), DequantizedValue.ArrayOfUint.Num());
 
 			for (const auto& OriginalInnerValue : MakeArrayView(OriginalValue.ArrayOfUint))
 			{
-				const SIZE_T InnerElementIndex = &OriginalInnerValue - OriginalValue.ArrayOfUint.GetData();
+				const int32 InnerElementIndex = static_cast<int32>(&OriginalInnerValue - OriginalValue.ArrayOfUint.GetData());
 				const auto& DequantizedInnerValue = DequantizedValue.ArrayOfUint[InnerElementIndex];
 				UE_NET_ASSERT_EQ(OriginalInnerValue, DequantizedInnerValue);
 			}
diff --git a/Engine/Plugins/Runtime/ReplicationSystemTestPlugin/Source/Private/Tests/Serialization/TestCompressedChangeMask.cpp b/Engine/Plugins/Runtime/ReplicationSystemTestPlugin/Source/Private/Tests/Serialization/TestCompressedChangeMask.cpp
--- a/Engine/Plugins/Runtime/ReplicationSystemTestPlugin/Source/Private/Tests/Serialization/TestCompressedChangeMask.cpp
+++ b/Engine/Plugins/Runtime/ReplicationSystemTestPlugin/Source/Private/Tests/Serialization/TestCompressedChangeMask.cpp
@@ -232,8 +232,8 @@ UE_NET_TEST_FIXTURE(FCompressedChangeMaskTest, WriteChangeMasks)
 	// Log stats
 	if (ChangeMaskStats.TotalChangeMaskBits > 0)
 	{
-		float Ratio = ChangeMaskStats.TotalWrittenBits / (float)ChangeMaskStats.TotalChangeMaskBits;
-		float WrittenBitsPerDirtyBit = ChangeMaskStats.TotalWrittenBits / (float)ChangeMaskStats.TotalSetBits;
+		float Ratio = static_cast<float>(ChangeMaskStats.TotalWrittenBits) / static_cast<float>(ChangeMaskStats.TotalChangeMaskBits);
+		float WrittenBitsPerDirtyBit = static_cast<float>(ChangeMaskStats.TotalWrittenBits) / static_cast<float>(ChangeMaskStats.TotalSetBits);
 
 		UE_NET_LOG(TEXT("")) << TStringBuilder<128>().Appendf(TEXT("WriteChangeMasks::Summary SetBits %u / TotalBits %u WrittenBits %u Ratio %f BitsPerDirtyBit %f"), ChangeMaskStats.TotalSetBits, ChangeMaskStats.TotalChangeMaskBits, ChangeMaskStats.TotalWrittenBits, Ratio, WrittenBitsPerDirtyBit);
 	}
diff --git a/Engine/Plugins/Runtime/ReplicationSystemTestPlugin/Source/Private/Tests/Serialization/TestEnumNetSerializer.cpp b/Engine/Plugins/Runtime/ReplicationSystemTestPlugin/Source/Private/Tests/Serialization/TestEnumNetSerializer.cpp
--- a/Engine/Plugins/Runtime/ReplicationSystemTestPlugin/Source/Private/Tests/Serialization/TestEnumNetSerializer.cpp
+++ b/Engine/Plugins/Runtime/ReplicationSystemTestPlugin/Source/Private/Tests/Serialization/TestEnumNetSerializer.cpp
@@ -198,7 +198,7 @@ void FTestEnumIntNetSerializer<SerializerConfig, SourceType, EnumType>::TestVali
 	{
 		TArray<bool> ExpectedResults;
 		ExpectedResults.SetNumUninitialized(Values.Num());
-		for (SIZE_T ValueIt = 0, ValueEndIt = Values.Num(); ValueIt != ValueEndIt; ++ValueIt)
+		for (int32 ValueIt = 0, ValueEndIt = Values.Num(); ValueIt != ValueEndIt; ++ValueIt)
 		{
 			ExpectedResults[ValueIt] = true;
 		}
diff --git a/Engine/Plugins/Runtime/ReplicationSystemTestPlugin/Source/Private/Tests/Serialization/TestFloatNetSerializer.cpp b/Engine/Plugins/Runtime/ReplicationSystemTestPlugin/Source/Private/Tests/Serialization/TestFloatNetSerializer.cpp
--- a/Engine/Plugins/Runtime/ReplicationSystemTestPlugin/Source/Private/Tests/Serialization/TestFloatNetSerializer.cpp
+++ b/Engine/Plugins/Runtime/ReplicationSystemTestPlugin/Source/Private/Tests/Serialization/TestFloatNetSerializer.cpp
@@ -43,7 +43,7 @@ const T FTestFloatNetSerializerBase<T>::Values[] =
 	47.11f,
 	757575.0f,
 	757576.0f,
-	1.26698556E-23,
+	1.26698556E-23f,
 };
 template<typename T>
 const SIZE_T FTestFloatNetSerializerBase<T>::ValueCount = sizeof(Values)/sizeof(Values[0]);
diff --git a/Engine/Plugins/Runtime/ReplicationSystemTestPlugin/Source/Private/Tests/Serialization/TestIntRangeNetSerializer.cpp b/Engine/Plugins/Runtime/ReplicationSystemTestPlugin/Source/Private/Tests/Serialization/TestIntRangeNetSerializer.cpp
--- a/Engine/Plugins/Runtime/ReplicationSystemTestPlugin/Source/Private/Tests/Serialization/TestIntRangeNetSerializer.cpp
+++ b/Engine/Plugins/Runtime/ReplicationSystemTestPlugin/Source/Private/Tests/Serialization/TestIntRangeNetSerializer.cpp
@@ -227,7 +227,7 @@ void FTestIntRangeNetSerializer<SerializerConfig, SourceType>::SetupConfigs()
 		SerializerConfig& Config = Configs[0];
 		Config.LowerBound = std::numeric_limits<SourceType>::min();
 		Config.UpperBound = std::numeric_limits<SourceType>::max();
-		Config.BitCount = GetBitsNeededForRange(Config.LowerBound, Config.UpperBound);
+		Config.BitCount = static_cast<uint8>(GetBitsNeededForRange(Config.LowerBound, Config.UpperBound));
 	}
 
 	// Zero range
@@ -235,7 +235,7 @@ void FTestIntRangeNetSerializer<SerializerConfig, SourceType>::SetupConfigs()
 		SerializerConfig& Config = Configs[1];
 		Config.LowerBound = -64;
 		Config.UpperBound = -64;
-		Config.BitCount = GetBitsNeededForRange(Config.LowerBound, Config.UpperBound);
+		Config.BitCount = static_cast<uint8>(GetBitsNeededForRange(Config.LowerBound, Config.UpperBound));
 	}
 
 	// Arbitrary ranges
@@ -243,14 +243,14 @@ void FTestIntRangeNetSerializer<SerializerConfig, SourceType>::SetupConfigs()
 		SerializerConfig& Config = Configs[2];
 		Config.LowerBound = std::numeric_limits<SourceType>::min()/2;
 		Config.UpperBound = std::numeric_limits<SourceType>::max()/4;
-		Config.BitCount = GetBitsNeededForRange(Config.LowerBound, Config.UpperBound);
+		Config.BitCount = static_cast<uint8>(GetBitsNeededForRange(Config.LowerBound, Config.UpperBound));
 	}
 
 	{
 		SerializerConfig& Config = Configs[3];
 		Config.LowerBound = std::numeric_limits<SourceType>::max()/4;
 		Config.UpperBound = std::numeric_limits<SourceType>::max()/2;
-		Config.BitCount = GetBitsNeededForRange(Config.LowerBound, Config.UpperBound);
+		Config.BitCount = static_cast<uint8>(GetBitsNeededForRange(Config.LowerBound, Config.UpperBound));
 	}
 }
 
diff --git a/Engine/Plugins/Runtime/ReplicationSystemTestPlugin/Source/Private/Tests/Serialization/TestNetRoleNetSerializer.cpp b/Engine/Plugins/Runtime/ReplicationSystemTestPlugin/Source/Private/Tests/Serialization/TestNetRoleNetSerializer.cpp
--- a/Engine/Plugins/Runtime/ReplicationSystemTestPlugin/Source/Private/Tests/Serialization/TestNetRoleNetSerializer.cpp
+++ b/Engine/Plugins/Runtime/ReplicationSystemTestPlugin/Source/Private/Tests/Serialization/TestNetRoleNetSerializer.cpp
@@ -194,7 +194,7 @@ void FTestNetRoleNetSerializer::TestValidate()
 	{
 		TArray<bool> ExpectedResults;
 		ExpectedResults.SetNumUninitialized(Values.Num());
-		for (SIZE_T ValueIt = 0, ValueEndIt = Values.Num(); ValueIt != ValueEndIt; ++ValueIt)
+		for (int32 ValueIt = 0, ValueEndIt = Values.Num(); ValueIt != ValueEndIt; ++ValueIt)
 		{
 			ExpectedResults[ValueIt] = true;
 		}
diff --git a/Engine/Plugins/Runtime/ReplicationSystemTestPlugin/Source/Private/Tests/Serialization/TestPackedVectorNetSerializers.cpp b/Engine/Plugins/Runtime/ReplicationSystemTestPlugin/Source/Private/Tests/Serialization/TestPackedVectorNetSerializers.cpp
--- a/Engine/Plugins/Runtime/ReplicationSystemTestPlugin/Source/Private/Tests/Serialization/TestPackedVectorNetSerializers.cpp
+++ b/Engine/Plugins/Runtime/ReplicationSystemTestPlugin/Source/Private/Tests/Serialization/TestPackedVectorNetSerializers.cpp
@@ -333,9 +333,9 @@ bool FTestPackedVectorNetSerializerBase::IsAlmostEqual(NetSerializerValuePointer
 	// Allow ULPs difference and absolute difference
 	const FVector V0 =  *reinterpret_cast<FVector*>(Value0);
 	const FVector V1 =  *reinterpret_cast<FVector*>(Value1);
-	const float AbsDiffX = FMath::Abs(V0.X - V1.X);
-	const float AbsDiffY = FMath::Abs(V0.Y - V1.Y);
-	const float AbsDiffZ = FMath::Abs(V0.Z - V1.Z);
+	const ScalarType AbsDiffX = FMath::Abs(V0.X - V1.X);
+	const ScalarType AbsDiffY = FMath::Abs(V0.Y - V1.Y);
+	const ScalarType AbsDiffZ = FMath::Abs(V0.Z - V1.Z);
 
 	const bool EqualUlpX = FMath::IsNearlyEqualByULP(V0.X, V1.X, MaxUlpDiff);
 	const bool EqualUlpY = FMath::IsNearlyEqualByULP(V0.Y, V1.Y, MaxUlpDiff);
@@ -396,7 +396,7 @@ void FTestVectorNetQuantizeNormalNetSerializer::TestIsEqual()
 
 	TArray<FVector> CompareValues[2];
 	TArray<bool> ExpectedResults[2];
-	const SIZE_T AllValueCount = AllValues.Num();
+	const int32 AllValueCount = AllValues.Num();
 
 	{
 		// Make sure nothing messes with our test values.
diff --git a/Engine/Plugins/Runtime/ReplicationSystemTestPlugin/Source/Private/Tests/Serialization/TestPolymorphicStructNetSerializer.h b/Engine/Plugins/Runtime/ReplicationSystemTestPlugin/Source/Private/Tests/Serialization/TestPolymorphicStructNetSerializer.h
--- a/Engine/Plugins/Runtime/ReplicationSystemTestPlugin/Source/Private/Tests/Serialization/TestPolymorphicStructNetSerializer.h
+++ b/Engine/Plugins/Runtime/ReplicationSystemTestPlugin/Source/Private/Tests/Serialization/TestPolymorphicStructNetSerializer.h
@@ -134,7 +134,7 @@ public:
 
 	static void SetArrayNum(FExamplePolymorphicArrayStruct& ArrayContainer, SIZE_T Num)
 	{
-		ArrayContainer.Data.SetNum(static_cast<SSIZE_T>(Num));
+		ArrayContainer.Data.SetNum(IntCastChecked<int32>(Num));
 	}
 
 private:
diff --git a/Engine/Plugins/Runtime/ReplicationSystemTestPlugin/Source/Private/Tests/Serialization/TestQuatNetSerializers.cpp b/Engine/Plugins/Runtime/ReplicationSystemTestPlugin/Source/Private/Tests/Serialization/TestQuatNetSerializers.cpp
--- a/Engine/Plugins/Runtime/ReplicationSystemTestPlugin/Source/Private/Tests/Serialization/TestQuatNetSerializers.cpp
+++ b/Engine/Plugins/Runtime/ReplicationSystemTestPlugin/Source/Private/Tests/Serialization/TestQuatNetSerializers.cpp
@@ -141,7 +141,7 @@ void FTestQuatNetSerializerBase<QuatType>::TestIsEqual()
 {
 	TArray<QuatType> CompareValues[2];
 	TArray<bool> ExpectedResults[2];
-	const SIZE_T ValueCount = Values.Num();
+	const int32 ValueCount = Values.Num();
 
 	{
 		// Make sure nothing messes with our test values.
@@ -153,7 +153,7 @@ void FTestQuatNetSerializerBase<QuatType>::TestIsEqual()
 	{
 		CompareValues[1].SetNumUninitialized(ValueCount);
 		ExpectedResults[1].Init(false, ValueCount);
-		for (SIZE_T ValueIt = 0; ValueIt < ValueCount; ++ValueIt)
+		for (int32 ValueIt = 0; ValueIt < ValueCount; ++ValueIt)
 		{
 			// Make sure nothing messes with our test values.
 			FPlatformMemory::Memcpy(CompareValues[1].GetData() + ValueIt, Values.GetData() + ((ValueIt + 1) % ValueCount), sizeof(QuatType));
diff --git a/Engine/Plugins/Runtime/ReplicationSystemTestPlugin/Source/Private/Tests/Serialization/TestRepMovementNetSerializer.cpp b/Engine/Plugins/Runtime/ReplicationSystemTestPlugin/Source/Private/Tests/Serialization/TestRepMovementNetSerializer.cpp
--- a/Engine/Plugins/Runtime/ReplicationSystemTestPlugin/Source/Private/Tests/Serialization/TestRepMovementNetSerializer.cpp
+++ b/Engine/Plugins/Runtime/ReplicationSystemTestPlugin/Source/Private/Tests/Serialization/TestRepMovementNetSerializer.cpp
@@ -72,19 +72,19 @@ void FTestRepMovementNetSerializer::SetUp()
 
 void FTestRepMovementNetSerializer::TestIsEqual()
 {
-	constexpr SIZE_T TestRoundCount = 2;
+	constexpr int32 TestRoundCount = 2;
 	TArray<FRepMovement> CompareValues[TestRoundCount];
 	TArray<bool> ExpectedResults[TestRoundCount];
 
-	const SIZE_T ValueCount = Values.Num();
-	for (SIZE_T TestRoundIt = 0, TestRoundEndIt = TestRoundCount; TestRoundIt != TestRoundEndIt; ++TestRoundIt)
+	const int32 ValueCount = Values.Num();
+	for (int32 TestRoundIt = 0, TestRoundEndIt = TestRoundCount; TestRoundIt != TestRoundEndIt; ++TestRoundIt)
 	{
 		CompareValues[TestRoundIt].Reserve(ValueCount);
 		ExpectedResults[TestRoundIt].Reserve(ValueCount);
 
-		const SIZE_T SourceValueItOffset = TestRoundIt;
+		const int32 SourceValueItOffset = TestRoundIt;
 		const bool ExpectedResult = (TestRoundIt == 0U);
-		for (SIZE_T ValueIt = 0, ValueEndIt = ValueCount; ValueIt < ValueEndIt; ++ValueIt)
+		for (int32 ValueIt = 0, ValueEndIt = ValueCount; ValueIt < ValueEndIt; ++ValueIt)
 		{
 			CompareValues[TestRoundIt].Push(Values[(ValueIt + SourceValueItOffset) % ValueEndIt]);
 			ExpectedResults[TestRoundIt].Push(ExpectedResult);
diff --git a/Engine/Plugins/Runtime/ReplicationSystemTestPlugin/Source/Private/Tests/Serialization/TestSoftClassPathNetSerializer.cpp b/Engine/Plugins/Runtime/ReplicationSystemTestPlugin/Source/Private/Tests/Serialization/TestSoftClassPathNetSerializer.cpp
--- a/Engine/Plugins/Runtime/ReplicationSystemTestPlugin/Source/Private/Tests/Serialization/TestSoftClassPathNetSerializer.cpp
+++ b/Engine/Plugins/Runtime/ReplicationSystemTestPlugin/Source/Private/Tests/Serialization/TestSoftClassPathNetSerializer.cpp
@@ -116,7 +116,7 @@ void FTestSoftClassPathNetSerializer::TestValidate()
 
 void FTestSoftClassPathNetSerializer::TestIsEqual()
 {
-	const SIZE_T TestValueCount = TestValues.Num();
+	const int32 TestValueCount = TestValues.Num();
 
 	TArray<FSoftClassPath> CompareValues[2];
 	TArray<bool> ExpectedResults[2];
@@ -128,8 +128,8 @@ void FTestSoftClassPathNetSerializer::TestIsEqual()
 	ExpectedResults[1].Reserve(TestValueCount);
 	for (const FSoftClassPath& Value : TestValues)
 	{
-		const SIZE_T ValueIndex = &Value - TestValues.GetData();
-		const SIZE_T NextValueIndex = (ValueIndex + 1U) % TestValueCount;
+		const int32 ValueIndex = static_cast<int32>(&Value - TestValues.GetData());
+		const int32 NextValueIndex = (ValueIndex + 1U) % TestValueCount;
 		CompareValues[1].Add(TestValues[NextValueIndex]);
 		ExpectedResults[1].Add(TestValues[ValueIndex] == TestValues[NextValueIndex]);
 	}
diff --git a/Engine/Plugins/Runtime/ReplicationSystemTestPlugin/Source/Private/Tests/Serialization/TestSoftObjectNetSerializer.cpp b/Engine/Plugins/Runtime/ReplicationSystemTestPlugin/Source/Private/Tests/Serialization/TestSoftObjectNetSerializer.cpp
--- a/Engine/Plugins/Runtime/ReplicationSystemTestPlugin/Source/Private/Tests/Serialization/TestSoftObjectNetSerializer.cpp
+++ b/Engine/Plugins/Runtime/ReplicationSystemTestPlugin/Source/Private/Tests/Serialization/TestSoftObjectNetSerializer.cpp
@@ -119,7 +119,7 @@ void FTestSoftObjectNetSerializer::TestValidate()
 
 void FTestSoftObjectNetSerializer::TestIsEqual()
 {
-	const SIZE_T TestValueCount = TestValues.Num();
+	const int32 TestValueCount = TestValues.Num();
 
 	TArray<FSoftObjectPtr> CompareValues[2];
 	TArray<bool> ExpectedResults[2];
@@ -131,8 +131,8 @@ void FTestSoftObjectNetSerializer::TestIsEqual()
 	ExpectedResults[1].Reserve(TestValueCount);
 	for (const FSoftObjectPtr& Value : TestValues)
 	{
-		const SIZE_T ValueIndex = &Value - TestValues.GetData();
-		const SIZE_T NextValueIndex = (ValueIndex + 1U) % TestValueCount;
+		const int32 ValueIndex = static_cast<int32>(&Value - TestValues.GetData());
+		const int32 NextValueIndex = (ValueIndex + 1U) % TestValueCount;
 		CompareValues[1].Add(TestValues[NextValueIndex]);
 		ExpectedResults[1].Add(TestValues[ValueIndex] == TestValues[NextValueIndex]);
 	}
diff --git a/Engine/Plugins/Runtime/ReplicationSystemTestPlugin/Source/Private/Tests/Serialization/TestSoftObjectPathNetSerializer.cpp b/Engine/Plugins/Runtime/ReplicationSystemTestPlugin/Source/Private/Tests/Serialization/TestSoftObjectPathNetSerializer.cpp
--- a/Engine/Plugins/Runtime/ReplicationSystemTestPlugin/Source/Private/Tests/Serialization/TestSoftObjectPathNetSerializer.cpp
+++ b/Engine/Plugins/Runtime/ReplicationSystemTestPlugin/Source/Private/Tests/Serialization/TestSoftObjectPathNetSerializer.cpp
@@ -126,7 +126,7 @@ void FTestSoftObjectPathNetSerializer::TestValidate()
 
 void FTestSoftObjectPathNetSerializer::TestIsEqual()
 {
-	const SIZE_T TestValueCount = TestValues.Num();
+	const int32 TestValueCount = TestValues.Num();
 
 	TArray<FSoftObjectPath> CompareValues[2];
 	TArray<bool> ExpectedResults[2];
@@ -138,8 +138,8 @@ void FTestSoftObjectPathNetSerializer::TestIsEqual()
 	ExpectedResults[1].Reserve(TestValueCount);
 	for (const FSoftObjectPath& Value : TestValues)
 	{
-		const SIZE_T ValueIndex = &Value - TestValues.GetData();
-		const SIZE_T NextValueIndex = (ValueIndex + 1U) % TestValueCount;
+		const int32 ValueIndex = static_cast<int32>(&Value - TestValues.GetData());
+		const int32 NextValueIndex = (ValueIndex + 1U) % TestValueCount;
 		CompareValues[1].Add(TestValues[NextValueIndex]);
 		ExpectedResults[1].Add(TestValues[ValueIndex] == TestValues[NextValueIndex]);
 	}
diff --git a/Engine/Plugins/Runtime/ReplicationSystemTestPlugin/Source/Private/Tests/Serialization/TestUintRangeNetSerializer.cpp b/Engine/Plugins/Runtime/ReplicationSystemTestPlugin/Source/Private/Tests/Serialization/TestUintRangeNetSerializer.cpp
--- a/Engine/Plugins/Runtime/ReplicationSystemTestPlugin/Source/Private/Tests/Serialization/TestUintRangeNetSerializer.cpp
+++ b/Engine/Plugins/Runtime/ReplicationSystemTestPlugin/Source/Private/Tests/Serialization/TestUintRangeNetSerializer.cpp
@@ -108,7 +108,7 @@ void FTestUintRangeNetSerializer<SerializerConfig, SourceType>::SetupConfigs()
 		SerializerConfig& Config = Configs[0];
 		Config.LowerBound = std::numeric_limits<SourceType>::min();
 		Config.UpperBound = std::numeric_limits<SourceType>::max();
-		Config.BitCount = GetBitsNeededForRange(Config.LowerBound, Config.UpperBound);
+		Config.BitCount = static_cast<uint8>(GetBitsNeededForRange(Config.LowerBound, Config.UpperBound));
 	}
 
 	// Zero range
@@ -116,7 +116,7 @@ void FTestUintRangeNetSerializer<SerializerConfig, SourceType>::SetupConfigs()
 		SerializerConfig& Config = Configs[1];
 		Config.LowerBound = 47;
 		Config.UpperBound = 47;
-		Config.BitCount = GetBitsNeededForRange(Config.LowerBound, Config.UpperBound);
+		Config.BitCount = static_cast<uint8>(GetBitsNeededForRange(Config.LowerBound, Config.UpperBound));
 	}
 
 	// Arbitrary ranges
@@ -124,14 +124,14 @@ void FTestUintRangeNetSerializer<SerializerConfig, SourceType>::SetupConfigs()
 		SerializerConfig& Config = Configs[2];
 		Config.LowerBound = std::numeric_limits<SourceType>::min()/2;
 		Config.UpperBound = std::numeric_limits<SourceType>::max()/4;
-		Config.BitCount = GetBitsNeededForRange(Config.LowerBound, Config.UpperBound);
+		Config.BitCount = static_cast<uint8>(GetBitsNeededForRange(Config.LowerBound, Config.UpperBound));
 	}
 
 	{
 		SerializerConfig& Config = Configs[3];
 		Config.LowerBound = std::numeric_limits<SourceType>::max()/4;
 		Config.UpperBound = std::numeric_limits<SourceType>::max()/2;
-		Config.BitCount = GetBitsNeededForRange(Config.LowerBound, Config.UpperBound);
+		Config.BitCount = static_cast<uint8>(GetBitsNeededForRange(Config.LowerBound, Config.UpperBound));
 	}
 }
 
diff --git a/Engine/Plugins/Runtime/ReplicationSystemTestPlugin/Source/ReplicationSystemTestPlugin.Build.cs b/Engine/Plugins/Runtime/ReplicationSystemTestPlugin/Source/ReplicationSystemTestPlugin.Build.cs
--- a/Engine/Plugins/Runtime/ReplicationSystemTestPlugin/Source/ReplicationSystemTestPlugin.Build.cs
+++ b/Engine/Plugins/Runtime/ReplicationSystemTestPlugin/Source/ReplicationSystemTestPlugin.Build.cs
@@ -55,6 +55,8 @@ namespace UnrealBuildTool.Rules
 				PrivateIncludePathModuleNames.Add("Catch2");
 			}
 
+			UnsafeTypeCastWarningLevel = WarningLevel.Error;
+
 			PrivateDefinitions.Add(String.Format("UE_NET_WITH_LOW_LEVEL_TESTS={0}", Target.ExplicitTestsTarget ? "1" : "0"));
 
 			SetupIrisSupport(Target);
diff --git a/Engine/Source/Runtime/Core/Public/Containers/StringConv.h b/Engine/Source/Runtime/Core/Public/Containers/StringConv.h
--- a/Engine/Source/Runtime/Core/Public/Containers/StringConv.h
+++ b/Engine/Source/Runtime/Core/Public/Containers/StringConv.h
@@ -120,7 +120,7 @@ namespace StringConv
 					const uint32 LowCodepoint = Codepoint;
 
 					// Combine our high and low surrogates together to a single Unicode codepoint
-					Codepoint = StringConv::EncodeSurrogate(HighCodepoint, LowCodepoint);
+					Codepoint = StringConv::EncodeSurrogate(static_cast<uint16>(HighCodepoint), static_cast<uint16>(LowCodepoint));
 
 					StrBuffer[Index] = Codepoint;
 					{
diff --git a/Engine/Source/Runtime/Experimental/Iris/Core/Private/Iris/Core/IrisDebugging.cpp b/Engine/Source/Runtime/Experimental/Iris/Core/Private/Iris/Core/IrisDebugging.cpp
--- a/Engine/Source/Runtime/Experimental/Iris/Core/Private/Iris/Core/IrisDebugging.cpp
+++ b/Engine/Source/Runtime/Experimental/Iris/Core/Private/Iris/Core/IrisDebugging.cpp
@@ -332,7 +332,7 @@ void NetObjectStateToString(FStringBuilderBase& StringBuilder, FNetRefHandle Ref
 	FReplicationInstanceOperations::OutputInternalStateToString(NetSerializationContext, StringBuilder, nullptr, InternalStateBuffer, ReplicatedObjectData.InstanceProtocol, ReplicatedObjectData.Protocol);
 }
 
-void DebugOutputNetObjectState(uint32 NetRefHandleId, uint32 ReplicationSystemId)
+void DebugOutputNetObjectState(uint64 NetRefHandleId, uint32 ReplicationSystemId)
 {
 	const FNetRefHandle RefHandle = Private::FNetRefHandleManager::MakeNetRefHandle(NetRefHandleId, ReplicationSystemId);
 
diff --git a/Engine/Source/Runtime/Experimental/Iris/Core/Public/Iris/Core/IrisDebugging.h b/Engine/Source/Runtime/Experimental/Iris/Core/Public/Iris/Core/IrisDebugging.h
--- a/Engine/Source/Runtime/Experimental/Iris/Core/Public/Iris/Core/IrisDebugging.h
+++ b/Engine/Source/Runtime/Experimental/Iris/Core/Public/Iris/Core/IrisDebugging.h
@@ -59,7 +59,7 @@ void NetObjectProtocolReferencesToString(FStringBuilderBase& StringBuilder, uint
 extern "C" IRISCORE_API UReplicationSystem* GetReplicationSystemForDebug(uint32 Id);
 
 /** DebugOutputObject state for Handle specified by NetRefHandleId and ReplicationSystemId. */
-extern "C" IRISCORE_API void DebugOutputNetObjectState(uint32 NetRefHandleId, uint32 ReplicationSystemId);
+extern "C" IRISCORE_API void DebugOutputNetObjectState(uint64 NetRefHandleId, uint32 ReplicationSystemId);
 
 /**
  * Variant of NetObjectStateToString that can be used from breakpoints and in watch window to print the current state of a NetRefHandle.
diff --git a/Engine/Source/Runtime/Experimental/Iris/Core/Public/Iris/ReplicationState/Private/IrisFastArraySerializerInternal.h b/Engine/Source/Runtime/Experimental/Iris/Core/Public/Iris/ReplicationState/Private/IrisFastArraySerializerInternal.h
--- a/Engine/Source/Runtime/Experimental/Iris/Core/Public/Iris/ReplicationState/Private/IrisFastArraySerializerInternal.h
+++ b/Engine/Source/Runtime/Experimental/Iris/Core/Public/Iris/ReplicationState/Private/IrisFastArraySerializerInternal.h
@@ -126,33 +126,42 @@ void TIrisFastArrayEditor<FastArrayType>::MarkItemDirty(FastArrayItemType& Item)
 	}
 
 	const FastArrayItemArrayType& Array = FastArray.GetItemArray();
-	const int32 Index = (&Item - Array.GetData());
-	if (Array.IsValidIndex(Index))
+	const uint64 Index64 = &Item - Array.GetData();
+	if (ensure((IntFitsIn<int32, uint64>(Index64))))
 	{
-		// If this is a new element make sure it is at the end, otherwise we must dirty array
-		if ((Item.ReplicationID == INDEX_NONE) && (Index != (Array.Num() - 1)))
+		const int32 Index = static_cast<int32>(Index64);
+		if (Array.IsValidIndex(Index))
 		{
-			// Must mark all items that might have been shifted dirty, including the new one, this is to mimic behavior of current FastArraySerializer
-			for (FastArrayItemType& ArrayItem : MakeArrayView(&Item, Array.Num() - Index))
+			// If this is a new element make sure it is at the end, otherwise we must dirty array
+			if ((Item.ReplicationID == INDEX_NONE) && (Index != (Array.Num() - 1)))
 			{
-				const bool bIsWritingOnClient = false;
-				if (!FastArray.template ShouldWriteFastArrayItem<FastArrayItemType, FastArrayType>(ArrayItem, bIsWritingOnClient))
+				// Must mark all items that might have been shifted dirty, including the new one, this is to mimic behavior of current FastArraySerializer
+				for (FastArrayItemType& ArrayItem : MakeArrayView(&Item, Array.Num() - Index))
 				{
-					if (ArrayItem.ReplicationID == INDEX_NONE)
+					const bool bIsWritingOnClient = false;
+					if (!FastArray.template ShouldWriteFastArrayItem<FastArrayItemType, FastArrayType>(ArrayItem, bIsWritingOnClient))
 					{
-						FastArray.MarkItemDirty(ArrayItem);						
+						if (ArrayItem.ReplicationID == INDEX_NONE)
+						{
+							FastArray.MarkItemDirty(ArrayItem);
+						}
+
+						const uint64 ItemIndex64 = &ArrayItem - Array.GetData();
+						if (ensure((IntFitsIn<int32, uint64>(ItemIndex64))))
+						{
+							// Mark item dirty in changemask
+							FIrisFastArraySerializerPrivateAccessor::MarkArrayItemDirty(FastArray, static_cast<int32>(ItemIndex64));
+						}
 					}
-					// Mark item dirty in changemask
-					FIrisFastArraySerializerPrivateAccessor::MarkArrayItemDirty(FastArray, (&ArrayItem - Array.GetData()));
 				}
 			}
-		}
-		else
-		{
-			// This is in order to execute old logic
-			FastArray.MarkItemDirty(Item);
-			// Mark item dirty in changemask
-			FIrisFastArraySerializerPrivateAccessor::MarkArrayItemDirty(FastArray, Index);
+			else
+			{
+				// This is in order to execute old logic
+				FastArray.MarkItemDirty(Item);
+				// Mark item dirty in changemask
+				FIrisFastArraySerializerPrivateAccessor::MarkArrayItemDirty(FastArray, Index);
+			}
 		}
 	}
 };
diff --git a/Engine/Source/Runtime/Experimental/Iris/Core/Public/Iris/ReplicationState/ReplicationStateDescriptorImplementationMacros.h b/Engine/Source/Runtime/Experimental/Iris/Core/Public/Iris/ReplicationState/ReplicationStateDescriptorImplementationMacros.h
--- a/Engine/Source/Runtime/Experimental/Iris/Core/Public/Iris/ReplicationState/ReplicationStateDescriptorImplementationMacros.h
+++ b/Engine/Source/Runtime/Experimental/Iris/Core/Public/Iris/ReplicationState/ReplicationStateDescriptorImplementationMacros.h
@@ -51,7 +51,7 @@ constexpr uint16 GetInternalStateAlignment(const T& MemberDescriptors, const FIn
 
 	for (uint32 It = 0; It < Count; ++It)
 	{
-		Alignment = FPlatformMath::Max<uint16>(Alignment, InternalTypeInfo[It].Alignment);
+		Alignment = FPlatformMath::Max<uint16>(Alignment, static_cast<uint16>(InternalTypeInfo[It].Alignment));
 	}
 
 	return Alignment;
