Iris - Add concept of Ordered RPCs. Unicast RPCs are now ordered with respect to reliable RPCs and other unicast RPCs targeting the same actor. From client to server all reliable and unicast RPCs are ordered with respect to eachother, regardless of target actor.

diff --git a/Engine/Plugins/Runtime/ReplicationSystemTestPlugin/Source/Private/Tests/ReplicationSystem/RPC/ReplicatedTestObjectWithRPC.cpp b/Engine/Plugins/Runtime/ReplicationSystemTestPlugin/Source/Private/Tests/ReplicationSystem/RPC/ReplicatedTestObjectWithRPC.cpp
--- a/Engine/Plugins/Runtime/ReplicationSystemTestPlugin/Source/Private/Tests/ReplicationSystem/RPC/ReplicatedTestObjectWithRPC.cpp
+++ b/Engine/Plugins/Runtime/ReplicationSystemTestPlugin/Source/Private/Tests/ReplicationSystem/RPC/ReplicatedTestObjectWithRPC.cpp
@@ -109,21 +109,37 @@ bool UTestReplicatedObjectWithRPC::CallRemoteFunction(UFunction* Function, void*
 void UTestReplicatedObjectWithRPC::ClientRPC_Implementation()
 {
 	bClientRPCCalled = true;
+	ClientRPCCallOrder = ++CallOrder;
 }
 
 void UTestReplicatedObjectWithRPC::ClientRPCWithParam_Implementation(int32 IntParam)
 {
 	ClientRPCWithParamCalled = IntParam;
+	ClientRPCWithParamCallOrder = ++CallOrder;
+}
+
+void UTestReplicatedObjectWithRPC::ClientUnreliableRPC_Implementation()
+{
+	bClientUnreliableRPCCalled = true;
+	ClientUnreliableRPCCallOrder = ++CallOrder;
 }
 
 void UTestReplicatedObjectWithRPC::ServerRPC_Implementation()
 {
 	bServerRPCCalled = true;
+	ServerRPCCallOrder = ++CallOrder;
 }
 
 void UTestReplicatedObjectWithRPC::ServerRPCWithParam_Implementation(int32 IntParam)
 {
 	ServerRPCWithParamCalled = IntParam;
+	ServerRPCWithParamCallOrder = ++CallOrder;
+}
+
+void UTestReplicatedObjectWithRPC::ServerUnreliableRPC_Implementation()
+{
+	bServerUnreliableRPCCalled = true;
+	ServerUnreliableRPCCallOrder = ++CallOrder;
 }
 
 void UTestReplicatedObjectWithRPC::NetMulticast_MultiCastRPCSendImmediate_Implementation()
diff --git a/Engine/Plugins/Runtime/ReplicationSystemTestPlugin/Source/Private/Tests/ReplicationSystem/RPC/ReplicatedTestObjectWithRPC.h b/Engine/Plugins/Runtime/ReplicationSystemTestPlugin/Source/Private/Tests/ReplicationSystem/RPC/ReplicatedTestObjectWithRPC.h
--- a/Engine/Plugins/Runtime/ReplicationSystemTestPlugin/Source/Private/Tests/ReplicationSystem/RPC/ReplicatedTestObjectWithRPC.h
+++ b/Engine/Plugins/Runtime/ReplicationSystemTestPlugin/Source/Private/Tests/ReplicationSystem/RPC/ReplicatedTestObjectWithRPC.h
@@ -40,18 +40,32 @@ public:
 	UFUNCTION(Reliable, Client)
 	void ClientRPC();
 	bool bClientRPCCalled = false;
+	int32 ClientRPCCallOrder = 0;
 
 	UFUNCTION(Reliable, Client)
 	void ClientRPCWithParam(int32 IntParam);
 	int32 ClientRPCWithParamCalled = 0;
+	int32 ClientRPCWithParamCallOrder = 0;
+
+	UFUNCTION(Unreliable, Client)
+	void ClientUnreliableRPC();
+	bool bClientUnreliableRPCCalled = false;
+	int32 ClientUnreliableRPCCallOrder = 0;
 
 	UFUNCTION(Reliable, Server)
 	void ServerRPC();
 	bool bServerRPCCalled = false;
+	int32 ServerRPCCallOrder = 0;
 
 	UFUNCTION(Reliable, Server)
 	void ServerRPCWithParam(int32 IntParam);
 	int32 ServerRPCWithParamCalled = 0;
+	int32 ServerRPCWithParamCallOrder = 0;
+
+	UFUNCTION(Unreliable, Server)
+	void ServerUnreliableRPC();
+	bool bServerUnreliableRPCCalled = false;
+	int32 ServerUnreliableRPCCallOrder = 0;
 
 	UFUNCTION(NetMulticast, unreliable)
 	void NetMulticast_MultiCastRPCSendImmediate();
@@ -62,6 +76,5 @@ public:
 	int32 NetMulticast_MultiCastRPCCallOrder = 0;
 
 	int32 CallOrder = 0;
-
 };
 
diff --git a/Engine/Plugins/Runtime/ReplicationSystemTestPlugin/Source/Private/Tests/ReplicationSystem/RPC/TestObjectRPC.cpp b/Engine/Plugins/Runtime/ReplicationSystemTestPlugin/Source/Private/Tests/ReplicationSystem/RPC/TestObjectRPC.cpp
--- a/Engine/Plugins/Runtime/ReplicationSystemTestPlugin/Source/Private/Tests/ReplicationSystem/RPC/TestObjectRPC.cpp
+++ b/Engine/Plugins/Runtime/ReplicationSystemTestPlugin/Source/Private/Tests/ReplicationSystem/RPC/TestObjectRPC.cpp
@@ -205,4 +205,127 @@ namespace UE::Net::Private
 		UE_NET_ASSERT_TRUE(ServerSubObject->ServerRPCWithParamCalled == IntParam);
 	}
 
+	UE_NET_TEST_FIXTURE(FRPCTestFixture, TestUnreliableRPCIsOrderedWithReliableRPCToClient)
+	{
+		// Add a client
+		FReplicationSystemTestClient* Client = CreateClient();
+
+		// Spawn object on server
+		UTestReplicatedObjectWithRPC* ServerObject = Server->CreateObject<UTestReplicatedObjectWithRPC>();
+
+		ServerObject->bIsServerObject = true;
+		ServerObject->ReplicationSystem = Server->GetReplicationSystem();
+		Server->ReplicationSystem->SetOwningNetConnection(ServerObject->NetRefHandle, Client->ConnectionIdOnServer);
+
+		// Send and deliver packet
+		Server->UpdateAndSend({Client});
+
+		UTestReplicatedObjectWithRPC* ClientObject = Cast<UTestReplicatedObjectWithRPC>(Client->GetReplicationBridge()->GetReplicatedObject(ServerObject->NetRefHandle));
+		
+		// Verify that created server handle now also exists on client
+		UE_NET_ASSERT_NE(ClientObject, nullptr);
+
+		// Call reliable + unreliable + reliable RPCs
+		ServerObject->ClientRPC();
+		ServerObject->ClientUnreliableRPC();
+		ServerObject->ClientRPCWithParam(1);
+
+		// Send and deliver packet
+		Server->UpdateAndSend({Client});
+
+		// Verify RPC reception
+		UE_NET_ASSERT_TRUE(ClientObject->bClientRPCCalled);
+		UE_NET_ASSERT_TRUE(ClientObject->bClientUnreliableRPCCalled);
+		UE_NET_ASSERT_NE(ClientObject->ClientRPCWithParamCalled, 0);
+
+		// Verify RPC call order
+		UE_NET_ASSERT_LT(ClientObject->ClientRPCCallOrder, ClientObject->ClientUnreliableRPCCallOrder);
+		UE_NET_ASSERT_LT(ClientObject->ClientUnreliableRPCCallOrder, ClientObject->ClientRPCWithParamCallOrder);
+	}
+
+	UE_NET_TEST_FIXTURE(FRPCTestFixture, TestUnreliableRPCIsOrderedWithReliableRPCToServer)
+	{
+		// Add a client
+		FReplicationSystemTestClient* Client = CreateClient();
+
+		// Spawn object on server
+		UTestReplicatedObjectWithRPC* ServerObject = Server->CreateObject<UTestReplicatedObjectWithRPC>();
+
+		ServerObject->bIsServerObject = true;
+		ServerObject->ReplicationSystem = Server->GetReplicationSystem();
+		Server->ReplicationSystem->SetOwningNetConnection(ServerObject->NetRefHandle, Client->ConnectionIdOnServer);
+
+		// Send and deliver packet
+		Server->UpdateAndSend({Client});
+
+		UTestReplicatedObjectWithRPC* ClientObject = Cast<UTestReplicatedObjectWithRPC>(Client->GetReplicationBridge()->GetReplicatedObject(ServerObject->NetRefHandle));
+		
+		// Verify that created server handle now also exists on client
+		UE_NET_ASSERT_NE(ClientObject, nullptr);
+
+		ClientObject->ReplicationSystem = Client->GetReplicationSystem();
+
+		// Call reliable + unreliable + reliable RPCs
+		ClientObject->ServerRPC();
+		ClientObject->ServerUnreliableRPC();
+		ClientObject->ServerRPCWithParam(1);
+
+		// Send and deliver packet
+		Client->UpdateAndSend(Server);
+
+		// Verify RPC reception
+		UE_NET_ASSERT_TRUE(ServerObject->bServerRPCCalled);
+		UE_NET_ASSERT_TRUE(ServerObject->bServerUnreliableRPCCalled);
+		UE_NET_ASSERT_NE(ServerObject->ServerRPCWithParamCalled, 0);
+
+		// Verify RPC call order
+		UE_NET_ASSERT_LT(ServerObject->ServerRPCCallOrder, ServerObject->ServerUnreliableRPCCallOrder);
+		UE_NET_ASSERT_LT(ServerObject->ServerUnreliableRPCCallOrder, ServerObject->ServerRPCWithParamCallOrder);
+	}
+
+	UE_NET_TEST_FIXTURE(FRPCTestFixture, TestUnreliableRPCIsNotResentAfterPacketLoss)
+	{
+		// Add a client
+		FReplicationSystemTestClient* Client = CreateClient();
+
+		// Spawn object on server
+		UTestReplicatedObjectWithRPC* ServerObject = Server->CreateObject<UTestReplicatedObjectWithRPC>();
+
+		ServerObject->bIsServerObject = true;
+		ServerObject->ReplicationSystem = Server->GetReplicationSystem();
+		Server->ReplicationSystem->SetOwningNetConnection(ServerObject->NetRefHandle, Client->ConnectionIdOnServer);
+
+		// Send and deliver packet
+		Server->UpdateAndSend({Client});
+
+		UTestReplicatedObjectWithRPC* ClientObject = Cast<UTestReplicatedObjectWithRPC>(Client->GetReplicationBridge()->GetReplicatedObject(ServerObject->NetRefHandle));
+		
+		// Verify that created server handle now also exists on client
+		UE_NET_ASSERT_NE(ClientObject, nullptr);
+
+		// Call reliable + unreliable + reliable RPCs
+		ServerObject->ClientRPC();
+		ServerObject->ClientUnreliableRPC();
+		ServerObject->ClientRPCWithParam(1);
+
+		// Send and do not deliver packet
+		Server->UpdateAndSend({Client}, DoNotDeliverPacket);
+
+		// Verify no RPCs were received
+		UE_NET_ASSERT_FALSE(ClientObject->bClientRPCCalled);
+		UE_NET_ASSERT_FALSE(ClientObject->bClientUnreliableRPCCalled);
+		UE_NET_ASSERT_EQ(ClientObject->ClientRPCWithParamCalled, 0);
+
+		// Send and deliver packet
+		Server->UpdateAndSend({Client}, DeliverPacket);
+
+		// Verify unreliable RPC was dropped
+		UE_NET_ASSERT_TRUE(ClientObject->bClientRPCCalled);
+		UE_NET_ASSERT_FALSE(ClientObject->bClientUnreliableRPCCalled);
+		UE_NET_ASSERT_NE(ClientObject->ClientRPCWithParamCalled, 0);
+
+		// Verify RPC call order
+		UE_NET_ASSERT_LT(ClientObject->ClientRPCCallOrder, ClientObject->ClientRPCWithParamCallOrder);
+	}
+
 }
\ No newline at end of file
diff --git a/Engine/Source/Runtime/Experimental/Iris/Core/Private/Iris/ReplicationSystem/AttachmentReplication.cpp b/Engine/Source/Runtime/Experimental/Iris/Core/Private/Iris/ReplicationSystem/AttachmentReplication.cpp
--- a/Engine/Source/Runtime/Experimental/Iris/Core/Private/Iris/ReplicationSystem/AttachmentReplication.cpp
+++ b/Engine/Source/Runtime/Experimental/Iris/Core/Private/Iris/ReplicationSystem/AttachmentReplication.cpp
@@ -31,9 +31,6 @@ namespace AttachmentReplicationCVars
 
 	static int32 MaxSimultaneousObjectsWithRPCs = 4096;
 	FAutoConsoleVariableRef CVarMaxSimultaneousObjectsWithRPCs(TEXT("net.MaxSimultaneousObjectsWithRPCs"), MaxSimultaneousObjectsWithRPCs, TEXT("Maximum number of objects that can have unsent RPCs at the same time. "));
-
-	static bool bMakeAllRPCsReliable = false;
-	FAutoConsoleVariableRef CVarAllRPCsReliable(TEXT("net.Iris.AllRPCReliable"), bMakeAllRPCsReliable, TEXT("Forces all the queued RPCs to be reliable. "));
 }
 
 static const FName NetError_UnreliableQueueFull("Unreliable attachment queue full");
@@ -190,7 +187,7 @@ FNetObjectAttachmentSendQueue::~FNetObjectAttachmentSendQueue()
 bool FNetObjectAttachmentSendQueue::Enqueue(TArrayView<const TRefCountPtr<FNetBlob>> Attachments)
 {
 	const FNetBlobCreationInfo& CreationInfo = Attachments[0]->GetCreationInfo();
-	if (EnumHasAnyFlags(CreationInfo.Flags, ENetBlobFlags::Reliable) || UE::Net::Private::AttachmentReplicationCVars::bMakeAllRPCsReliable)
+	if (EnumHasAnyFlags(CreationInfo.Flags, ENetBlobFlags::Reliable | ENetBlobFlags::Ordered))
 	{
 		if (ReliableQueue == nullptr)
 		{
diff --git a/Engine/Source/Runtime/Experimental/Iris/Core/Private/Iris/ReplicationSystem/NetBlob/NetRPC.cpp b/Engine/Source/Runtime/Experimental/Iris/Core/Private/Iris/ReplicationSystem/NetBlob/NetRPC.cpp
--- a/Engine/Source/Runtime/Experimental/Iris/Core/Private/Iris/ReplicationSystem/NetBlob/NetRPC.cpp
+++ b/Engine/Source/Runtime/Experimental/Iris/Core/Private/Iris/ReplicationSystem/NetBlob/NetRPC.cpp
@@ -339,6 +339,12 @@ bool FNetRPC::ResolveFunctionAndObject(FNetSerializationContext& Context)
 
 	Function = FunctionDescriptor->Function;
 
+	// Patch up NetBlobFlags based on function flags.
+	if (Function && ((Function->FunctionFlags & FUNC_NetReliable) != 0))
+	{
+		CreationInfo.Flags |= ENetBlobFlags::Reliable;
+	}
+
 	// Set the BlobDescriptor even if it has zero size so that we can trace with a meaningful name.
 	BlobDescriptor = FunctionDescriptor->Descriptor;
 
diff --git a/Engine/Source/Runtime/Experimental/Iris/Core/Private/Iris/ReplicationSystem/NetBlob/NetRPCHandler.cpp b/Engine/Source/Runtime/Experimental/Iris/Core/Private/Iris/ReplicationSystem/NetBlob/NetRPCHandler.cpp
--- a/Engine/Source/Runtime/Experimental/Iris/Core/Private/Iris/ReplicationSystem/NetBlob/NetRPCHandler.cpp
+++ b/Engine/Source/Runtime/Experimental/Iris/Core/Private/Iris/ReplicationSystem/NetBlob/NetRPCHandler.cpp
@@ -4,6 +4,7 @@
 #include "Iris/ReplicationSystem/NetBlob/NetRPC.h"
 #include "Iris/ReplicationSystem/ReplicationSystem.h"
 #include "Iris/ReplicationSystem/ReplicationSystemInternal.h"
+#include "HAL/IConsoleManager.h"
 
 UNetRPCHandler::UNetRPCHandler()
 {
@@ -23,6 +24,12 @@ TRefCountPtr<UE::Net::Private::FNetRPC> UNetRPCHandler::CreateRPC(const UE::Net:
 	FNetBlobCreationInfo CreationInfo;
 	CreationInfo.Type = GetNetBlobType();
 	CreationInfo.Flags = ((Function->FunctionFlags & FUNC_NetReliable) != 0) ? UE::Net::ENetBlobFlags::Reliable : UE::Net::ENetBlobFlags::None;
+	// Unicast RPCs should be ordered with respect to other reliable and unicast RPCs.
+	if ((Function->FunctionFlags & FUNC_NetMulticast) == 0)
+	{
+		CreationInfo.Flags |= UE::Net::ENetBlobFlags::Ordered;
+	}
+
 	FNetRPC* RPC = FNetRPC::Create(ReplicationSystem, CreationInfo, ObjectReference, Function, Parameters);
 	return RPC;
 }
diff --git a/Engine/Source/Runtime/Experimental/Iris/Core/Private/Iris/ReplicationSystem/NetBlob/ReliableNetBlobQueue.cpp b/Engine/Source/Runtime/Experimental/Iris/Core/Private/Iris/ReplicationSystem/NetBlob/ReliableNetBlobQueue.cpp
--- a/Engine/Source/Runtime/Experimental/Iris/Core/Private/Iris/ReplicationSystem/NetBlob/ReliableNetBlobQueue.cpp
+++ b/Engine/Source/Runtime/Experimental/Iris/Core/Private/Iris/ReplicationSystem/NetBlob/ReliableNetBlobQueue.cpp
@@ -105,7 +105,8 @@ uint32 FReliableNetBlobQueue::SerializeInternal(FNetSerializationContext& Contex
 		FNetBitStreamRollbackScope RollbackScope(*Writer);
 		FNetExportRollbackScope ExportRollbackScope(Context);
 
-		const TRefCountPtr<FNetBlob>& Attachment = NetBlobs[Index];
+		// GetRefCount() is not const.
+		TRefCountPtr<FNetBlob>& Attachment = NetBlobs[Index];
 
 		// If this sequence is disjoint from the previous sequence we need to serialize the full index.
 		// It's important that the sequence number is sent first so we can validate it before receiving exports and payload.
@@ -114,17 +115,22 @@ uint32 FReliableNetBlobQueue::SerializeInternal(FNetSerializationContext& Contex
 			Writer->WriteBits(Index, IndexBitCount);
 		}
 
-		// If we have exports, append them, if attachment is rolled back we will roll back any appended exports as well.
-		ObjectReferenceCache->AddPendingExports(Context, Attachment->CallGetExports());
-
-		Attachment->SerializeCreationInfo(Context, Attachment->GetCreationInfo());
-		if (bSerializeWithObject)
+		// Unreliable blobs may have been released.
+		const bool bHasData = Attachment.GetRefCount() > 0;
+		if (Writer->WriteBool(bHasData))
 		{
-			Attachment->SerializeWithObject(Context, RefHandle);
-		}
-		else
-		{
-			Attachment->Serialize(Context);
+			// If we have exports, append them, if attachment is rolled back we will roll back any appended exports as well.
+			ObjectReferenceCache->AddPendingExports(Context, Attachment->CallGetExports());
+
+			Attachment->SerializeCreationInfo(Context, Attachment->GetCreationInfo());
+			if (bSerializeWithObject)
+			{
+				Attachment->SerializeWithObject(Context, RefHandle);
+			}
+			else
+			{
+				Attachment->Serialize(Context);
+			}
 		}
 
 		const uint32 HasMoreBlobsWritePos = Writer->GetPosBits();
@@ -181,7 +187,18 @@ void FReliableNetBlobQueue::CommitReplicationRecord(const FReliableNetBlobQueue:
 		UnsentBlobCount -= Count;
 		for (uint32 Seq = Record.Sequences[Index], EndSeq = Seq + Count; Seq != EndSeq; ++Seq)
 		{
-			SetSequenceIsSent(Seq);
+			const uint32 BlobIndex = SequenceToIndex(Seq);
+			SetIndexIsSent(BlobIndex);
+
+			// Release unreliable blobs. They should not be resent.
+			TRefCountPtr<FNetBlob>& RefCountBlob = NetBlobs[BlobIndex];
+			if (const FNetBlob* Blob = RefCountBlob.GetReference())
+			{
+				if (!EnumHasAnyFlags(Blob->GetCreationInfo().Flags, ENetBlobFlags::Reliable))
+				{
+					RefCountBlob.SafeRelease();
+				}
+			}
 		}
 	}
 }
@@ -233,25 +250,28 @@ uint32 FReliableNetBlobQueue::DeserializeInternal(FNetSerializationContext& Cont
 			return DeserializedCount;
 		}
 
-		FNetBlobCreationInfo CreationInfo;
-		FNetBlob::DeserializeCreationInfo(Context, CreationInfo);
-		CreationInfo.Flags = CreationInfo.Flags | ENetBlobFlags::Reliable;
-		const TRefCountPtr<FNetBlob>& Blob = BlobReceiver->CreateNetBlob(CreationInfo);
-		if (!Blob.IsValid())
+		TRefCountPtr<FNetBlob> Blob;
+		if (const bool bHasData = Reader->ReadBool())
 		{
-			UE_LOG(LogIris, Warning, TEXT("%hs"), "Unable to create blob.");
-			Context.SetError(GNetError_UnsupportedNetBlob);
+			FNetBlobCreationInfo CreationInfo;
+			FNetBlob::DeserializeCreationInfo(Context, CreationInfo);
+			Blob = BlobReceiver->CreateNetBlob(CreationInfo);
+			if (!Blob.IsValid())
+			{
+				UE_LOG(LogIris, Warning, TEXT("%hs"), "Unable to create blob.");
+				Context.SetError(GNetError_UnsupportedNetBlob);
 
-			return DeserializedCount;
-		}
+				return DeserializedCount;
+			}
 
-		if (bSerializeWithObject)
-		{
-			Blob->DeserializeWithObject(Context, RefHandle);
-		}
-		else
-		{
-			Blob->Deserialize(Context);
+			if (bSerializeWithObject)
+			{
+				Blob->DeserializeWithObject(Context, RefHandle);
+			}
+			else
+			{
+				Blob->Deserialize(Context);
+			}
 		}
 		
 		bHasMoreBlobs = Reader->ReadBool();
@@ -263,7 +283,7 @@ uint32 FReliableNetBlobQueue::DeserializeInternal(FNetSerializationContext& Cont
 			break;
 		}
 
-		NetBlobs[Index] = Blob;
+		NetBlobs[Index] = MoveTemp(Blob);
 		LastSeq = FMath::Max(LastSeq, ReceivedSeq + 1U);
 		SetIndexIsAcked(Index);
 
@@ -289,12 +309,23 @@ bool FReliableNetBlobQueue::Enqueue(const TRefCountPtr<FNetBlob>& Blob)
 	return true;
 }
 
-const TRefCountPtr<FNetBlob>* FReliableNetBlobQueue::Peek() const
+const TRefCountPtr<FNetBlob>* FReliableNetBlobQueue::Peek()
 {
-	const uint32 Index = SequenceToIndex(FirstSeq);
-	if (IsIndexAcked(Index))
+	for (; FirstSeq < LastSeq; ++FirstSeq)
 	{
-		return &NetBlobs[Index];
+		const uint32 Index = SequenceToIndex(FirstSeq);
+		if (!IsIndexAcked(Index))
+		{
+			return nullptr;
+		}
+
+		if (NetBlobs[Index].GetRefCount() > 0)
+		{
+			return &NetBlobs[Index];
+		}
+
+		// We skip over empty blobs and ack them.
+		ClearIndexIsAcked(Index);
 	}
 
 	return nullptr;
diff --git a/Engine/Source/Runtime/Experimental/Iris/Core/Private/Iris/ReplicationSystem/ReplicationReader.cpp b/Engine/Source/Runtime/Experimental/Iris/Core/Private/Iris/ReplicationSystem/ReplicationReader.cpp
--- a/Engine/Source/Runtime/Experimental/Iris/Core/Private/Iris/ReplicationSystem/ReplicationReader.cpp
+++ b/Engine/Source/Runtime/Experimental/Iris/Core/Private/Iris/ReplicationSystem/ReplicationReader.cpp
@@ -2114,8 +2114,9 @@ void FReplicationReader::ResolveAndDispatchAttachments(FNetSerializationContext&
 		{
 			while (const TRefCountPtr<FNetBlob>* Attachment = AttachmentQueue->PeekReliable())
 			{
-				// Delay attachments with unresolved pending references
-				if (bCanDelayAttachments)
+				// Delay reliable attachments with unresolved pending references
+				const bool bIsReliable = EnumHasAnyFlags(Attachment->GetReference()->GetCreationInfo().Flags, ENetBlobFlags::Reliable);
+				if (bIsReliable && bCanDelayAttachments)
 				{
 					bool bDelayRpc = false;
 
diff --git a/Engine/Source/Runtime/Experimental/Iris/Core/Public/Iris/ReplicationSystem/NetBlob/NetBlob.h b/Engine/Source/Runtime/Experimental/Iris/Core/Public/Iris/ReplicationSystem/NetBlob/NetBlob.h
--- a/Engine/Source/Runtime/Experimental/Iris/Core/Public/Iris/ReplicationSystem/NetBlob/NetBlob.h
+++ b/Engine/Source/Runtime/Experimental/Iris/Core/Public/Iris/ReplicationSystem/NetBlob/NetBlob.h
@@ -29,7 +29,7 @@ enum class ENetBlobFlags : uint32
 {
 	None = 0,
 
-	/** The blob should be delivered reliably in order with respect to other reliable blobs. */
+	/** The blob should be delivered reliably in order with respect to other reliable blobs. Implies Ordered. */
 	Reliable = 1U << 0U,
 
 	/** Used for FRawDataNetBlob derived classes to avoid duplicate serialization when splitting large blob. */
@@ -37,6 +37,9 @@ enum class ENetBlobFlags : uint32
 
 	/** Used to indicate that this blob have ObjectReferences that might have to be exported. */
 	HasExports = RawDataNetBlob << 1U,
+
+	/** The blob should respect delivery order with respect to other Ordered blobs, including Reliable ones. Unreliable ordered blobs will only be sent once. */
+	Ordered = HasExports << 1U,
 };
 ENUM_CLASS_FLAGS(ENetBlobFlags);
 
diff --git a/Engine/Source/Runtime/Experimental/Iris/Core/Public/Iris/ReplicationSystem/NetBlob/ReliableNetBlobQueue.h b/Engine/Source/Runtime/Experimental/Iris/Core/Public/Iris/ReplicationSystem/NetBlob/ReliableNetBlobQueue.h
--- a/Engine/Source/Runtime/Experimental/Iris/Core/Public/Iris/ReplicationSystem/NetBlob/ReliableNetBlobQueue.h
+++ b/Engine/Source/Runtime/Experimental/Iris/Core/Public/Iris/ReplicationSystem/NetBlob/ReliableNetBlobQueue.h
@@ -81,7 +81,7 @@ public:
 	IRISCORE_API bool Enqueue(const TRefCountPtr<FNetBlob>& Blob);
 
 	/** On the receiving end this will return a pointer to the next blob that can be processed. */
-	IRISCORE_API const TRefCountPtr<FNetBlob>* Peek() const;
+	IRISCORE_API const TRefCountPtr<FNetBlob>* Peek();
 
 	/** On the receiving end this will remove the next blob to be processed from the queue. Call after processing the blob returned from Peek(). */
 	IRISCORE_API void Pop();
