Iris - Added dependent actor test with dynamic filtering

diff --git a/Engine/Plugins/Runtime/ReplicationSystemTestPlugin/Source/Private/Tests/ReplicationSystem/TestDependentObjects.cpp b/Engine/Plugins/Runtime/ReplicationSystemTestPlugin/Source/Private/Tests/ReplicationSystem/TestDependentObjects.cpp
--- a/Engine/Plugins/Runtime/ReplicationSystemTestPlugin/Source/Private/Tests/ReplicationSystem/TestDependentObjects.cpp
+++ b/Engine/Plugins/Runtime/ReplicationSystemTestPlugin/Source/Private/Tests/ReplicationSystem/TestDependentObjects.cpp
@@ -1,15 +1,96 @@
 // Copyright Epic Games, Inc. All Rights Reserved.
 
 #include "ReplicationSystemServerClientTestFixture.h"
-#include "Iris/Serialization/NetBitStreamReader.h"
-#include "Iris/Serialization/NetBitStreamWriter.h"
 #include "Iris/ReplicationSystem/ReplicationSystem.h"
+#include "Iris/ReplicationSystem/Filtering/NetObjectFilterDefinitions.h"
+#include "Tests/ReplicationSystem/Filtering/MockNetObjectFilter.h"
 
-#include "Iris/ReplicationSystem/ReplicationSystemInternal.h"
-#include "Iris/ReplicationSystem/NetTokenStore.h"
 
 namespace UE::Net::Private
 {
+	
+class FTestFilteredDependentObjectFixture : public FReplicationSystemServerClientTestFixture
+{
+protected:
+	virtual void SetUp() override
+	{
+		InitFilterDefinitions();
+		FReplicationSystemServerClientTestFixture::SetUp();
+		InitFilterHandles();
+	}
+
+	virtual void TearDown() override
+	{
+		FReplicationSystemServerClientTestFixture::TearDown();
+		RestoreFilterDefinitions();
+	}
+
+	void SetMockFilterStatus(UE::Net::ENetFilterStatus FilterStatus)
+	{
+		UMockNetObjectFilter::FFunctionCallSetup CallSetup;
+		CallSetup.AddObject.bReturnValue = true;
+		CallSetup.Filter.bFilterOutByDefault = FilterStatus == UE::Net::ENetFilterStatus::Disallow;
+		MockFilter->SetFunctionCallSetup(CallSetup);
+	}
+
+	FNetObjectFilterHandle NotRoutedFilterHandle = InvalidNetObjectFilterHandle;
+	FNetObjectFilterHandle MockFilterHandle = InvalidNetObjectFilterHandle;
+	TObjectPtr<UMockNetObjectFilter> MockFilter = nullptr;
+
+private:
+	void InitFilterDefinitions()
+	{
+		const UClass* NetObjectFilterDefinitionsClass = UNetObjectFilterDefinitions::StaticClass();
+		const FProperty* DefinitionsProperty = NetObjectFilterDefinitionsClass->FindPropertyByName("NetObjectFilterDefinitions");
+		check(DefinitionsProperty != nullptr);
+
+		// Save CDO state.
+		UNetObjectFilterDefinitions* FilterDefinitions = GetMutableDefault<UNetObjectFilterDefinitions>();
+		DefinitionsProperty->CopyCompleteValue(&OriginalFilterDefinitions, (void*)(UPTRINT(FilterDefinitions) + DefinitionsProperty->GetOffset_ForInternal()));
+
+		// Modify definitions to only include our filters. 
+		TArray<FNetObjectFilterDefinition> NewFilterDefinitions;
+		{
+			FNetObjectFilterDefinition& NotRoutedDefinition = NewFilterDefinitions.Emplace_GetRef();
+			NotRoutedDefinition.FilterName = "NotRouted";
+			NotRoutedDefinition.ClassName = "/Script/IrisCore.FilterOutNetObjectFilter";
+			NotRoutedDefinition.ConfigClassName = "/Script/IrisCore.NetObjectGridFilterConfig";
+		}
+
+		{
+			FNetObjectFilterDefinition& MockDefinition = NewFilterDefinitions.Emplace_GetRef();
+			MockDefinition.FilterName = "Mock";
+			MockDefinition.ClassName = "/Script/ReplicationSystemTestPlugin.MockNetObjectFilter";
+			MockDefinition.ConfigClassName = "/Script/ReplicationSystemTestPlugin.MockNetObjectFilterConfig";
+		}
+
+		DefinitionsProperty->CopyCompleteValue((void*)(UPTRINT(FilterDefinitions) + DefinitionsProperty->GetOffset_ForInternal()), &NewFilterDefinitions);
+	}
+
+	void RestoreFilterDefinitions()
+	{
+		// Restore CDO state from the saved state.
+		const UClass* NetObjectFilterDefinitionsClass = UNetObjectFilterDefinitions::StaticClass();
+		const FProperty* DefinitionsProperty = NetObjectFilterDefinitionsClass->FindPropertyByName("NetObjectFilterDefinitions");
+		UNetObjectFilterDefinitions* FilterDefinitions = GetMutableDefault<UNetObjectFilterDefinitions>();
+		DefinitionsProperty->CopyCompleteValue((void*)(UPTRINT(FilterDefinitions) + DefinitionsProperty->GetOffset_ForInternal()), &OriginalFilterDefinitions);
+		OriginalFilterDefinitions.Empty();
+
+		NotRoutedFilterHandle = InvalidNetObjectFilterHandle;
+		MockFilterHandle = InvalidNetObjectFilterHandle;
+		MockFilter = nullptr;
+	}
+
+	void InitFilterHandles()
+	{
+		NotRoutedFilterHandle = Server->GetReplicationSystem()->GetFilterHandle("NotRouted");
+		MockFilterHandle = Server->GetReplicationSystem()->GetFilterHandle("Mock");
+		MockFilter = Cast<UMockNetObjectFilter>(Server->GetReplicationSystem()->GetFilter("Mock"));
+	}
+
+private:
+	TArray<FNetObjectFilterDefinition> OriginalFilterDefinitions;
+};
 
 UE_NET_TEST_FIXTURE(FReplicationSystemServerClientTestFixture, TestDependentObjectDroppedDataIsRetransmitted)
 {
@@ -758,5 +839,58 @@ UE_NET_TEST_FIXTURE(FReplicationSystemServerClientTestFixture, TestLateAddedNest
 	UReplicatedSubObjectOrderObject* ClientSharedDependentObject = Cast<UReplicatedSubObjectOrderObject>(Client->GetReplicationBridge()->GetReplicatedObject(ServerSharedDependentObject->NetRefHandle));
 }
 
+UE_NET_TEST_FIXTURE(FTestFilteredDependentObjectFixture, TestDependentObjectIsFilteredOutTogetherWithParent)
+{
+	UReplicationSystem* ReplicationSystem = Server->ReplicationSystem;
+	UReplicatedTestObjectBridge* Bridge = Server->GetReplicationBridge();
+
+	// Add a client
+	FReplicationSystemTestClient* Client = CreateClient();
+
+	// Spawn objects on server
+	UTestReplicatedIrisObject* ServerObject = Server->CreateObject(UTestReplicatedIrisObject::FComponents{});
+	UTestReplicatedIrisObject* ServerDependentObject = Server->CreateObject(UTestReplicatedIrisObject::FComponents{});
+	Bridge->AddDependentObject(ServerObject->NetRefHandle, ServerDependentObject->NetRefHandle);
+
+	// Dynamically filter out dependent object
+	ReplicationSystem->SetFilter(ServerDependentObject->NetRefHandle, NotRoutedFilterHandle);
+
+	// Add parent object to MockFilter so we can filter in/out as desired.
+	SetMockFilterStatus(UE::Net::ENetFilterStatus::Allow);
+	ReplicationSystem->SetFilter(ServerObject->NetRefHandle, MockFilterHandle);
+
+	// Send data
+	Server->UpdateAndSend({Client});
+
+	// Verify that objects have replicated
+	UTestReplicatedIrisObject* ClientObject = Cast<UTestReplicatedIrisObject>(Client->GetReplicationBridge()->GetReplicatedObject(ServerObject->NetRefHandle));
+	UE_NET_ASSERT_NE(ClientObject, nullptr);
+	UTestReplicatedIrisObject* ClientDependentObject = Cast<UTestReplicatedIrisObject>(Client->GetReplicationBridge()->GetReplicatedObject(ServerDependentObject->NetRefHandle));
+	UE_NET_ASSERT_NE(ClientDependentObject, nullptr);
+
+	// Filter out parent object. Dependent object should be filtered out as well due to being filtered out by default.
+	SetMockFilterStatus(UE::Net::ENetFilterStatus::Disallow);
+
+	// Send data
+	Server->UpdateAndSend({Client});
+
+	// Both parent and dependent should now be filtered out.
+	ClientObject = Cast<UTestReplicatedIrisObject>(Client->GetReplicationBridge()->GetReplicatedObject(ServerObject->NetRefHandle));
+	UE_NET_ASSERT_EQ(ClientObject, nullptr);
+	ClientDependentObject = Cast<UTestReplicatedIrisObject>(Client->GetReplicationBridge()->GetReplicatedObject(ServerDependentObject->NetRefHandle));
+	UE_NET_ASSERT_EQ(ClientDependentObject, nullptr);
+
+	// Restore filtering such that both objects are expected to be replicated again.
+	SetMockFilterStatus(UE::Net::ENetFilterStatus::Allow);
+
+	// Send data
+	Server->UpdateAndSend({Client});
+
+	ClientObject = Cast<UTestReplicatedIrisObject>(Client->GetReplicationBridge()->GetReplicatedObject(ServerObject->NetRefHandle));
+	UE_NET_ASSERT_NE(ClientObject, nullptr);
+	ClientDependentObject = Cast<UTestReplicatedIrisObject>(Client->GetReplicationBridge()->GetReplicatedObject(ServerDependentObject->NetRefHandle));
+	UE_NET_ASSERT_NE(ClientDependentObject, nullptr);
+}
+
 
 }
