Network
 - Improved handling of ack history overruns which will no longer cause false negatives, if we overshoot ackhistorybuffer client will now  flush the sequence history before accepting new data packets.

Added FNetPacketNotify::IsSequenceWindowFull. Make DataStreamChannel use the new functionality since some DataStreams require all sent packets to be correctly ACKed or NAKed. DataStreamChannel can spare a few packets, controlled via cvar net.Iris.PacketSequenceSafetyMargin which currently defaults to 4, in order to allow other channels to send data without breaking DataStreamChannel.

Moved NetPacketNotifyTests into TestNetPacketNotify.

diff --git a/Engine/Plugins/Runtime/ReplicationSystemTestPlugin/Source/Private/Tests/Core/TestNetPacketNotify.cpp b/Engine/Plugins/Runtime/ReplicationSystemTestPlugin/Source/Private/Tests/Core/TestNetPacketNotify.cpp
new file mode 100644
--- /dev/null
+++ b/Engine/Plugins/Runtime/ReplicationSystemTestPlugin/Source/Private/Tests/Core/TestNetPacketNotify.cpp
@@ -0,0 +1,640 @@
+// Copyright Epic Games, Inc. All Rights Reserved.
+
+#include "NetworkAutomationTest.h"
+#include "NetworkAutomationTestMacros.h"
+#include "Net/NetPacketNotify.h"
+#include "Logging/LogScopedVerbosityOverride.h"
+#include "Net/Util/SequenceHistory.h"
+#include "Net/Util/SequenceNumber.h"
+
+namespace UE::Net
+{
+
+FTestMessage& operator<<(FTestMessage& Message, const FNetPacketNotify::SequenceNumberT& SequenceNumber)
+{
+	return Message << SequenceNumber.Get();
+}
+
+FTestMessage& operator<<(FTestMessage& Message, const FNetPacketNotify::SequenceHistoryT& History)
+{
+	for (uint32 WordIt=0U; WordIt<FNetPacketNotify::SequenceHistoryT::WordCount; ++WordIt)
+	{
+		Message << History.Data()[WordIt];
+	}
+	return Message;
+}
+
+namespace Private
+{
+
+struct FNetPacketNotifyTestUtil : public FNetworkAutomationTestSuiteFixture
+{
+	enum
+	{
+		LastValidSequenceHistoryIndex = FNetPacketNotify::MaxSequenceHistoryLength - 1
+	};
+
+	FNetPacketNotify DefaultNotify;
+	FNetPacketNotifyTestUtil()
+	{
+		DefaultNotify.Init(FNetPacketNotify::SequenceNumberT(-1),  FNetPacketNotify::SequenceNumberT(0));
+	}
+
+	// Helper to fill in SequenceHistory with expected result
+	template <typename T>
+	static void InitHistory(FNetPacketNotify::SequenceHistoryT& History, const T& DataToSet)
+	{
+		const SIZE_T Count = sizeof(T) / sizeof(DataToSet[0]);
+		static_assert(Count < FNetPacketNotify::SequenceHistoryT::WordCount, "DataToSet must be smaller than HistoryBuffer");
+
+		for (SIZE_T It=0; It < Count; ++It)
+		{
+			History.Data()[It] = DataToSet[It];
+		}
+	}
+
+	// Pretend to receive and acknowledge incoming packet to generate ackdata
+	static int32 PretendReceiveSeq(FNetPacketNotify& PacketNotify, FNetPacketNotify::SequenceNumberT Seq, bool Ack = true)
+	{
+		FNetPacketNotify::FNotificationHeader Data;
+		Data.Seq = Seq;
+		Data.AckedSeq = PacketNotify.GetOutAckSeq();
+		Data.History = FNetPacketNotify::SequenceHistoryT(0);
+		Data.HistoryWordCount = 1;
+		
+		FNetPacketNotify::SequenceNumberT::DifferenceT SeqDelta = PacketNotify.Update(Data, [](FNetPacketNotify::SequenceNumberT AckedSequence, bool delivered) {});		
+		if (SeqDelta > 0 && !PacketNotify.IsWaitingForSequenceHistoryFlush())
+		{
+			if (Ack)
+			{
+				PacketNotify.AckSeq(Seq);
+			}
+		}
+
+		return SeqDelta;
+	}
+
+	// Pretend to send packet
+	static void PretendSendSeq(FNetPacketNotify& PacketNotify, FNetPacketNotify::SequenceNumberT LastAckSeq)
+	{
+		// set last InAcqSeq that we know that the remote end knows that we know (AckAck)
+		PacketNotify.WrittenHistoryWordCount = 1;
+		PacketNotify.WrittenInAckSeq = LastAckSeq;
+
+		// Store data
+		PacketNotify.CommitAndIncrementOutSeq();
+	}
+
+	// pretend to ack array of sequence numbers
+	template<typename T>
+	static void PretendAckSequenceNumbers(FNetPacketNotify& PacketNotify, const T& InSequenceNumbers)
+	{
+		SIZE_T SequenceNumberCount = sizeof(InSequenceNumbers) / sizeof(InSequenceNumbers[0]);
+
+		for (SIZE_T I=0; I<SequenceNumberCount; ++I)
+		{
+			FNetPacketNotifyTestUtil::PretendReceiveSeq(PacketNotify, InSequenceNumbers[I]);
+		}
+	}
+	
+	// Pretend that we received a packet
+	template<typename T>
+	static SIZE_T PretendReceivedPacket(FNetPacketNotify& PacketNotify, const FNetPacketNotify::FNotificationHeader Data, T& OutSequenceNumbers)
+	{
+		SIZE_T NotificationCount = 0;
+
+		auto HandleAck = [&OutSequenceNumbers, &NotificationCount](FNetPacketNotify::SequenceNumberT Seq, bool delivered)
+		{
+			const SIZE_T MaxSequenceNumberCount = sizeof(OutSequenceNumbers) / sizeof(OutSequenceNumbers[0]);
+
+			if (delivered)
+			{
+				if (NotificationCount < MaxSequenceNumberCount)
+				{
+					OutSequenceNumbers[NotificationCount] = Seq;
+				}
+				++NotificationCount;		
+			}
+		};
+		return PacketNotify.Update(Data, HandleAck);
+	}
+
+	struct FTestNode
+	{
+		FNetPacketNotify PacketNotify;
+		TArray<FNetPacketNotify::SequenceNumberT> AcceptedInPackets;
+		TArray<FNetPacketNotify::SequenceNumberT> AcknowledgedPackets;
+
+		FTestNode(FNetPacketNotify::SequenceNumberT FirstSequence = 0)
+		{
+			PacketNotify.Init(FNetPacketNotify::SequenceNumberT(FirstSequence.Get() - 1), FirstSequence);
+		}
+
+		FTestNode(FNetPacketNotify::SequenceNumberT InSequnce, FNetPacketNotify::SequenceNumberT OutSequnce)
+		{
+			PacketNotify.Init(InSequnce, OutSequnce);
+		}
+
+		FNetPacketNotify::FNotificationHeader Send()
+		{
+			const SIZE_T CurrentHistoryWordCount = FMath::Clamp<SIZE_T>((PacketNotify.GetCurrentSequenceHistoryLength() + FNetPacketNotify::SequenceHistoryT::BitsPerWord - 1u) / FNetPacketNotify::SequenceHistoryT::BitsPerWord, 1u, FNetPacketNotify::SequenceHistoryT::WordCount);
+					
+			FNetPacketNotify::FNotificationHeader Data;
+			Data.Seq = PacketNotify.OutSeq;
+			Data.AckedSeq = PacketNotify.InAckSeq;
+			Data.HistoryWordCount = CurrentHistoryWordCount;
+			Data.History = PacketNotify.InSeqHistory;
+
+			PacketNotify.WrittenHistoryWordCount = CurrentHistoryWordCount;
+			PacketNotify.WrittenInAckSeq = Data.AckedSeq;
+			PacketNotify.CommitAndIncrementOutSeq();
+
+			return Data;
+		}
+
+		void Receive(const FNetPacketNotify::FNotificationHeader Data)
+		{
+			SIZE_T NotificationCount = 0;
+
+			auto HandleAck = [&](FNetPacketNotify::SequenceNumberT Seq, bool bDelivered)
+			{
+				if (bDelivered)
+				{
+					AcknowledgedPackets.Add(Seq);
+				}
+			};
+
+			const uint32 SequenceDelta = PacketNotify.Update(Data, HandleAck);
+			if (SequenceDelta > 0U && !PacketNotify.IsWaitingForSequenceHistoryFlush())
+			{
+				PacketNotify.AckSeq(PacketNotify.GetInSeq());
+				AcceptedInPackets.Add(PacketNotify.GetInSeq());
+			}
+		}
+
+		void SendAndDeliverTo(FTestNode& Target)
+		{
+			Target.Receive(Send());
+		}
+
+		void SendAndDeliverTo(FTestNode& Target, TArrayView<const FNetPacketNotify::SequenceNumberT> SequenceNumbers)
+		{
+			for (FNetPacketNotify::SequenceNumberT Seq : SequenceNumbers)
+			{
+				// Skip missing seqs
+				while (Seq != PacketNotify.GetOutSeq())
+				{
+					Send();
+				}
+				// Send and receive seq
+				Target.Receive(Send());
+			}
+		}
+
+		bool ValidateAcknowledgedPackets(TArrayView<const FNetPacketNotify::SequenceNumberT> ExpectedSequenceNumbers) const
+		{
+			return AcknowledgedPackets == ExpectedSequenceNumbers;
+		}
+
+		bool ValidateAcceptedPackets(TArrayView<const FNetPacketNotify::SequenceNumberT> ExpectedSequenceNumbers) const
+		{
+			return AcceptedInPackets == ExpectedSequenceNumbers;
+		}
+	};
+};
+
+UE_NET_TEST_FIXTURE(FNetPacketNotifyTestUtil, Fill)
+{
+	FNetPacketNotify Acks = DefaultNotify;
+
+	FNetPacketNotify::SequenceNumberT ExpectedInSeq(30);
+	FNetPacketNotify::SequenceHistoryT ExpectedInSeqHistory(0x55555555u, 1);
+				
+	for (FNetPacketNotify::SequenceNumberT::SequenceT I = 0U; I < 16U; ++I)
+	{
+		FNetPacketNotifyTestUtil::PretendReceiveSeq(Acks, I * 2U);
+	}
+		
+	// Verify InSeq
+	UE_NET_ASSERT_EQ(Acks.GetInSeq(), ExpectedInSeq);
+
+	// Verify History
+	UE_NET_ASSERT_EQ(Acks.GetInSeqHistory(), ExpectedInSeqHistory);
+}
+
+UE_NET_TEST_FIXTURE(FNetPacketNotifyTestUtil, DropEveryOther)
+{
+	FNetPacketNotify Acks = DefaultNotify;
+
+	FNetPacketNotify::SequenceNumberT ExpectedInSeq(31U);
+	FNetPacketNotify::SequenceHistoryT ExpectedInSeqHistory(0xffffffffu, 1U);
+				
+	for (FNetPacketNotify::SequenceNumberT::SequenceT I=0U; I<32U; ++I)
+	{
+		FNetPacketNotifyTestUtil::PretendReceiveSeq(Acks, I, true);
+	}
+
+	// Verify InSeq
+	UE_NET_ASSERT_EQ(Acks.GetInSeq(), ExpectedInSeq);
+
+	// Verify History
+	UE_NET_ASSERT_EQ(Acks.GetInSeqHistory(), ExpectedInSeqHistory);
+}
+
+UE_NET_TEST_FIXTURE(FNetPacketNotifyTestUtil, SequenceNumberOverflow)
+{
+	LOG_SCOPE_VERBOSITY_OVERRIDE(LogNetTraffic, ELogVerbosity::Error);
+
+	FNetPacketNotify Acks = DefaultNotify;
+
+	const FNetPacketNotify::SequenceNumberT MaxWindowSeq(FNetPacketNotify::SequenceNumberT::SeqNumberHalf);
+	const FNetPacketNotify::SequenceNumberT ExpectedInSeq(0);
+
+	// Verify default
+	UE_NET_ASSERT_NE(Acks.GetInSeq(), ExpectedInSeq);
+
+	// Pretend receive first seq
+	FNetPacketNotifyTestUtil::PretendReceiveSeq(Acks, ExpectedInSeq, true);
+
+	// Expects that we have accepted and acked the this sequence
+	UE_NET_ASSERT_EQ(Acks.GetInSeq(), ExpectedInSeq);
+	UE_NET_ASSERT_EQ(Acks.GetInAckSeq(), ExpectedInSeq);
+
+	// Verify that we did not accept invalid sequence
+	FNetPacketNotifyTestUtil::PretendReceiveSeq(Acks, MaxWindowSeq, true);
+	UE_NET_ASSERT_EQ(Acks.GetInSeq(), ExpectedInSeq);
+	UE_NET_ASSERT_EQ(Acks.GetInAckSeq(), ExpectedInSeq);
+
+	// Verify that we can receive next valid sequence
+	const FNetPacketNotify::SequenceNumberT NextExpectedInSeq(1);
+	FNetPacketNotifyTestUtil::PretendReceiveSeq(Acks, NextExpectedInSeq, true);
+	UE_NET_ASSERT_EQ(Acks.GetInSeq(), NextExpectedInSeq);
+	UE_NET_ASSERT_EQ(Acks.GetInAckSeq(), NextExpectedInSeq);
+
+	// Verify that we accepted a new sequence even if it is out and bounds and will trigger a SequencyHistoryFlush
+	FNetPacketNotifyTestUtil::PretendReceiveSeq(Acks, MaxWindowSeq, true);
+	UE_NET_ASSERT_GT(Acks.GetInSeq(), NextExpectedInSeq);
+	UE_NET_ASSERT_GT(Acks.GetInAckSeq(), NextExpectedInSeq);
+	UE_NET_ASSERT_TRUE(Acks.IsWaitingForSequenceHistoryFlush());
+}
+
+UE_NET_TEST_FIXTURE(FNetPacketNotifyTestUtil, BurstDrop)
+{
+	FNetPacketNotify Acks = DefaultNotify;
+
+	FNetPacketNotify::SequenceNumberT ExpectedInSeq(128);
+	FNetPacketNotify::SequenceHistoryT ExpectedInSeqHistory;
+	uint32 ExpectedArray[] = {0x1, 0, 0, 0x20000000 };
+	FNetPacketNotifyTestUtil::InitHistory(ExpectedInSeqHistory, ExpectedArray );
+
+	// Drop early
+	FNetPacketNotifyTestUtil::PretendReceiveSeq(Acks, 3);
+
+	// Large gap until next seq
+	FNetPacketNotifyTestUtil::PretendReceiveSeq(Acks, 128);
+
+	UE_NET_ASSERT_EQ(Acks.GetInSeq(), ExpectedInSeq);
+	UE_NET_ASSERT_EQ(Acks.GetInSeqHistory(), ExpectedInSeqHistory);
+}
+
+UE_NET_TEST_FIXTURE(FNetPacketNotifyTestUtil, CreateHistory)
+{
+	FNetPacketNotify Acks = DefaultNotify;
+
+	const FNetPacketNotify::SequenceNumberT ExpectedInSeq(18);
+	const FNetPacketNotify::SequenceHistoryT ExpectedInSeqHistory(0x8853u);
+
+	const FNetPacketNotify::SequenceNumberT AckdPacketIds[] = {3, 7, 12, 14, 17, 18};
+	const SIZE_T ExpectedCount = sizeof(AckdPacketIds)/sizeof((AckdPacketIds)[0]);		
+
+	FNetPacketNotifyTestUtil::PretendAckSequenceNumbers(Acks, AckdPacketIds);
+ 
+	UE_NET_ASSERT_EQ(Acks.GetInSeq(), ExpectedInSeq);
+	UE_NET_ASSERT_EQ(Acks.GetInSeqHistory(), ExpectedInSeqHistory);
+}
+
+UE_NET_TEST_FIXTURE(FNetPacketNotifyTestUtil, Notifications)
+{
+	FNetPacketNotify Acks = DefaultNotify;
+
+	static const FNetPacketNotify::SequenceNumberT ExpectedAckdPacketIds[] = {3, 7, 12, 14, 17, 18};
+	static const SIZE_T ExpectedCount = sizeof(ExpectedAckdPacketIds)/sizeof((ExpectedAckdPacketIds)[0]);		
+
+	FNetPacketNotify::SequenceNumberT RcvdAcks[ExpectedCount] = { 0 };
+	SIZE_T RcvdCount = 0;
+
+	// Fill in some data
+	FNetPacketNotify::FNotificationHeader Data;
+	Data.Seq = FNetPacketNotify::SequenceNumberT(0);
+	Data.AckedSeq = FNetPacketNotify::SequenceNumberT(18);
+	Data.History = FNetPacketNotify::SequenceHistoryT(0x8853u);
+	Data.HistoryWordCount = 1;
+
+	// Need to fake ack record as well.
+	for (SIZE_T It=0; It <= 18; ++It)
+	{
+		FNetPacketNotifyTestUtil::PretendSendSeq(Acks, 0);
+	}
+	
+	SIZE_T DeltaSeq = FNetPacketNotifyTestUtil::PretendReceivedPacket(Acks, Data, RcvdAcks);
+
+	UE_NET_ASSERT_EQ(DeltaSeq, SIZE_T(1));
+	UE_NET_ASSERT_EQ(FPlatformMemory::Memcmp(ExpectedAckdPacketIds, RcvdAcks, sizeof(ExpectedAckdPacketIds)), 0);
+}
+
+UE_NET_TEST_FIXTURE(FNetPacketNotifyTestUtil, ReceiveInvalidAck)
+{
+	FNetPacketNotify Acks = DefaultNotify;
+
+	static const FNetPacketNotify::SequenceNumberT ExpectedAckdPacketIds[] = {3, 7, 12, 14, 17, 18};
+	static const SIZE_T ExpectedCount = sizeof(ExpectedAckdPacketIds)/sizeof((ExpectedAckdPacketIds)[0]);		
+
+	FNetPacketNotify::SequenceNumberT RcvdAcks[ExpectedCount] = { 0 };
+	SIZE_T RcvdCount = 0;
+
+	// Fill in some data
+	FNetPacketNotify::FNotificationHeader Data;
+	Data.Seq = FNetPacketNotify::SequenceNumberT(0);
+	Data.AckedSeq = FNetPacketNotify::SequenceNumberT(19);
+	Data.History = FNetPacketNotify::SequenceHistoryT(0x8853u);
+	Data.HistoryWordCount = 1;
+
+	// Need to fake ack record as well.
+	for (SIZE_T It=0; It <= 18; ++It)
+	{
+		FNetPacketNotifyTestUtil::PretendSendSeq(Acks, FNetPacketNotify::SequenceNumberT(0U));
+	}
+	
+	SIZE_T DeltaSeq = FNetPacketNotifyTestUtil::PretendReceivedPacket(Acks, Data, RcvdAcks);
+
+	UE_NET_ASSERT_EQ(DeltaSeq, SIZE_T(0));
+}
+
+UE_NET_TEST_FIXTURE(FNetPacketNotifyTestUtil, FillSeqWindowWitNoMargin)
+{
+	LOG_SCOPE_VERBOSITY_OVERRIDE(LogNetTraffic, ELogVerbosity::Error);
+
+	// Test without safety margin
+	for (FNetPacketNotify::SequenceNumberT SeqOffset : {FNetPacketNotify::SequenceNumberT(0), FNetPacketNotify::SequenceNumberT(3)})
+	{
+		FNetPacketNotify PacketNotify = DefaultNotify;
+		for (FNetPacketNotify::SequenceNumberT It = 0; It < FNetPacketNotify::MaxSequenceHistoryLength - 1; ++It)
+		{
+			PretendSendSeq(PacketNotify, SeqOffset);
+		}
+		UE_NET_ASSERT_FALSE_MSG(PacketNotify.IsSequenceWindowFull(), "Test SeqWindowFull {0, MaxSequenceHistoryLength - 2}");
+
+		// Fill sequence window
+		FNetPacketNotifyTestUtil::PretendSendSeq(PacketNotify, SeqOffset);
+		UE_NET_ASSERT_TRUE_MSG(PacketNotify.IsSequenceWindowFull(), "Test SeqWindowFull {0, MaxSequenceHistoryLength - 1}");
+
+		// Fake all acked. We expect acks on all sequence numbers.
+		{
+			FNetPacketNotify::FNotificationHeader NotificationData;
+			NotificationData.Seq = 0;
+			NotificationData.AckedSeq = PacketNotify.GetOutSeq() - 1;
+			NotificationData.History = FNetPacketNotify::SequenceHistoryT(0xFFFFFFFFU, FNetPacketNotify::SequenceHistoryT::WordCount);
+			NotificationData.HistoryWordCount = FNetPacketNotify::SequenceHistoryT::WordCount;
+
+			PacketNotify.Update(NotificationData, [this](FNetPacketNotify::SequenceNumberT, bool bWasDelivered)
+				{
+					UE_NET_ASSERT_TRUE_MSG(bWasDelivered, "Test SeqWindowFull all delivered");
+				});
+		}
+	}
+}
+
+UE_NET_TEST_FIXTURE(FNetPacketNotifyTestUtil, FillSeqWindowWitSafetyMargin)
+{
+	// Test with safety margin
+	constexpr unsigned SafetyMargin = 3;
+	FNetPacketNotify PacketNotify = DefaultNotify;
+	for (SIZE_T It = 0; It < FNetPacketNotify::MaxSequenceHistoryLength - 1 - SafetyMargin; ++It)
+	{
+		FNetPacketNotifyTestUtil::PretendSendSeq(PacketNotify, 0);
+	}
+	UE_NET_ASSERT_FALSE_MSG(PacketNotify.IsSequenceWindowFull(SafetyMargin), "Test SeqWindowFull Margin 3 {0, MaxSequenceHistoryLength - 2}");
+
+	FNetPacketNotifyTestUtil::PretendSendSeq(PacketNotify, 0);
+	UE_NET_ASSERT_TRUE_MSG(PacketNotify.IsSequenceWindowFull(SafetyMargin), "Test SeqWindowFull Margin 3 {0, MaxSequenceHistoryLength - 1}");
+}
+
+UE_NET_TEST_FIXTURE(FNetPacketNotifyTestUtil, TestSendAndReceive)
+{
+	FTestNode Src;
+	FTestNode Dst;
+
+	// Pretend to send and accept packet in both directions
+	Src.SendAndDeliverTo(Dst);
+	Dst.SendAndDeliverTo(Src);
+
+	// Validate expected state
+	UE_NET_ASSERT_TRUE(Src.ValidateAcknowledgedPackets({0}));
+	UE_NET_ASSERT_TRUE(Dst.ValidateAcceptedPackets({0}));
+}
+
+UE_NET_TEST_FIXTURE(FNetPacketNotifyTestUtil, TestIgnoresOutOfBound)
+{
+	// Initialize with an offset that wraps around the sequence space
+	FTestNode Src(FNetPacketNotify::SequenceNumberT::SeqNumberHalf);
+	FTestNode Dst;
+
+	// Pretend to send and accept packet in both directions
+	Src.SendAndDeliverTo(Dst);
+	Dst.SendAndDeliverTo(Src);
+
+	// Validate expected state
+	UE_NET_ASSERT_EQ(Dst.AcceptedInPackets.Num(), 0);
+	UE_NET_ASSERT_EQ(Src.AcknowledgedPackets.Num(), 0);	
+}
+
+UE_NET_TEST_FIXTURE(FNetPacketNotifyTestUtil, TestOutOfBoundDoesNotTriggerFlush)
+{
+	FTestNode Src;
+	FTestNode Dst;
+
+	// Pretend to send and accept packet in both directions
+	Src.SendAndDeliverTo(Dst , {FNetPacketNotify::SequenceNumberT::SeqNumberMax, 0U});
+	Dst.SendAndDeliverTo(Src);
+
+	// Validate expected state
+	UE_NET_ASSERT_FALSE(Dst.PacketNotify.IsWaitingForSequenceHistoryFlush());
+	UE_NET_ASSERT_TRUE(Src.ValidateAcknowledgedPackets({0}));
+	UE_NET_ASSERT_TRUE(Dst.ValidateAcceptedPackets({0}));
+}
+
+UE_NET_TEST_FIXTURE(FNetPacketNotifyTestUtil, TestBurstReceiveValidRangeWrapAround)
+{
+	LOG_SCOPE_VERBOSITY_OVERRIDE(LogNetTraffic, ELogVerbosity::Error);
+
+	// Initialize with an offset that wraps around the sequence space
+	FTestNode Src(FNetPacketNotify::SequenceNumberT::SeqNumberHalf + 2);
+	FTestNode Dst(FNetPacketNotify::SequenceNumberT::SeqNumberHalf + 2);
+
+	// Pretend to send and accept packet in both directions
+	Src.SendAndDeliverTo(Dst, {FNetPacketNotify::SequenceNumberT::SeqNumberMax, 0U});
+	Dst.SendAndDeliverTo(Src);
+
+	// Validate expected state
+	UE_NET_ASSERT_FALSE(Dst.PacketNotify.IsWaitingForSequenceHistoryFlush());
+	UE_NET_ASSERT_TRUE(Src.ValidateAcknowledgedPackets({FNetPacketNotify::SequenceNumberT::SeqNumberMax, 0U}));
+	UE_NET_ASSERT_TRUE(Dst.ValidateAcceptedPackets({FNetPacketNotify::SequenceNumberT::SeqNumberMax, 0U}));
+}
+
+UE_NET_TEST_FIXTURE(FNetPacketNotifyTestUtil, TestAcceptsHugeDeltaWithNoPreviousAcks)
+{
+	LOG_SCOPE_VERBOSITY_OVERRIDE(LogNetTraffic, ELogVerbosity::Error);
+
+	// Initialize with an offset that is at the end of the sequence space
+	const FNetPacketNotify::SequenceNumberT FirstOutSeq(FNetPacketNotify::SequenceNumberT::SeqNumberHalf - 2);
+
+	FTestNode Src(-1, FirstOutSeq);
+	FTestNode Dst;
+
+	// Pretend to send and accept packet in both directions
+	Src.SendAndDeliverTo(Dst);
+	Dst.SendAndDeliverTo(Src);
+
+	// Validate expected state
+	UE_NET_ASSERT_TRUE(Src.ValidateAcknowledgedPackets({FirstOutSeq}));
+	UE_NET_ASSERT_TRUE(Dst.ValidateAcceptedPackets({FirstOutSeq}));
+}
+
+UE_NET_TEST_FIXTURE(FNetPacketNotifyTestUtil, TestBurstReceiveValidRange)
+{
+	FTestNode Src;
+	FTestNode Dst;
+
+	// Send a range burst of packets
+	Src.SendAndDeliverTo(Dst, {0U, LastValidSequenceHistoryIndex});
+
+	// Respond
+	Dst.SendAndDeliverTo(Src);
+
+	// Validate expected state
+	UE_NET_ASSERT_TRUE(Src.ValidateAcknowledgedPackets({0U, LastValidSequenceHistoryIndex}));
+	UE_NET_ASSERT_TRUE(Dst.ValidateAcceptedPackets({0U, LastValidSequenceHistoryIndex}));	
+}
+
+UE_NET_TEST_FIXTURE(FNetPacketNotifyTestUtil, TestBurstReceiveTriggersWaitForSequenceWindowFlush)
+{
+	LOG_SCOPE_VERBOSITY_OVERRIDE(LogNetTraffic, ELogVerbosity::Error);
+
+	FTestNode Src;
+	FTestNode Dst;
+
+	// Send a range burst of packets
+	Src.SendAndDeliverTo(Dst, {0U, FNetPacketNotify::MaxSequenceHistoryLength});
+
+	// Respond
+	Dst.SendAndDeliverTo(Src);
+
+	// Verify that we triggered overshoot
+	UE_NET_ASSERT_TRUE(Dst.PacketNotify.IsWaitingForSequenceHistoryFlush());
+
+	// Validate expected state
+	UE_NET_ASSERT_TRUE(Src.ValidateAcknowledgedPackets({0U}));
+	UE_NET_ASSERT_TRUE(Dst.ValidateAcceptedPackets({0U}));	
+}
+
+UE_NET_TEST_FIXTURE(FNetPacketNotifyTestUtil, TestReceiveDoesNotTriggersWaitForSequenceWindowFlushIfNoAcks)
+{
+	LOG_SCOPE_VERBOSITY_OVERRIDE(LogNetTraffic, ELogVerbosity::Error);
+
+	FTestNode Src;
+	FTestNode Dst;
+
+	// Send a specific sequence that would overflow window if we had already acked data
+	// we have no previous acks we can accept the sequence and allow the other end to inject any missing nacks
+	Src.SendAndDeliverTo(Dst, {FNetPacketNotify::MaxSequenceHistoryLength + 1, FNetPacketNotify::MaxSequenceHistoryLength + 2});
+
+	// Respond
+	Dst.SendAndDeliverTo(Src);
+
+	// Verify that we triggered overshoot
+	UE_NET_ASSERT_FALSE(Dst.PacketNotify.IsWaitingForSequenceHistoryFlush());
+
+	// Validate expected state
+	UE_NET_ASSERT_TRUE(Src.ValidateAcknowledgedPackets({FNetPacketNotify::MaxSequenceHistoryLength + 1, FNetPacketNotify::MaxSequenceHistoryLength + 2}));
+	UE_NET_ASSERT_TRUE(Dst.ValidateAcceptedPackets({FNetPacketNotify::MaxSequenceHistoryLength + 1, FNetPacketNotify::MaxSequenceHistoryLength + 2}));
+}
+
+UE_NET_TEST_FIXTURE(FNetPacketNotifyTestUtil, TestCanRecoverFromsWaitForSequenceWindowFlush)
+{
+	LOG_SCOPE_VERBOSITY_OVERRIDE(LogNetTraffic, ELogVerbosity::Error);
+
+	FTestNode Src;
+	FTestNode Dst;
+
+	// Send a range burst of packets
+	Src.SendAndDeliverTo(Dst, {0U, FNetPacketNotify::MaxSequenceHistoryLength});
+
+	// Respond
+	Dst.SendAndDeliverTo(Src);
+
+	// Verify that we triggered overshoot
+	UE_NET_ASSERT_TRUE(Dst.PacketNotify.IsWaitingForSequenceHistoryFlush());
+
+	// Validate expected state
+	UE_NET_ASSERT_TRUE(Src.ValidateAcknowledgedPackets({0U}));
+	UE_NET_ASSERT_TRUE(Dst.ValidateAcceptedPackets({0U}));
+
+	// Send a range burst of packets
+	Src.SendAndDeliverTo(Dst, {FNetPacketNotify::MaxSequenceHistoryLength + 1, FNetPacketNotify::MaxSequenceHistoryLength + 2});
+
+	// Verify that we recovered from overshoot
+	UE_NET_ASSERT_FALSE(Dst.PacketNotify.IsWaitingForSequenceHistoryFlush());
+
+	// Respond
+	Dst.SendAndDeliverTo(Src);
+
+	// Validate expected state
+	UE_NET_ASSERT_TRUE(Src.ValidateAcknowledgedPackets({0U, FNetPacketNotify::MaxSequenceHistoryLength + 1, FNetPacketNotify::MaxSequenceHistoryLength + 2}));
+	UE_NET_ASSERT_TRUE(Dst.ValidateAcceptedPackets({0U, FNetPacketNotify::MaxSequenceHistoryLength + 1, FNetPacketNotify::MaxSequenceHistoryLength + 2}));
+}
+
+UE_NET_TEST_FIXTURE(FNetPacketNotifyTestUtil, TestGetBothEndsWaitForSequenceWindowFlush)
+{
+	LOG_SCOPE_VERBOSITY_OVERRIDE(LogNetTraffic, ELogVerbosity::Error);
+
+	FTestNode Src;
+	FTestNode Dst;
+
+	// Send a range burst of packets
+	Src.SendAndDeliverTo(Dst, {0U, FNetPacketNotify::MaxSequenceHistoryLength});
+	Dst.SendAndDeliverTo(Src, {0U, FNetPacketNotify::MaxSequenceHistoryLength});
+
+	// Verify that we triggered overshoot on both ends
+	UE_NET_ASSERT_TRUE(Dst.PacketNotify.IsWaitingForSequenceHistoryFlush());
+	UE_NET_ASSERT_TRUE(Src.PacketNotify.IsWaitingForSequenceHistoryFlush());
+
+	// Send a range burst of from source to src, the packets should be accepted as we have cleared the sequence history
+	Src.SendAndDeliverTo(Dst, {FNetPacketNotify::MaxSequenceHistoryLength + 1, FNetPacketNotify::MaxSequenceHistoryLength + 2});
+
+	// Verify that we recovered from overshoot
+	UE_NET_ASSERT_FALSE(Dst.PacketNotify.IsWaitingForSequenceHistoryFlush());
+
+	// Src should still be waiting on flush
+	UE_NET_ASSERT_TRUE(Src.PacketNotify.IsWaitingForSequenceHistoryFlush());
+
+	// Send a range burst of from dst to src, the packets should be accepted as we have cleared the sequence history
+	Dst.SendAndDeliverTo(Src, {FNetPacketNotify::MaxSequenceHistoryLength + 1, FNetPacketNotify::MaxSequenceHistoryLength + 2});
+
+	// Verify that both recovered from overshoot
+	UE_NET_ASSERT_FALSE(Src.PacketNotify.IsWaitingForSequenceHistoryFlush());
+	UE_NET_ASSERT_FALSE(Dst.PacketNotify.IsWaitingForSequenceHistoryFlush());
+
+	// Respond
+	Src.SendAndDeliverTo(Src);
+	Dst.SendAndDeliverTo(Dst);
+
+	// Validate expected state
+	UE_NET_ASSERT_TRUE(Src.ValidateAcknowledgedPackets({0, FNetPacketNotify::MaxSequenceHistoryLength + 1, FNetPacketNotify::MaxSequenceHistoryLength + 2}));
+	UE_NET_ASSERT_TRUE(Src.ValidateAcceptedPackets({0, FNetPacketNotify::MaxSequenceHistoryLength + 1, FNetPacketNotify::MaxSequenceHistoryLength + 2, FNetPacketNotify::MaxSequenceHistoryLength + 3}));
+	UE_NET_ASSERT_TRUE(Src.ValidateAcknowledgedPackets({0, FNetPacketNotify::MaxSequenceHistoryLength + 1, FNetPacketNotify::MaxSequenceHistoryLength + 2}));
+	UE_NET_ASSERT_TRUE(Src.ValidateAcceptedPackets({0, FNetPacketNotify::MaxSequenceHistoryLength + 1, FNetPacketNotify::MaxSequenceHistoryLength + 2, FNetPacketNotify::MaxSequenceHistoryLength + 3}));
+}
+
+}
+}
+
diff --git a/Engine/Source/Runtime/Engine/Classes/Engine/NetConnection.h b/Engine/Source/Runtime/Engine/Classes/Engine/NetConnection.h
--- a/Engine/Source/Runtime/Engine/Classes/Engine/NetConnection.h
+++ b/Engine/Source/Runtime/Engine/Classes/Engine/NetConnection.h
@@ -980,9 +980,15 @@ public:
 	/** Poll the connection. If it is timed out, close it. */
 	ENGINE_API virtual void Tick(float DeltaSeconds);
 
-	/** Return whether this channel is ready for sending. */
+	/** Return whether this connection is ready for sending. */
 	ENGINE_API virtual int32 IsNetReady( bool Saturate );
 
+	/**
+	 * Return whether this connection can send packets without exhausting the packet sequence history window, as it could cause packets to be NAKed even when they've been received by the remote peer. 
+	 * @param SafetyMargin A small number representing how many packets you would like to keep as a safety margin for heart beats or other important packets.
+	 */
+	ENGINE_API bool IsPacketSequenceWindowFull(uint32 SafetyMargin=0U);
+
 	/** 
 	 * Handle the player controller client
 	 *
diff --git a/Engine/Source/Runtime/Engine/Private/Net/Experimental/Iris/DataStreamChannel.cpp b/Engine/Source/Runtime/Engine/Private/Net/Experimental/Iris/DataStreamChannel.cpp
--- a/Engine/Source/Runtime/Engine/Private/Net/Experimental/Iris/DataStreamChannel.cpp
+++ b/Engine/Source/Runtime/Engine/Private/Net/Experimental/Iris/DataStreamChannel.cpp
@@ -33,6 +33,13 @@ namespace UE::Net::Private
 		bIrisSaturateBandwidth,
 		TEXT("Whether to saturate the bandwidth or not. Default is false."
 		));
+
+	static int32 IrisPacketSequenceSafetyMargin = 4;
+	static FAutoConsoleVariableRef CVarIrisPacketSequenceSafetyMargin(
+		TEXT("net.Iris.PacketSequenceSafetyMargin"),
+		IrisPacketSequenceSafetyMargin,
+		TEXT("How many packets to spare before considering packet sequence full. This allows a few non-DataStreamChannel packets to be sent without messing up packet acking."
+		));
 }
 
 UDataStreamChannel::UDataStreamChannel(const FObjectInitializer& ObjectInitializer)
@@ -421,14 +428,17 @@ void UDataStreamChannel::ReceivedNak(int32 PacketId)
 #endif // UE_WITH_IRIS
 }
 
-//
+// Some DataStreams require perfect acking. If the ack sequence window is full we would get NAKs for packets thay may have been received.
 bool UDataStreamChannel::IsPacketWindowFull() const
 {
-#if UE_WITH_IRIS
-	return WriteRecords.Count() == WriteRecords.AllocatedCapacity();
-#else
-	return false;
-#endif // UE_WITH_IRIS
+	const uint32 IrisPacketSequenceSafetyMarginUnsigned = static_cast<uint32>(FPlatformMath::Max(0, UE::Net::Private::IrisPacketSequenceSafetyMargin));
+	if (Connection->IsPacketSequenceWindowFull(IrisPacketSequenceSafetyMarginUnsigned))
+	{
+		UE_LOG(LogIris, Verbose, TEXT("Packet window full."));
+		return true;
+	}
+
+	return WriteRecords.Count() >= WriteRecords.AllocatedCapacity();
 }
 
 void UDataStreamChannel::AddReferencedObjects(UObject* Object, FReferenceCollector& Collector)
@@ -442,7 +452,7 @@ void UDataStreamChannel::AddReferencedObjects(UObject* Object, FReferenceCollect
 	Super::AddReferencedObjects(Channel, Collector);
 }
 
-void UDataStreamChannel::AppendExportBunches(TArray<FOutBunch *>& OutExportBunches)
+void UDataStreamChannel::AppendExportBunches(TArray<FOutBunch*>& OutExportBunches)
 {
 }
 
diff --git a/Engine/Source/Runtime/Engine/Private/Net/NetPacketNotify.cpp b/Engine/Source/Runtime/Engine/Private/Net/NetPacketNotify.cpp
--- a/Engine/Source/Runtime/Engine/Private/Net/NetPacketNotify.cpp
+++ b/Engine/Source/Runtime/Engine/Private/Net/NetPacketNotify.cpp
@@ -12,19 +12,41 @@ FNetPacketNotify::FNetPacketNotify()
 {
 }
 
-SIZE_T FNetPacketNotify::GetCurrentSequenceHistoryLength() const
+FNetPacketNotify::SequenceNumberT::DifferenceT FNetPacketNotify::GetCurrentSequenceHistoryLength() const
 {
 	if (InAckSeq >= InAckSeqAck)
 	{
-		return (SIZE_T)SequenceNumberT::Diff(InAckSeq, InAckSeqAck);
+		return FMath::Min(SequenceNumberT::Diff(InAckSeq, InAckSeqAck), (SequenceNumberT::DifferenceT)SequenceHistoryT::Size);
 	}
 	else
 	{
 		// Worst case send full history
-		return SequenceHistoryT::Size;
+		return (SequenceNumberT::DifferenceT)SequenceHistoryT::Size;
 	}
 }
 
+bool FNetPacketNotify::WillSequenceFitInSequenceHistory(SequenceNumberT Seq) const
+{
+	if (Seq >= InAckSeqAck)
+	{
+		return (SIZE_T)SequenceNumberT::Diff(Seq, InAckSeqAck) <= SequenceHistoryT::Size;
+	}
+
+	return false;
+}
+
+bool FNetPacketNotify::GetHasUnacknowledgedAcks() const
+{
+	for (SequenceNumberT::DifferenceT It = 0, EndIt = GetCurrentSequenceHistoryLength(); It < EndIt; ++It)
+	{
+		if (InSeqHistory.IsDelivered(It))
+		{
+			return true;
+		}
+	}
+	return false;
+}
+
 FNetPacketNotify::SequenceNumberT FNetPacketNotify::UpdateInAckSeqAck(SequenceNumberT::DifferenceT AckCount, SequenceNumberT AckedSeq)
 {
 	if ((SIZE_T)AckCount <= AckRecord.Count())
@@ -40,6 +62,7 @@ FNetPacketNotify::SequenceNumberT FNetPacketNotify::UpdateInAckSeqAck(SequenceNu
 		// verify that we have a matching sequence number
 		if (AckData.OutSeq == AckedSeq)
 		{
+			UE_LOG_PACKET_NOTIFY(TEXT("FNetPacketNotify::UpdateInAckSeqAck - InAckSeqAck: %u"), AckData.InAckSeq.Get());
 			return AckData.InAckSeq;
 		}
 	}
@@ -58,6 +81,13 @@ void FNetPacketNotify::Init(SequenceNumberT InitialInSeq, SequenceNumberT Initia
 	InAckSeqAck = InitialInSeq;
 	OutSeq = InitialOutSeq;
 	OutAckSeq = SequenceNumberT(InitialOutSeq.Get() - 1);
+	WaitingForFlushSeqAck = OutAckSeq;
+}
+
+void FNetPacketNotify::SetWaitForSequenceHistoryFlush()
+{
+	UE_LOG_PACKET_NOTIFY_WARNING(TEXT("FNetPacketNotify::SetWaitForSequenceHistoryFlush - Wait for ack of next OutSeq: %u"), OutSeq.Get());
+	WaitingForFlushSeqAck = OutSeq;
 }
 
 void FNetPacketNotify::AckSeq(SequenceNumberT AckedSeq, bool IsAck)
@@ -70,12 +100,72 @@ void FNetPacketNotify::AckSeq(SequenceNumberT AckedSeq, bool IsAck)
 
 		const bool bReportAcked = InAckSeq == AckedSeq ? IsAck : false;
 
-		UE_LOG_PACKET_NOTIFY(TEXT("FNetPacketNotify::AckSeq - AckedSeq: %u, IsAck %u"), InAckSeq.Get(), bReportAcked ? 1u : 0u);
+		UE_LOG_PACKET_NOTIFY(TEXT("FNetPacketNotify::AckSeq - AckedSeq: %u, IsAck %u, AckHistorySize: %d"), InAckSeq.Get(), bReportAcked ? 1u : 0u, GetCurrentSequenceHistoryLength());
 
 		InSeqHistory.AddDeliveryStatus(bReportAcked);		
 	}
 }
 
+FNetPacketNotify::SequenceNumberT::DifferenceT FNetPacketNotify::InternalUpdate(const FNotificationHeader& NotificationData, SequenceNumberT::DifferenceT InSeqDelta)
+{
+	// We must check if we will overflow our outgoing ack-window, if we do and it contains processed data we must initiate a re-sync of ack-sequence history.
+	// This is done by ignoring any new packets until we are in sync again. 
+	// This would typically only occur in situations where we would have had huge packet loss or spikes on the receiving end.
+	if (!IsWaitingForSequenceHistoryFlush() && !WillSequenceFitInSequenceHistory(NotificationData.Seq))
+	{
+		if (GetHasUnacknowledgedAcks())
+		{
+			SetWaitForSequenceHistoryFlush();
+		}
+		else
+		{
+			// We can reset if we have no previous acks and can then safely synthesize nacks on the receiving end
+			const SequenceNumberT NewInAckSeqAck(NotificationData.Seq.Get() - 1);
+			UE_LOG_PACKET_NOTIFY_WARNING(TEXT("FNetPacketNotify::Reset SequenceHistory - New InSeqDelta: %u Old: %u"), NewInAckSeqAck.Get(), InAckSeqAck.Get());
+			InAckSeqAck = NewInAckSeqAck;
+		}
+	}
+
+	if (!IsWaitingForSequenceHistoryFlush())
+	{
+		// Just accept the incoming sequence, under normal circumstances NetConnection explicitly handles the acks.
+		InSeq = NotificationData.Seq;
+
+		return InSeqDelta;
+	}
+	else
+	{
+		// Until we have flushed the history we treat incoming packets as lost while still advancing ack window as far as we can.
+		SequenceNumberT NewInSeqToAck(NotificationData.Seq);
+
+		// Still waiting on flush, but we can fill up the history
+		if (!WillSequenceFitInSequenceHistory(NotificationData.Seq) && GetHasUnacknowledgedAcks())
+		{
+			// Mark everything we can as lost up until the end of the sequence history 
+			NewInSeqToAck = SequenceNumberT(InAckSeqAck.Get() + (MaxSequenceHistoryLength - GetCurrentSequenceHistoryLength()));
+		}
+
+		if (NewInSeqToAck >= InSeq)
+		{
+			const SequenceNumberT::DifferenceT AdjustedSequenceDelta = SequenceNumberT::Diff(NewInSeqToAck, InSeq);
+
+			InSeq = NewInSeqToAck;
+
+			// Nack driven from here
+			AckSeq(NewInSeqToAck, false);
+
+			UE_LOG_PACKET_NOTIFY(TEXT("FNetPacketNotify::Update - Waiting for sequence history flush - Rejected: %u Accepted: InSeq: %u Adjusted delta %d"), NotificationData.Seq.Get(), InSeq.Get(), AdjustedSequenceDelta);
+			
+			return AdjustedSequenceDelta;
+		}
+		else
+		{
+			UE_LOG_PACKET_NOTIFY(TEXT("FNetPacketNotify::Update - Waiting for sequence history flush - Rejected: %u Accepted: InSeq: %u"), NotificationData.Seq.Get(), InSeq.Get());
+			return 0;
+		}
+	}
+};
+
 namespace 
 {
 	struct FPackedHeader
diff --git a/Engine/Source/Runtime/Engine/Private/Net/Tests/NetPacketNotifyTests.cpp b/Engine/Source/Runtime/Engine/Private/Net/Tests/NetPacketNotifyTests.cpp
deleted file mode 100644
diff --git a/Engine/Source/Runtime/Engine/Private/NetConnection.cpp b/Engine/Source/Runtime/Engine/Private/NetConnection.cpp
--- a/Engine/Source/Runtime/Engine/Private/NetConnection.cpp
+++ b/Engine/Source/Runtime/Engine/Private/NetConnection.cpp
@@ -2209,7 +2209,6 @@ void UNetConnection::FlushNet(bool bIgnoreSimulation)
 		OutLagTime[Index]				= PacketSentTimeInS;
 		OutBytesPerSecondHistory[Index]	= FMath::Min(OutBytesPerSecond / 1024, 255);
 		
-
 		// Increase outgoing sequence number
 		if (!IsInternalAck())
 		{
@@ -2389,6 +2388,11 @@ int32 UNetConnection::IsNetReady(bool Saturate)
 	return QueuedBits + SendBuffer.GetNumBits() <= 0;
 }
 
+bool UNetConnection::IsPacketSequenceWindowFull(uint32 SafetyMargin)
+{
+	return PacketNotify.IsSequenceWindowFull(SafetyMargin);
+}
+
 void UNetConnection::ReadInput( float DeltaSeconds )
 {}
 
@@ -2824,6 +2828,9 @@ void UNetConnection::ReceivedPacket( FBitReader& Reader, bool bIsReinjectedPacke
 
 	const FEngineNetworkCustomVersion::Type PacketEngineNetVer = static_cast<FEngineNetworkCustomVersion::Type>(Reader.EngineNetVer());
 
+	// If we choose to not process this packet we need to restore it
+	const int32 OldInPacketId = InPacketId;
+
 	if (IsInternalAck())
 	{
 		++InPacketId;
@@ -2930,6 +2937,49 @@ void UNetConnection::ReceivedPacket( FBitReader& Reader, bool bIsReinjectedPacke
 				return;
 			}
 
+			// Process acks
+			// Lambda to dispatch delivery notifications, 
+			auto HandlePacketNotification = [&Header, &ChannelsToClose, this](FNetPacketNotify::SequenceNumberT AckedSequence, bool bDelivered)
+			{
+				// Increase LastNotifiedPacketId, this is a full packet Id
+				++LastNotifiedPacketId;
+				++OutTotalNotifiedPackets;
+				Driver->IncreaseOutTotalNotifiedPackets();
+
+				// Sanity check
+				if (FNetPacketNotify::SequenceNumberT(LastNotifiedPacketId) != AckedSequence)
+				{
+					UE_LOG(LogNet, Warning, TEXT("LastNotifiedPacketId != AckedSequence"));
+
+					Close(ENetCloseResult::AckSequenceMismatch);
+
+					return;
+				}
+
+				if (bDelivered)
+				{
+					ReceivedAck(LastNotifiedPacketId, ChannelsToClose);
+				}
+				else
+				{
+					ReceivedNak(LastNotifiedPacketId);
+				};
+			};
+
+			// Update incoming sequence data and deliver packet notifications
+			// Packet is only accepted if both the incoming sequence number and incoming ack data are valid		
+			const int32 UpdatedPacketSequenceDelta = PacketNotify.Update(Header, HandlePacketNotification);
+			if (PacketNotify.IsWaitingForSequenceHistoryFlush())
+			{
+				// Mark acks dirty
+				++HasDirtyAcks;
+
+				// Since we did not necessarily ack or nack all packets we need to update the InPacketId to reflect this.
+				InPacketId = OldInPacketId + UpdatedPacketSequenceDelta;
+
+				return;
+			}
+
 			if (MissingPacketCount > 10)
 			{
 				UE_LOG(LogNetTraffic, Verbose, TEXT("High single frame packet loss. PacketsLost: %i %s" ), MissingPacketCount, *Describe());
@@ -2941,6 +2991,8 @@ void UNetConnection::ReceivedPacket( FBitReader& Reader, bool bIsReinjectedPacke
 			Driver->InTotalPacketsLost += MissingPacketCount;
 			InPacketId += PacketSequenceDelta;
 
+			check(FNetPacketNotify::SequenceNumberT(InPacketId).Get() == Header.Seq.Get());
+
 			PacketAnalytics.TrackInPacket(InPacketId, MissingPacketCount);
 		}
 		else
@@ -2971,38 +3023,6 @@ void UNetConnection::ReceivedPacket( FBitReader& Reader, bool bIsReinjectedPacke
 			return;
 		}
 
-		// Lambda to dispatch delivery notifications, 
-		auto HandlePacketNotification = [&Header, &ChannelsToClose, this](FNetPacketNotify::SequenceNumberT AckedSequence, bool bDelivered)
-		{
-			// Increase LastNotifiedPacketId, this is a full packet Id
-			++LastNotifiedPacketId;
-			++OutTotalNotifiedPackets;
-			Driver->IncreaseOutTotalNotifiedPackets();
-
-			// Sanity check
-			if (FNetPacketNotify::SequenceNumberT(LastNotifiedPacketId) != AckedSequence)
-			{
-				UE_LOG(LogNet, Warning, TEXT("LastNotifiedPacketId != AckedSequence"));
-
-				Close(ENetCloseResult::AckSequenceMismatch);
-
-				return;
-			}
-
-			if (bDelivered)
-			{
-				ReceivedAck(LastNotifiedPacketId, ChannelsToClose);
-			}
-			else
-			{
-				ReceivedNak(LastNotifiedPacketId);
-			};
-		};
-
-		// Update incoming sequence data and deliver packet notifications
-		// Packet is only accepted if both the incoming sequence number and incoming ack data are valid
-		PacketNotify.Update(Header, HandlePacketNotification);
-
 		// Extra information associated with the header (read only after acks have been processed)
 		if (PacketSequenceDelta > 0 && !ReadPacketInfo(Reader, bHasPacketInfoPayload, PacketEngineNetVer))
 		{
@@ -3065,22 +3085,22 @@ void UNetConnection::ReceivedPacket( FBitReader& Reader, bool bIsReinjectedPacke
 	}
 
 	// Acknowledge the packet.
-	if ( !bSkipAck )
+	if (!bSkipAck)
 	{
 		LastGoodPacketRealtime = PostReceiveTime;
 	}
 
-	if( !IsInternalAck() )
+	if(!IsInternalAck())
 	{
 		// We always call AckSequence even if we are explicitly rejecting the packet as this updates the expected InSeq used to drive future acks.
-		if ( bSkipAck )
+		if (bSkipAck)
 		{
 			// Explicit Nak, we treat this packet as dropped but we still report it to the sending side as quickly as possible
-			PacketNotify.NakSeq( InPacketId );
+			PacketNotify.NakSeq(InPacketId);
 		}
 		else
 		{
-			PacketNotify.AckSeq( InPacketId );
+			PacketNotify.AckSeq(InPacketId);
 
 			// Keep stats happy
 			++OutTotalAcks;
@@ -3090,20 +3110,6 @@ void UNetConnection::ReceivedPacket( FBitReader& Reader, bool bIsReinjectedPacke
 		// We do want to let the other side know about the ack, so even if there are no other outgoing data when we tick the connection we will send an ackpacket.
 		TimeSensitive = 1;
 		++HasDirtyAcks;
-
-		// This is to allow us to recover from hitches were we process more than FNetPacketNotify::SequenceHistoryLength packets in a row withouht sending out any packets.
-		// In most cases this allows us to recover smoothly without reporting any pessimistic naks which will be the case if we overshoot the ack history.
-		// Note: This should only occur if we are running with no timeouts!
-		if (HasDirtyAcks >= FNetPacketNotify::MaxSequenceHistoryLength)
-		{
-			UE_LOG(LogNet, Warning, TEXT("UNetConnection::ReceivedPacket - Too many received packets to ack (%u) since last sent packet. InSeq: %u %s NextOutGoingSeq: %u"), HasDirtyAcks, PacketNotify.GetInSeq().Get(), *Describe(), PacketNotify.GetOutSeq().Get());
-
-			FlushNet();
-			if (HasDirtyAcks) // if acks still are dirty, flush again
-			{
-				FlushNet();
-			}
-		}
 	}
 
 	// Flush trace content collector
@@ -4109,7 +4115,7 @@ int32 UNetConnection::SendRawBunch(FOutBunch& Bunch, bool InAllowMerge, const FN
 	PrepareWriteBitsToSendBuffer(BunchHeaderBits, BunchBits);
 
 	// We want to mark the packet in which we write the data as TimeSensitive
-	// Note: we want to mark the packet as TimeSensitive here, as PrepareWriteBitsToSendBuffer migth flush the packet
+	// Note: we want to mark the packet as TimeSensitive here, as PrepareWriteBitsToSendBuffer might flush the packet
 	TimeSensitive = 1;
 
 	// Report bunch
diff --git a/Engine/Source/Runtime/Engine/Public/Net/NetPacketNotify.h b/Engine/Source/Runtime/Engine/Public/Net/NetPacketNotify.h
--- a/Engine/Source/Runtime/Engine/Public/Net/NetPacketNotify.h
+++ b/Engine/Source/Runtime/Engine/Public/Net/NetPacketNotify.h
@@ -27,6 +27,11 @@
 struct FBitWriter;
 struct FBitReader;
 
+namespace UE::Net::Private
+{
+	struct FNetPacketNotifyTestUtil;
+}
+
 /** 
 	FNetPacketNotify - Drives delivery of sequence numbers, acknowledgments and notifications of delivery sequence numbers
 */
@@ -95,6 +100,7 @@ public:
 	*/
 	template<class Functor>
 	SequenceNumberT::DifferenceT Update(const FNotificationHeader& NotificationData, Functor&& InFunc);
+
 	/** Get the current SequenceHistory */
 	const SequenceHistoryT& GetInSeqHistory() const { return InSeqHistory; }
 
@@ -113,8 +119,17 @@ public:
 	/** If we do have more unacknowledged sequence numbers in-flight than our maximum sendwindow we should not send more as the receiving end will not be able to detect if the sequence number has wrapped around */
 	bool CanSend() const { SequenceNumberT NextOutSeq = OutSeq; ++NextOutSeq; return NextOutSeq >= OutAckSeq; }
 
-	/** Get the current sequenceHistory length in bits */
-	SIZE_T GetCurrentSequenceHistoryLength() const;
+	/**
+	 * Return whether we can send packets without exhausting the packet sequence history window, as it could cause packets to be NAKed even when they've been received by the remote peer. 
+	 * @param SafetyMargin A small number representing how many packets you would like to keep as a safety margin for heart beats or other important packets.
+	 */
+	bool IsSequenceWindowFull(uint32 SafetyMargin=0U) const;
+
+	/** Get the current sequenceHistory length in bits, clamped to the maximum history length */
+	SequenceNumberT::DifferenceT GetCurrentSequenceHistoryLength() const;
+
+	/** Returns true if we are currently waiting for a flush of the sequence window */
+	bool IsWaitingForSequenceHistoryFlush() const { return WaitingForFlushSeqAck > OutAckSeq; }
 
 private:
 	struct FSentAckData
@@ -133,23 +148,41 @@ private:
 	SequenceNumberT InSeq;				// Last sequence number received and accepted from remote
 	SequenceNumberT InAckSeq;			// Last sequence number received from remote that we have acknowledged, this is needed since we support accepting a packet but explicitly not acknowledge it as received.
 	SequenceNumberT InAckSeqAck;		// Last sequence number received from remote that we have acknowledged and also knows that the remote has received the ack, used to calculate how big our history must be
+	SequenceNumberT WaitingForFlushSeqAck;
 
 	// Track outgoing sequence data
 	SequenceNumberT OutSeq;				// Outgoing sequence number
 	SequenceNumberT OutAckSeq;			// Last sequence number that we know that the remote side have received.
 
 private:
+
 	SequenceNumberT UpdateInAckSeqAck(SequenceNumberT::DifferenceT AckCount, SequenceNumberT AckedSeq);
+	SequenceNumberT::DifferenceT InternalUpdate(const FNotificationHeader& NotificationData, SequenceNumberT::DifferenceT InSeqDelta);
+
+	// Returns true if sequence history contains any packets marked as received for which we have not yet received an ack
+	bool GetHasUnacknowledgedAcks() const;
+
+	// Returns true if we can acknowledge the Seq without overshooting the sequence history
+	bool WillSequenceFitInSequenceHistory(SequenceNumberT Seq) const;
+
+	// Initiates a wait for a flush of the sequence history
+	void SetWaitForSequenceHistoryFlush();
 
 	template<class Functor>
 	inline void ProcessReceivedAcks(const FNotificationHeader& NotificationData, Functor&& InFunc);
 	void AckSeq(SequenceNumberT AckedSeq, bool IsAck);
 
-#if WITH_DEV_AUTOMATION_TESTS
-	friend struct FNetPacketNotifyTestUtil;
+#if WITH_AUTOMATION_WORKER
+	friend UE::Net::Private::FNetPacketNotifyTestUtil;
 #endif
 };
 
+inline bool FNetPacketNotify::IsSequenceWindowFull(uint32 SafetyMargin) const
+{
+	const SequenceNumberT SequenceLength = OutSeq - OutAckSeq;
+	return SequenceLength > MaxSequenceHistoryLength || (SafetyMargin >= MaxSequenceHistoryLength) || (SequenceLength.Get() > (MaxSequenceHistoryLength - SafetyMargin));
+}
+
 template<class Functor>
 FNetPacketNotify::SequenceNumberT::DifferenceT FNetPacketNotify::Update(const FNotificationHeader& NotificationData, Functor&& InFunc)
 {
@@ -161,10 +194,7 @@ FNetPacketNotify::SequenceNumberT::DifferenceT FNetPacketNotify::Update(const FN
 	
 		ProcessReceivedAcks(NotificationData, InFunc);
 
-		// accept sequence
-		InSeq = NotificationData.Seq;
-
-		return InSeqDelta;
+		return InternalUpdate(NotificationData, InSeqDelta);
 	}
 	else
 	{
@@ -177,25 +207,36 @@ void FNetPacketNotify::ProcessReceivedAcks(const FNotificationHeader& Notificati
 {
 	if (NotificationData.AckedSeq > OutAckSeq)
 	{
-		UE_LOG_PACKET_NOTIFY(TEXT("Notification::ProcessReceivedAcks - AckedSeq: %u, OutAckSeq: %u"), NotificationData.AckedSeq.Get(), OutAckSeq.Get());
-
 		SequenceNumberT::DifferenceT AckCount = SequenceNumberT::Diff(NotificationData.AckedSeq, OutAckSeq);
 
-		// Update InAckSeqAck used to track the needed number of bits to transmit our ack history
-		InAckSeqAck = UpdateInAckSeqAck(AckCount, NotificationData.AckedSeq);
+		UE_LOG_PACKET_NOTIFY(TEXT("Notification::ProcessReceivedAcks - AckedSeq: %u, OutAckSeq: %u AckCount: %u"), NotificationData.AckedSeq.Get(), OutAckSeq.Get(), AckCount);
 
+		// Update InAckSeqAck used to track the needed number of bits to transmit our ack history
+		// Note: As we might reset sequence history we need to check if we already have advanced the InAckSeqAck
+		const SequenceNumberT NewInAckSeqAck = UpdateInAckSeqAck(AckCount, NotificationData.AckedSeq);
+		if (NewInAckSeqAck > InAckSeqAck)
+		{
+			InAckSeqAck = NewInAckSeqAck;
+		}
+		
 		// ExpectedAck = OutAckSeq + 1
 		SequenceNumberT CurrentAck(OutAckSeq);
 		++CurrentAck;
 
-		// Warn if the received sequence number is greater than our history buffer, since if that is the case we have to treat the data as lost.
+		// Make sure that we only look at the sequence history bit included in the notification data as the sequence history might have been reset, 
+		// in which case we might not receive the max size history even though the ack-count is bigger than the history
+		const SequenceNumberT::DifferenceT HistoryBits = NotificationData.HistoryWordCount * SequenceHistoryT::BitsPerWord;
+
+		// Warn if the received sequence number is greater than our history buffer, since if that is the case we have to treat the data as lost
+		// Note: This should normally not be a problem since we try to flush the sequence history before allowing an overshoot of the sequence history window on the sending side.
+		// If this occurs with no hitches on server or client, there might be reason to investigate if too much data is being sent in which case the the size sequence history might have to be increased.
 		if (AckCount > (SequenceNumberT::DifferenceT)(SequenceHistoryT::Size))
 		{
-			UE_LOG_PACKET_NOTIFY_WARNING(TEXT("Notification::ProcessReceivedAcks - Missed Acks: AckedSeq: %u, OutAckSeq: %u, FirstMissingSeq: %u Count: %u"), NotificationData.AckedSeq.Get(), OutAckSeq.Get(), CurrentAck.Get(), AckCount - (SequenceNumberT::DifferenceT)(SequenceHistoryT::Size));
+			UE_LOG_PACKET_NOTIFY_WARNING(TEXT("FNetPacketNotify::ProcessReceivedAcks - Missed Acks: AckedSeq: %u, OutAckSeq: %u, FirstMissingSeq: %u Count: %u"), NotificationData.AckedSeq.Get(), OutAckSeq.Get(), CurrentAck.Get(), AckCount - (SequenceNumberT::DifferenceT)(SequenceHistoryT::Size));
 		}
 
 		// Everything not found in the history buffer is treated as lost
-		while (AckCount > (SequenceNumberT::DifferenceT)(SequenceHistoryT::Size))
+		while (AckCount > HistoryBits)
 		{
 			--AckCount;
 			InFunc(CurrentAck, false);
@@ -211,6 +252,12 @@ void FNetPacketNotify::ProcessReceivedAcks(const FNotificationHeader& Notificati
 			++CurrentAck;
 		}
 		OutAckSeq = NotificationData.AckedSeq;
+
+		// Are we done waiting for an reset of the ack history?
+		if (OutAckSeq > WaitingForFlushSeqAck)
+		{
+			WaitingForFlushSeqAck = OutAckSeq;
+		}
 	}
 }
 
diff --git a/Engine/Source/Runtime/Engine/Public/Net/Util/SequenceHistory.h b/Engine/Source/Runtime/Engine/Public/Net/Util/SequenceHistory.h
--- a/Engine/Source/Runtime/Engine/Public/Net/Util/SequenceHistory.h
+++ b/Engine/Source/Runtime/Engine/Public/Net/Util/SequenceHistory.h
@@ -27,10 +27,12 @@ public:
 public:	
 	TSequenceHistory();
 
-#if WITH_DEV_AUTOMATION_TESTS
+#if WITH_AUTOMATION_WORKER
 	explicit TSequenceHistory(WordT Value, SIZE_T Count = 1);
 	WordT* Data() { return &Storage[0]; }
+	const WordT* Data() const { return &Storage[0]; }
 #endif
+
 	/** Reset */
 	void Reset();
 
@@ -72,7 +74,7 @@ TSequenceHistory<HistorySize>::TSequenceHistory()
 	Reset();
 }
 
-#if WITH_DEV_AUTOMATION_TESTS
+#if WITH_AUTOMATION_WORKER
 template <SIZE_T HistorySize>
 TSequenceHistory<HistorySize>::TSequenceHistory(WordT Value, SIZE_T Count)
 {
