Iris
* Added unit test for setting a push model property dirty from inside PreNetUpdate

#rnx

diff --git a/Engine/Plugins/Runtime/ReplicationSystemTestPlugin/Source/Private/Tests/ReplicationSystem/ReplicatedTestObject.cpp b/Engine/Plugins/Runtime/ReplicationSystemTestPlugin/Source/Private/Tests/ReplicationSystem/ReplicatedTestObject.cpp
--- a/Engine/Plugins/Runtime/ReplicationSystemTestPlugin/Source/Private/Tests/ReplicationSystem/ReplicatedTestObject.cpp
+++ b/Engine/Plugins/Runtime/ReplicationSystemTestPlugin/Source/Private/Tests/ReplicationSystem/ReplicatedTestObject.cpp
@@ -21,6 +21,7 @@
 #include "Templates/SharedPointer.h"
 #include "Containers/Map.h"
 #include "UObject/StrongObjectPtr.h"
+#include "Net/Core/PushModel/PushModel.h"
 #include "Net/UnrealNetwork.h"
 #include "Iris/Serialization/NetBitStreamUtil.h"
 
@@ -107,6 +108,7 @@ bool UReplicatedTestObjectBridge::WriteCreationHeader(UE::Net::FNetSerialization
 	uint16 NumIrisComponentsToSpawn = 0U;
 	uint16 NumDynamicComponentsToSpawn = 0U;
 	uint16 NumConnectionFilteredComponentsToSpawn = 0U;
+	uint16 NumObjectReferenceComponentsToSpawn = 0U;
 	bool bForceFailToInstantiateOnRemote = false;
 
 	const UObject* Object = GetReplicatedObject(Handle);
@@ -120,6 +122,7 @@ bool UReplicatedTestObjectBridge::WriteCreationHeader(UE::Net::FNetSerialization
 		NumIrisComponentsToSpawn = IntCastChecked<uint16>(TestReplicatedIrisObject->IrisComponents.Num());
 		NumDynamicComponentsToSpawn = IntCastChecked<uint16>(TestReplicatedIrisObject->DynamicStateComponents.Num());
 		NumConnectionFilteredComponentsToSpawn = IntCastChecked<uint16>(TestReplicatedIrisObject->ConnectionFilteredComponents.Num());
+		NumObjectReferenceComponentsToSpawn = IntCastChecked<uint16>(TestReplicatedIrisObject->ObjectReferenceComponents.Num());
 	}
 
 	UObject* Archetype = Object->GetArchetype();
@@ -130,6 +133,7 @@ bool UReplicatedTestObjectBridge::WriteCreationHeader(UE::Net::FNetSerialization
 	Writer.WriteBits(NumIrisComponentsToSpawn, 16);
 	Writer.WriteBits(NumDynamicComponentsToSpawn, 16);
 	Writer.WriteBits(NumConnectionFilteredComponentsToSpawn, 16);
+	Writer.WriteBits(NumObjectReferenceComponentsToSpawn, 16);
 	Writer.WriteBool(bForceFailToInstantiateOnRemote);
 
 	return !Writer.IsOverflown();
@@ -148,6 +152,7 @@ UObjectReplicationBridge::FCreationHeader* UReplicatedTestObjectBridge::ReadCrea
 	Header->NumIrisComponentsToSpawn = Reader.ReadBits(16);
 	Header->NumDynamicComponentsToSpawn = Reader.ReadBits(16);
 	Header->NumConnectionFilteredComponentsToSpawn = Reader.ReadBits(16);
+	Header->NumObjectReferenceComponentsToSpawn = Reader.ReadBits(16);
 	Header->bForceFailCreateRemoteInstance = Reader.ReadBool();
 
 	if (Reader.IsOverflown())
@@ -187,6 +192,7 @@ FObjectReplicationBridgeInstantiateResult UReplicatedTestObjectBridge::BeginInst
 		Components.IrisComponentCount = Header->NumIrisComponentsToSpawn;
 		Components.DynamicStateComponentCount = Header->NumDynamicComponentsToSpawn;
 		Components.ConnectionFilteredComponentCount = Header->NumConnectionFilteredComponentsToSpawn;
+		Components.ObjectReferenceComponentCount = Header->NumObjectReferenceComponentsToSpawn;
 
 		CreatedTestObject->AddComponents(Components);
 	}
@@ -279,6 +285,12 @@ void UTestReplicatedIrisPushModelComponentWithObjectReference::GetLifetimeReplic
 	DOREPLIFETIME_WITH_PARAMS(ThisClass, WeakObjectPtrObjectRef, LifetimeParams);
 }
 
+void UTestReplicatedIrisPushModelComponentWithObjectReference::ModifyIntA()
+{
+	IntA += 1;
+	MARK_PROPERTY_DIRTY_FROM_NAME(ThisClass, IntA, this);
+}
+
 //////////////////////////////////////////////////////////////////////////
 // Implementation for UTestReplicatedIrisDynamicStatePropertyComponent
 //////////////////////////////////////////////////////////////////////////
diff --git a/Engine/Plugins/Runtime/ReplicationSystemTestPlugin/Source/Private/Tests/ReplicationSystem/ReplicatedTestObject.h b/Engine/Plugins/Runtime/ReplicationSystemTestPlugin/Source/Private/Tests/ReplicationSystem/ReplicatedTestObject.h
--- a/Engine/Plugins/Runtime/ReplicationSystemTestPlugin/Source/Private/Tests/ReplicationSystem/ReplicatedTestObject.h
+++ b/Engine/Plugins/Runtime/ReplicationSystemTestPlugin/Source/Private/Tests/ReplicationSystem/ReplicatedTestObject.h
@@ -109,6 +109,8 @@ class UTestReplicatedIrisPushModelComponentWithObjectReference : public UObject
 public:
 	UTestReplicatedIrisPushModelComponentWithObjectReference();
 
+	void ModifyIntA();
+
 	UPROPERTY(Transient, Replicated)
 	int32 IntA;
 
@@ -541,6 +543,7 @@ protected:
 		uint32 NumIrisComponentsToSpawn;
 		uint32 NumDynamicComponentsToSpawn;
 		uint32 NumConnectionFilteredComponentsToSpawn;
+		uint32 NumObjectReferenceComponentsToSpawn;
 		bool bForceFailCreateRemoteInstance;
 	};
 
diff --git a/Engine/Plugins/Runtime/ReplicationSystemTestPlugin/Source/Private/Tests/ReplicationSystem/TestDirtyNetObjectTracker.cpp b/Engine/Plugins/Runtime/ReplicationSystemTestPlugin/Source/Private/Tests/ReplicationSystem/TestDirtyNetObjectTracker.cpp
--- a/Engine/Plugins/Runtime/ReplicationSystemTestPlugin/Source/Private/Tests/ReplicationSystem/TestDirtyNetObjectTracker.cpp
+++ b/Engine/Plugins/Runtime/ReplicationSystemTestPlugin/Source/Private/Tests/ReplicationSystem/TestDirtyNetObjectTracker.cpp
@@ -181,7 +181,7 @@ UE_NET_TEST_FIXTURE(FReplicationSystemServerClientTestFixture, GlobalDirtyTracke
 }
 
 /** Test that validates behavior when dirtying other actors inside PreUpdate/PreReplication */
-UE_NET_TEST_FIXTURE(FReplicationSystemServerClientTestFixture, DirtyInsidePreUpdateTest)
+UE_NET_TEST_FIXTURE(FReplicationSystemServerClientTestFixture, ForceOtherObjectDirtyInsidePreUpdateTest)
 {
 	// Add client
 	FReplicationSystemTestClient* Client = CreateClient();
@@ -295,4 +295,44 @@ UE_NET_TEST_FIXTURE(FReplicationSystemServerClientTestFixture, DirtyInsidePreUpd
 	Server->DestroyObject(ServerObjectB);
 }
 
+/** Test that validates that a push model enabled object is marked as dirty inside PreUpdate/PreReplication */
+UE_NET_TEST_FIXTURE(FReplicationSystemServerClientTestFixture, PushModelMarkSelfDirtyInsidePreUpdateTest)
+{
+	// Add client
+	FReplicationSystemTestClient* Client = CreateClient();
+
+	// Spawn object with a PreUpdate call
+	TStrongObjectPtr<UTestReplicatedIrisObject> ServerObject = TStrongObjectPtr(NewObject<UTestReplicatedIrisObject>());
+	ServerObject->AddComponents({.ObjectReferenceComponentCount = 1});
+	{
+		UObjectReplicationBridge::FCreateNetRefHandleParams Params;
+		Params.bNeedsPreUpdate = true;
+		Server->ReplicationBridge->BeginReplication(ServerObject.Get(), Params);
+	}
+
+	// Send and deliver packet
+	Server->UpdateAndSend({ Client });
+
+	// Object should have been created on the client
+	UTestReplicatedIrisObject* ClientObject = Cast<UTestReplicatedIrisObject>(Client->GetReplicationBridge()->GetReplicatedObject(ServerObject->NetRefHandle));
+	UE_NET_ASSERT_NE(ClientObject, nullptr);
+
+	auto PreUpdateObject = [&](FNetRefHandle NetHandle, UObject* InObject, const UReplicationBridge* ReplicationBridge)
+	{
+		if (InObject == ServerObject.Get())
+		{
+			ServerObject->ObjectReferenceComponents[0]->ModifyIntA();
+		}
+	};
+
+	// Mark a property dirty during PreUpdate. As the object is polled every frame we expect the property to be updated on the client.
+	Server->GetReplicationBridge()->SetExternalPreUpdateFunctor(PreUpdateObject);
+
+	// Send and deliver packet
+	Server->UpdateAndSend({ Client });
+
+	// The property should be updated on the client
+	UE_NET_ASSERT_EQ(ClientObject->ObjectReferenceComponents[0]->IntA, ServerObject->ObjectReferenceComponents[0]->IntA);
+}
+
 }
