Replycation System Tests: fix unreachable code

diff --git a/Engine/Plugins/Runtime/ReplicationSystemTestPlugin/Source/Private/Tests/ReplicationSystem/RPC/ReplicatedTestObjectWithRPC.cpp b/Engine/Plugins/Runtime/ReplicationSystemTestPlugin/Source/Private/Tests/ReplicationSystem/RPC/ReplicatedTestObjectWithRPC.cpp
--- a/Engine/Plugins/Runtime/ReplicationSystemTestPlugin/Source/Private/Tests/ReplicationSystem/RPC/ReplicatedTestObjectWithRPC.cpp
+++ b/Engine/Plugins/Runtime/ReplicationSystemTestPlugin/Source/Private/Tests/ReplicationSystem/RPC/ReplicatedTestObjectWithRPC.cpp
@@ -89,9 +89,6 @@ int32 UTestReplicatedObjectWithRPC::GetFunctionCallspace(UFunction* Function, FF
 			return FunctionCallspace::Local;
 		}
 	}
-
-	// Call remotely
-	return FunctionCallspace::Remote;
 }
 
 bool UTestReplicatedObjectWithRPC::CallRemoteFunction(UFunction* Function, void* Parameters, FOutParmRec* OutParms, FFrame* Stack)
