Iris
- Fixed floating point precistoin bug with really high priorities which caused issues with scheduling of dependent objects
- Caused issues with scopes on weapons.

diff --git a/Engine/Plugins/Runtime/ReplicationSystemTestPlugin/Source/Private/Tests/ReplicationSystem/TestDependentObjects.cpp b/Engine/Plugins/Runtime/ReplicationSystemTestPlugin/Source/Private/Tests/ReplicationSystem/TestDependentObjects.cpp
--- a/Engine/Plugins/Runtime/ReplicationSystemTestPlugin/Source/Private/Tests/ReplicationSystem/TestDependentObjects.cpp
+++ b/Engine/Plugins/Runtime/ReplicationSystemTestPlugin/Source/Private/Tests/ReplicationSystem/TestDependentObjects.cpp
@@ -892,5 +892,57 @@ UE_NET_TEST_FIXTURE(FTestFilteredDependentObjectFixture, TestDependentObjectIsFi
 	UE_NET_ASSERT_NE(ClientDependentObject, nullptr);
 }
 
+UE_NET_TEST_FIXTURE(FReplicationSystemServerClientTestFixture, TestDependentObjectScheduledBeforeParentWithHighPriority)
+{
+	UReplicationSystem* ReplicationSystem = Server->ReplicationSystem;
+	UReplicatedTestObjectBridge* Bridge = Server->GetReplicationBridge();
+
+	// Add a client
+	FReplicationSystemTestClient* Client = CreateClient();
+
+	// Spawn object on server
+	UReplicatedSubObjectOrderObject* ServerObject = Server->CreateObject<UReplicatedSubObjectOrderObject>();
+
+	// Spawn dependent object
+	UReplicatedSubObjectOrderObject* ServerDependentObject = Server->CreateObject<UReplicatedSubObjectOrderObject>();
+	Bridge->AddDependentObject(ServerObject->NetRefHandle, ServerDependentObject->NetRefHandle, EDependentObjectSchedulingHint::ScheduleBeforeParent);
+
+	// Set prio to only replicate with parent
+	ReplicationSystem->SetStaticPriority(ServerDependentObject->NetRefHandle, 0.f);
+
+	// Set high prio on parent
+	const float VeryHighPriority = 1.0E7f;
+	ReplicationSystem->SetStaticPriority(ServerObject->NetRefHandle, VeryHighPriority);
+
+	// Reset RepOrderCounter
+	UReplicatedSubObjectOrderObject::RepOrderCounter = 0U;
+
+	// Send and deliver packet
+	Server->PreSendUpdate();
+	Server->SendAndDeliverTo(Client, true);
+	Server->PostSendUpdate();
+
+	// Verify that objects have replicated
+	UReplicatedSubObjectOrderObject* ClientObject = Cast<UReplicatedSubObjectOrderObject>(Client->GetReplicationBridge()->GetReplicatedObject(ServerObject->NetRefHandle));
+	UReplicatedSubObjectOrderObject* ClientDependentObject = Cast<UReplicatedSubObjectOrderObject>(Client->GetReplicationBridge()->GetReplicatedObject(ServerDependentObject->NetRefHandle));
+
+	UE_NET_ASSERT_NE(ClientObject, nullptr);
+	UE_NET_ASSERT_NE(ClientDependentObject, nullptr);
+
+	// Verify that they have replicated in expected order for initial objects
+	UE_NET_ASSERT_LT(ClientDependentObject->LastRepOrderCounter, ClientObject->LastRepOrderCounter);
+
+	// Modify both parent and dependent
+	ServerObject->IntA = 1;
+	ServerDependentObject->IntA = 1;
+
+	// Send and deliver packet
+	Server->PreSendUpdate();
+	Server->SendAndDeliverTo(Client, true);
+	Server->PostSendUpdate();
+
+	// Verify that they have replicated in expected order
+	UE_NET_ASSERT_LT(ClientDependentObject->LastRepOrderCounter, ClientObject->LastRepOrderCounter);
+}
 
 }
diff --git a/Engine/Source/Runtime/Experimental/Iris/Core/Private/Iris/ReplicationSystem/ReplicationWriter.cpp b/Engine/Source/Runtime/Experimental/Iris/Core/Private/Iris/ReplicationSystem/ReplicationWriter.cpp
--- a/Engine/Source/Runtime/Experimental/Iris/Core/Private/Iris/ReplicationSystem/ReplicationWriter.cpp
+++ b/Engine/Source/Runtime/Experimental/Iris/Core/Private/Iris/ReplicationSystem/ReplicationWriter.cpp
@@ -28,6 +28,7 @@
 #include "HAL/IConsoleManager.h"
 #include "ProfilingDebugging/CsvProfiler.h"
 #include <algorithm>
+#include <cmath> // std::nextafter
 
 #if UE_NET_ENABLE_REPLICATIONWRITER_LOG
 #	define UE_LOG_REPLICATIONWRITER(Format, ...)  UE_LOG(LogIris, Log, Format, ##__VA_ARGS__)
@@ -847,7 +848,7 @@ void FReplicationWriter::ScheduleDependentObjects(uint32 Index, float ParentPrio
 			if (bReplicateBeforeParent)
 			{
 				// Bump prio of dependent object to be scheduled before its parent.
-				UpdatedPriority = FMath::Max(ParentPriority + DependentObjectPriorityBump, LocalPriorities[DependentInternalIndex]);
+				UpdatedPriority = FMath::Max(std::nextafter(ParentPriority, std::numeric_limits<float>::infinity()), LocalPriorities[DependentInternalIndex]);
 				LocalPriorities[DependentInternalIndex] = UpdatedPriority;
 
 				// Schedule it, it does not matter if we add it to the scheduled list multiple times
