Iris
* Added unitest testing when one object causes a push model dirty change on a different object inside PreReplication

#rnx

diff --git a/Engine/Plugins/Runtime/ReplicationSystemTestPlugin/Source/Private/Tests/ReplicationSystem/ReplicationSystemServerClientTestFixture.cpp b/Engine/Plugins/Runtime/ReplicationSystemTestPlugin/Source/Private/Tests/ReplicationSystem/ReplicationSystemServerClientTestFixture.cpp
--- a/Engine/Plugins/Runtime/ReplicationSystemTestPlugin/Source/Private/Tests/ReplicationSystem/ReplicationSystemServerClientTestFixture.cpp
+++ b/Engine/Plugins/Runtime/ReplicationSystemTestPlugin/Source/Private/Tests/ReplicationSystem/ReplicationSystemServerClientTestFixture.cpp
@@ -100,9 +100,15 @@ uint32 FReplicationSystemTestNode::GetNetTraceId() const
 	return ReplicationSystem ? ReplicationSystem->GetId() : ~0U;
 }
 
-UTestReplicatedIrisObject* FReplicationSystemTestNode::CreateObject(const UObjectReplicationBridge::FCreateNetRefHandleParams& Params)
+UTestReplicatedIrisObject* FReplicationSystemTestNode::CreateObject(const UObjectReplicationBridge::FCreateNetRefHandleParams& Params, UTestReplicatedIrisObject::FComponents* ComponentsToCreate)
 {
 	UTestReplicatedIrisObject* CreatedObject = NewObject<UTestReplicatedIrisObject>();
+
+	if (ComponentsToCreate)
+	{
+		CreatedObject->AddComponents(*ComponentsToCreate);
+	}
+
 	CreatedObjects.Add(TStrongObjectPtr<UObject>(CreatedObject));
 
 	// Add it to the bridge for replication
diff --git a/Engine/Plugins/Runtime/ReplicationSystemTestPlugin/Source/Private/Tests/ReplicationSystem/ReplicationSystemServerClientTestFixture.h b/Engine/Plugins/Runtime/ReplicationSystemTestPlugin/Source/Private/Tests/ReplicationSystem/ReplicationSystemServerClientTestFixture.h
--- a/Engine/Plugins/Runtime/ReplicationSystemTestPlugin/Source/Private/Tests/ReplicationSystem/ReplicationSystemServerClientTestFixture.h
+++ b/Engine/Plugins/Runtime/ReplicationSystemTestPlugin/Source/Private/Tests/ReplicationSystem/ReplicationSystemServerClientTestFixture.h
@@ -116,7 +116,7 @@ public:
 		return static_cast<T*>(ReplicationBridge->GetReplicatedObject(Handle));
 	}
 
-	UTestReplicatedIrisObject* CreateObject(const UObjectReplicationBridge::FCreateNetRefHandleParams& Params);
+	UTestReplicatedIrisObject* CreateObject(const UObjectReplicationBridge::FCreateNetRefHandleParams& Params, UTestReplicatedIrisObject::FComponents* ComponentsToCreate=nullptr);
 	UTestReplicatedIrisObject* CreateObject(uint32 NumComponents, uint32 NumIrisComponents);
 	UTestReplicatedIrisObject* CreateSubObject(FNetRefHandle Owner, uint32 NumComponents, uint32 NumIrisComponents);
 	UTestReplicatedIrisObject* CreateObject(const UTestReplicatedIrisObject::FComponents& Components);
diff --git a/Engine/Plugins/Runtime/ReplicationSystemTestPlugin/Source/Private/Tests/ReplicationSystem/TestDirtyNetObjectTracker.cpp b/Engine/Plugins/Runtime/ReplicationSystemTestPlugin/Source/Private/Tests/ReplicationSystem/TestDirtyNetObjectTracker.cpp
--- a/Engine/Plugins/Runtime/ReplicationSystemTestPlugin/Source/Private/Tests/ReplicationSystem/TestDirtyNetObjectTracker.cpp
+++ b/Engine/Plugins/Runtime/ReplicationSystemTestPlugin/Source/Private/Tests/ReplicationSystem/TestDirtyNetObjectTracker.cpp
@@ -181,7 +181,7 @@ UE_NET_TEST_FIXTURE(FReplicationSystemServerClientTestFixture, GlobalDirtyTracke
 }
 
 /** Test that validates behavior when dirtying other actors inside PreUpdate/PreReplication */
-UE_NET_TEST_FIXTURE(FReplicationSystemServerClientTestFixture, ForceOtherObjectDirtyInsidePreUpdateTest)
+UE_NET_TEST_FIXTURE(FReplicationSystemServerClientTestFixture, NetForceUpdateOtherObjectInsidePreUpdateTest)
 {
 	// Add client
 	FReplicationSystemTestClient* Client = CreateClient();
@@ -295,6 +295,64 @@ UE_NET_TEST_FIXTURE(FReplicationSystemServerClientTestFixture, ForceOtherObjectD
 	Server->DestroyObject(ServerObjectB);
 }
 
+/** Test that validates behavior when dirtying other actors inside PreUpdate/PreReplication */
+UE_NET_TEST_FIXTURE(FReplicationSystemServerClientTestFixture, DirtyOtherObjectInsidePreUpdateTest)
+{
+	// Add client
+	FReplicationSystemTestClient* Client = CreateClient();
+
+	// Spawn object on server that is polled every frame.
+	UObjectReplicationBridge::FCreateNetRefHandleParams Params;
+	Params.bCanReceive = true;
+	Params.bUseClassConfigDynamicFilter = true;
+	Params.bNeedsPreUpdate = true;
+
+	UTestReplicatedIrisObject::FComponents ComponentsToCreate = { .ObjectReferenceComponentCount = 1 };
+
+	UTestReplicatedIrisObject* ServerObjectA = Server->CreateObject(Params, &ComponentsToCreate);
+	UTestReplicatedIrisObject* ServerObjectB = Server->CreateObject(Params, &ComponentsToCreate);
+
+	// Send and deliver packet
+	Server->UpdateAndSend({ Client });
+
+	// Object should have been created on the client
+	UTestReplicatedIrisObject* ClientObjectA = Cast<UTestReplicatedIrisObject>(Client->GetReplicationBridge()->GetReplicatedObject(ServerObjectA->NetRefHandle));
+	UTestReplicatedIrisObject* ClientObjectB = Cast<UTestReplicatedIrisObject>(Client->GetReplicationBridge()->GetReplicatedObject(ServerObjectB->NetRefHandle));
+	UE_NET_ASSERT_NE(ClientObjectA, nullptr);
+	UE_NET_ASSERT_NE(ClientObjectB, nullptr);
+
+	UTestReplicatedIrisObject* ObjectToDirty = nullptr;
+	auto PreUpdateObject = [&](FNetRefHandle NetHandle, UObject* ReplicatedObject, const UReplicationBridge* ReplicationBridge)
+	{
+		// There's only two objects, so when we update one we dirty the other
+		if (ObjectToDirty != ReplicatedObject)
+		{
+			ObjectToDirty->ObjectReferenceComponents[0]->ModifyIntA();
+		}
+	};
+	
+	// When ObjectA is polled modify and dirty ObjectB
+	ObjectToDirty = ServerObjectB;
+	Server->GetReplicationBridge()->SetExternalPreUpdateFunctor(PreUpdateObject);
+
+	// Send and deliver packet
+	Server->UpdateAndSend({ Client });
+
+	UE_NET_ASSERT_EQ(ClientObjectB->ObjectReferenceComponents[0]->IntA, ServerObjectB->ObjectReferenceComponents[0]->IntA);
+
+	// Now inverse it so when polling ObjectB we modify and dirty ObjectA
+	ObjectToDirty = ServerObjectA;
+
+	// Send and deliver packet
+	Server->UpdateAndSend({ Client });
+
+	UE_NET_ASSERT_EQ(ClientObjectA->ObjectReferenceComponents[0]->IntA, ServerObjectA->ObjectReferenceComponents[0]->IntA);
+
+	Server->DestroyObject(ServerObjectA);
+	Server->DestroyObject(ServerObjectB);
+}
+
+
 /** Test that validates that a push model enabled object is marked as dirty inside PreUpdate/PreReplication */
 UE_NET_TEST_FIXTURE(FReplicationSystemServerClientTestFixture, PushModelMarkSelfDirtyInsidePreUpdateTest)
 {
@@ -302,13 +360,11 @@ UE_NET_TEST_FIXTURE(FReplicationSystemServerClientTestFixture, PushModelMarkSelf
 	FReplicationSystemTestClient* Client = CreateClient();
 
 	// Spawn object with a PreUpdate call
-	TStrongObjectPtr<UTestReplicatedIrisObject> ServerObject = TStrongObjectPtr(NewObject<UTestReplicatedIrisObject>());
-	ServerObject->AddComponents({.ObjectReferenceComponentCount = 1});
-	{
-		UObjectReplicationBridge::FCreateNetRefHandleParams Params;
-		Params.bNeedsPreUpdate = true;
-		Server->ReplicationBridge->BeginReplication(ServerObject.Get(), Params);
-	}
+	UObjectReplicationBridge::FCreateNetRefHandleParams Params;
+	Params.bNeedsPreUpdate = true;
+	UTestReplicatedIrisObject::FComponents ComponentsToCreate = { .ObjectReferenceComponentCount = 1 };
+
+	UTestReplicatedIrisObject* ServerObject = Server->CreateObject(Params, &ComponentsToCreate);
 
 	// Send and deliver packet
 	Server->UpdateAndSend({ Client });
@@ -319,7 +375,7 @@ UE_NET_TEST_FIXTURE(FReplicationSystemServerClientTestFixture, PushModelMarkSelf
 
 	auto PreUpdateObject = [&](FNetRefHandle NetHandle, UObject* InObject, const UReplicationBridge* ReplicationBridge)
 	{
-		if (InObject == ServerObject.Get())
+		if (InObject == ServerObject)
 		{
 			ServerObject->ObjectReferenceComponents[0]->ModifyIntA();
 		}
