Iris
* Renamed Poll phase to PollAndCopy
* Renamed CopyAndQuantize phase to Quantize

diff --git a/Engine/Plugins/Runtime/ReplicationSystemTestPlugin/Source/Private/Tests/ReplicationSystem/TestReplicationOperations.cpp b/Engine/Plugins/Runtime/ReplicationSystemTestPlugin/Source/Private/Tests/ReplicationSystem/TestReplicationOperations.cpp
--- a/Engine/Plugins/Runtime/ReplicationSystemTestPlugin/Source/Private/Tests/ReplicationSystem/TestReplicationOperations.cpp
+++ b/Engine/Plugins/Runtime/ReplicationSystemTestPlugin/Source/Private/Tests/ReplicationSystem/TestReplicationOperations.cpp
@@ -60,7 +60,7 @@ struct FTestReplicationOperationsFixture : public FReplicationSystemTestFixture
 	}
 };
 
-UE_NET_TEST_FIXTURE(FTestReplicationOperationsFixture, CanPollAndRefreshPropertyData)
+UE_NET_TEST_FIXTURE(FTestReplicationOperationsFixture, CanPollAndCopyPropertyData)
 {
 	constexpr uint32 PropertyComponentCount = 5U;
 	constexpr uint32 IrisComponentCount = 3U;
@@ -108,7 +108,7 @@ UE_NET_TEST_FIXTURE(FTestReplicationOperationsFixture, CanPollAndRefreshProperty
 	TestObject->DynamicStateComponents[DynamicStateComponentIndex]->IntStaticArray[5] = 43;
 
 	// Poll data to update property replication state
-	FReplicationInstanceOperations::PollAndRefreshCachedPropertyData(ReplicationBridge->GetReplicationInstanceProtocol(TestObject->NetRefHandle));
+	FReplicationInstanceOperations::PollAndCopyPropertyData(ReplicationBridge->GetReplicationInstanceProtocol(TestObject->NetRefHandle));
 
 	// Verify values after poll
 	VerifyAs<int32>(1, State, 0);
@@ -156,7 +156,7 @@ UE_NET_TEST_FIXTURE(FTestReplicationOperationsFixture, CanQuantizeAndDequantize)
 	TestObject->DynamicStateComponents[DynamicStateComponentCount - 1]->IntStaticArray[4] = 127;
 
 	// Poll data to update property replication state
-	FReplicationInstanceOperations::PollAndRefreshCachedPropertyData(ReplicationBridge->GetReplicationInstanceProtocol(TestObject->NetRefHandle));
+	FReplicationInstanceOperations::PollAndCopyPropertyData(ReplicationBridge->GetReplicationInstanceProtocol(TestObject->NetRefHandle));
 
 	// Quantize to buffer
 	const uint32 BufferSize = 1024;
@@ -177,7 +177,7 @@ UE_NET_TEST_FIXTURE(FTestReplicationOperationsFixture, CanQuantizeAndDequantize)
 	SerializationContext.SetIsInitState(true);
 
 	// Copy data from test object to buffer
-	FReplicationInstanceOperations::CopyAndQuantize(SerializationContext, AlignedBuffer, GetChangeMaskWriter(ProtocolA->ChangeMaskBitCount), ReplicationBridge->GetReplicationInstanceProtocol(TestObject->NetRefHandle), ProtocolA);
+	FReplicationInstanceOperations::Quantize(SerializationContext, AlignedBuffer, GetChangeMaskWriter(ProtocolA->ChangeMaskBitCount), ReplicationBridge->GetReplicationInstanceProtocol(TestObject->NetRefHandle), ProtocolA);
 
 	// Verify that we did not overshoot
 	UE_NET_ASSERT_EQ(ValueToSet, AlignedBuffer[ProtocolA->InternalTotalSize]);
@@ -251,7 +251,7 @@ UE_NET_TEST_FIXTURE(FTestReplicationOperationsFixture, CanSerializeAndDeserializ
 	TestObject->DynamicStateComponents[DynamicStateComponentCount - 1]->IntStaticArray[4] = 127;
 
 	// Poll data to update property replication state
-	FReplicationInstanceOperations::PollAndRefreshCachedPropertyData(ReplicationBridge->GetReplicationInstanceProtocol(TestObject->NetRefHandle));
+	FReplicationInstanceOperations::PollAndCopyPropertyData(ReplicationBridge->GetReplicationInstanceProtocol(TestObject->NetRefHandle));
 
 	// Quantize to buffer
 	const uint32 BufferSize = 1024;
@@ -271,7 +271,7 @@ UE_NET_TEST_FIXTURE(FTestReplicationOperationsFixture, CanSerializeAndDeserializ
 	FInternalNetSerializationContext InternalContext;
 	SerializationContext.SetInternalContext(&InternalContext);
 	SerializationContext.SetIsInitState(true);
-	FReplicationInstanceOperations::CopyAndQuantize(SerializationContext, AlignedBuffer, GetChangeMaskWriter(ProtocolA->ChangeMaskBitCount), ReplicationBridge->GetReplicationInstanceProtocol(TestObject->NetRefHandle), ProtocolA);
+	FReplicationInstanceOperations::Quantize(SerializationContext, AlignedBuffer, GetChangeMaskWriter(ProtocolA->ChangeMaskBitCount), ReplicationBridge->GetReplicationInstanceProtocol(TestObject->NetRefHandle), ProtocolA);
 
 	// Verify that we did not overshoot
 	UE_NET_ASSERT_EQ(ValueToSet, AlignedBuffer[ProtocolA->InternalTotalSize]);
@@ -362,7 +362,7 @@ UE_NET_TEST_FIXTURE(FTestReplicationOperationsFixture, PreAndPostNetReceivedOnly
 	TestObject->Components[0]->IntB = 3;
 
 	// Poll data to update property replication state
-	FReplicationInstanceOperations::PollAndRefreshCachedPropertyData(ReplicationBridge->GetReplicationInstanceProtocol(TestObject->NetRefHandle));
+	FReplicationInstanceOperations::PollAndCopyPropertyData(ReplicationBridge->GetReplicationInstanceProtocol(TestObject->NetRefHandle));
 
 	// Quantize to buffer
 	uint8* AlignedBuffer = (uint8*)(GetTempAllocator()).Alloc(ProtocolA->InternalTotalSize, ProtocolA->InternalTotalAlignment);
@@ -370,7 +370,7 @@ UE_NET_TEST_FIXTURE(FTestReplicationOperationsFixture, PreAndPostNetReceivedOnly
 	// Copy data from test object to buffer
 	FNetSerializationContext NetContext;
 	NetContext.SetIsInitState(true);
-	FReplicationInstanceOperations::CopyAndQuantize(NetContext, AlignedBuffer, GetChangeMaskWriter(ProtocolA->ChangeMaskBitCount), ReplicationBridge->GetReplicationInstanceProtocol(TestObject->NetRefHandle), ProtocolA);
+	FReplicationInstanceOperations::Quantize(NetContext, AlignedBuffer, GetChangeMaskWriter(ProtocolA->ChangeMaskBitCount), ReplicationBridge->GetReplicationInstanceProtocol(TestObject->NetRefHandle), ProtocolA);
 
 	// Create second object 
 	UTestReplicatedIrisObject* TestObjectB = CreateObject(1, 0);
@@ -406,7 +406,7 @@ UE_NET_TEST_FIXTURE(FTestReplicationOperationsFixture, RepNotifyIsCalledForSimpl
 	TestObject->Components[0]->IntB ^= 1;
 
 	// Poll data to update property replication state
-	FReplicationInstanceOperations::PollAndRefreshCachedPropertyData(ReplicationBridge->GetReplicationInstanceProtocol(TestObject->NetRefHandle));
+	FReplicationInstanceOperations::PollAndCopyPropertyData(ReplicationBridge->GetReplicationInstanceProtocol(TestObject->NetRefHandle));
 
 	// Quantize to buffer
 	uint8* AlignedBuffer = (uint8*)(GetTempAllocator()).Alloc(Protocol->InternalTotalSize, Protocol->InternalTotalAlignment);
@@ -414,7 +414,7 @@ UE_NET_TEST_FIXTURE(FTestReplicationOperationsFixture, RepNotifyIsCalledForSimpl
 	// Copy data from test object to buffer
 	FNetSerializationContext NetContext;
 	NetContext.SetIsInitState(true);
-	FReplicationInstanceOperations::CopyAndQuantize(NetContext, AlignedBuffer, GetChangeMaskWriter(Protocol->ChangeMaskBitCount), ReplicationBridge->GetReplicationInstanceProtocol(TestObject->NetRefHandle), Protocol);
+	FReplicationInstanceOperations::Quantize(NetContext, AlignedBuffer, GetChangeMaskWriter(Protocol->ChangeMaskBitCount), ReplicationBridge->GetReplicationInstanceProtocol(TestObject->NetRefHandle), Protocol);
 
 	// Create second object 
 	UTestReplicatedIrisObject* TestObjectB = CreateObject(PropertyComponentCount, 0);
@@ -446,7 +446,7 @@ UE_NET_TEST_FIXTURE(FTestReplicationOperationsFixture, InitPropertyIsNotUpdatedW
 	TestObject->Components[0]->IntB ^= 1;
 
 	// Poll data to update property replication state
-	FReplicationInstanceOperations::PollAndRefreshCachedPropertyData(ReplicationBridge->GetReplicationInstanceProtocol(TestObject->NetRefHandle));
+	FReplicationInstanceOperations::PollAndCopyPropertyData(ReplicationBridge->GetReplicationInstanceProtocol(TestObject->NetRefHandle));
 
 	// Quantize to buffer
 	uint8* AlignedBuffer = (uint8*)(GetTempAllocator()).Alloc(Protocol->InternalTotalSize, Protocol->InternalTotalAlignment);
@@ -454,7 +454,7 @@ UE_NET_TEST_FIXTURE(FTestReplicationOperationsFixture, InitPropertyIsNotUpdatedW
 	// Copy data from test object to buffer
 	FNetSerializationContext NetContext;
 	NetContext.SetIsInitState(false);
-	FReplicationInstanceOperations::CopyAndQuantize(NetContext, AlignedBuffer, GetChangeMaskWriter(Protocol->ChangeMaskBitCount), ReplicationBridge->GetReplicationInstanceProtocol(TestObject->NetRefHandle), Protocol);
+	FReplicationInstanceOperations::Quantize(NetContext, AlignedBuffer, GetChangeMaskWriter(Protocol->ChangeMaskBitCount), ReplicationBridge->GetReplicationInstanceProtocol(TestObject->NetRefHandle), Protocol);
 
 	// Create second object 
 	UTestReplicatedIrisObject* TestObjectB = CreateObject(PropertyComponentCount, 0);
@@ -496,7 +496,7 @@ UE_NET_TEST_FIXTURE(FTestReplicationOperationsFixture, RepNotifyIsCalledExactlyO
 	TestObject->DynamicStateComponents[0]->IntStaticArray[5] = 1;
 
 	// Poll data to update property replication state
-	FReplicationInstanceOperations::PollAndRefreshCachedPropertyData(ReplicationBridge->GetReplicationInstanceProtocol(TestObject->NetRefHandle));
+	FReplicationInstanceOperations::PollAndCopyPropertyData(ReplicationBridge->GetReplicationInstanceProtocol(TestObject->NetRefHandle));
 
 	// Quantize to buffer
 	uint8* AlignedBuffer = (uint8*)(GetTempAllocator()).Alloc(Protocol->InternalTotalSize, Protocol->InternalTotalAlignment);
@@ -507,7 +507,7 @@ UE_NET_TEST_FIXTURE(FTestReplicationOperationsFixture, RepNotifyIsCalledExactlyO
 	FInternalNetSerializationContext InternalContext;
 	SerializationContext.SetInternalContext(&InternalContext);
 	SerializationContext.SetIsInitState(true);
-	FReplicationInstanceOperations::CopyAndQuantize(SerializationContext, AlignedBuffer, GetChangeMaskWriter(Protocol->ChangeMaskBitCount), ReplicationBridge->GetReplicationInstanceProtocol(TestObject->NetRefHandle), Protocol);
+	FReplicationInstanceOperations::Quantize(SerializationContext, AlignedBuffer, GetChangeMaskWriter(Protocol->ChangeMaskBitCount), ReplicationBridge->GetReplicationInstanceProtocol(TestObject->NetRefHandle), Protocol);
 
 	// Create second object 
 	UTestReplicatedIrisObject* TestObjectB = CreateObjectWithDynamicState(0, 0, DynamicStateComponentCount);
@@ -588,7 +588,7 @@ UE_NET_TEST_FIXTURE(FTestReplicationOperationsForObjectsFixture, StaleObjectPoin
 
 	// Copy data from test object to buffer
 	ReplicationSystem->PreSendUpdate(0.033f);
-	FReplicationInstanceOperations::CopyAndQuantize(SerializationContext, StateBuffers1[0], GetChangeMaskWriter(Protocol1->ChangeMaskBitCount), InstanceProtocol1, Protocol1);
+	FReplicationInstanceOperations::Quantize(SerializationContext, StateBuffers1[0], GetChangeMaskWriter(Protocol1->ChangeMaskBitCount), InstanceProtocol1, Protocol1);
 	ReplicationSystem->PostSendUpdate();
 
 	// Destroy first object and invalidate references to it.
@@ -597,7 +597,7 @@ UE_NET_TEST_FIXTURE(FTestReplicationOperationsForObjectsFixture, StaleObjectPoin
 	CollectGarbage(RF_NoFlags);
 
 	ReplicationSystem->PreSendUpdate(0.033f);
-	FReplicationInstanceOperations::CopyAndQuantize(SerializationContext, StateBuffers1[1], GetChangeMaskWriter(Protocol1->ChangeMaskBitCount), InstanceProtocol1, Protocol1);
+	FReplicationInstanceOperations::Quantize(SerializationContext, StateBuffers1[1], GetChangeMaskWriter(Protocol1->ChangeMaskBitCount), InstanceProtocol1, Protocol1);
 	ReplicationSystem->PostSendUpdate();
 
 	// Compare the two state buffers. Despite that we haven't manually changed any object references we expect them to differ thanks to garbage collection support.
diff --git a/Engine/Source/Runtime/Experimental/Iris/Core/Private/Iris/ReplicationSystem/NetRefHandleManager.cpp b/Engine/Source/Runtime/Experimental/Iris/Core/Private/Iris/ReplicationSystem/NetRefHandleManager.cpp
--- a/Engine/Source/Runtime/Experimental/Iris/Core/Private/Iris/ReplicationSystem/NetRefHandleManager.cpp
+++ b/Engine/Source/Runtime/Experimental/Iris/Core/Private/Iris/ReplicationSystem/NetRefHandleManager.cpp
@@ -26,7 +26,7 @@ FNetRefHandleManager::FNetRefHandleManager(FReplicationProtocolManager& InReplic
 , ScopeFrameData(MaxActiveObjectCount)
 , RelevantObjectsInternalIndices(MaxActiveObjectCount)
 , PolledObjectsInternalIndices(MaxActiveObjectCount)
-, DirtyObjectsToCopy(MaxActiveObjectCount)
+, DirtyObjectsToQuantize(MaxActiveObjectCount)
 , AssignedInternalIndices(MaxActiveObjectCount)
 , SubObjectInternalIndices(MaxActiveObjectCount)
 , DependentObjectInternalIndices(MaxActiveObjectCount)
diff --git a/Engine/Source/Runtime/Experimental/Iris/Core/Private/Iris/ReplicationSystem/NetRefHandleManager.h b/Engine/Source/Runtime/Experimental/Iris/Core/Private/Iris/ReplicationSystem/NetRefHandleManager.h
--- a/Engine/Source/Runtime/Experimental/Iris/Core/Private/Iris/ReplicationSystem/NetRefHandleManager.h
+++ b/Engine/Source/Runtime/Experimental/Iris/Core/Private/Iris/ReplicationSystem/NetRefHandleManager.h
@@ -210,8 +210,8 @@ public:
 	/** List of objects that we polled this frame */
 	FNetBitArrayView GetPolledObjectsInternalIndices() const { return MakeNetBitArrayView(PolledObjectsInternalIndices); }
 
-	/** List of objects that need to copy their state data */
-	FNetBitArrayView GetDirtyObjectsToCopy() const { return MakeNetBitArrayView(DirtyObjectsToCopy); }
+	/** List of objects that have dirty state data that needs to be quantized */
+	FNetBitArrayView GetDirtyObjectsToQuantize() const { return MakeNetBitArrayView(DirtyObjectsToQuantize); }
 
 	// Get bitarray for all internal indices that currently are assigned
 	const FNetBitArray& GetAssignedInternalIndices() const { return AssignedInternalIndices; }
@@ -365,8 +365,8 @@ private:
 	/** List of objects that we polled this frame */
 	FNetBitArray PolledObjectsInternalIndices;
 
-	/** List of the objects that are considered dirty and for whom we will copy their state data */
-	FNetBitArray DirtyObjectsToCopy;
+	/** List of the objects that are considered dirty and for whom we will quantize their state data */
+	FNetBitArray DirtyObjectsToQuantize;
 
 	// Bitset containing all internal indices that are assigned
 	FNetBitArray AssignedInternalIndices;
diff --git a/Engine/Source/Runtime/Experimental/Iris/Core/Private/Iris/ReplicationSystem/ObjectReplicationBridge.cpp b/Engine/Source/Runtime/Experimental/Iris/Core/Private/Iris/ReplicationSystem/ObjectReplicationBridge.cpp
--- a/Engine/Source/Runtime/Experimental/Iris/Core/Private/Iris/ReplicationSystem/ObjectReplicationBridge.cpp
+++ b/Engine/Source/Runtime/Experimental/Iris/Core/Private/Iris/ReplicationSystem/ObjectReplicationBridge.cpp
@@ -795,7 +795,7 @@ void UObjectReplicationBridge::PreSendUpdate()
 
 	ReconcileNewSubObjects(ObjectsConsideredForPolling);
 
-	Poll(ObjectsConsideredForPolling);
+	PollAndCopy(ObjectsConsideredForPolling);
 }
 
 void UObjectReplicationBridge::OnPostSendUpdate()
@@ -867,14 +867,14 @@ void UObjectReplicationBridge::ForcePollObject(FNetRefHandle Handle)
 
 	if (Handle.IsValid())
 	{
-		IRIS_PROFILER_SCOPE(UObjectReplicationBridge_ForcePollObject);
+		IRIS_PROFILER_SCOPE(UObjectReplicationBridge_ForcePollAndCopyObject);
 
 		FObjectPoller::FInitParams PollerInitParams;
 		PollerInitParams.ObjectReplicationBridge = this;
 		PollerInitParams.ReplicationSystemInternal = GetReplicationSystem()->GetReplicationSystemInternal();
 		FObjectPoller Poller(PollerInitParams);
 
-		Poller.PollSingleObject(Handle);
+		Poller.PollAndCopySingleObject(Handle);
 	}
 }
 
@@ -1034,18 +1034,18 @@ void UObjectReplicationBridge::PreUpdate(const UE::Net::FNetBitArrayView Objects
 	UE_NET_TRACE_FRAME_STATSCOUNTER(ReplicationSystem->GetId(), ReplicationSystem.PreUpdatedObjectCount, Stats.PreUpdatedObjectCount, ENetTraceVerbosity::Trace);
 }
 
-void UObjectReplicationBridge::Poll(const UE::Net::FNetBitArrayView ObjectsConsideredForPolling)
+void UObjectReplicationBridge::PollAndCopy(const UE::Net::FNetBitArrayView ObjectsConsideredForPolling)
 {
 	using namespace UE::Net::Private;
 
-	IRIS_PROFILER_SCOPE(UObjectReplicationBridge_Poll);
+	IRIS_PROFILER_SCOPE(UObjectReplicationBridge_PollAndCopy);
 
 	FObjectPoller::FInitParams PollerInitParams;
 	PollerInitParams.ObjectReplicationBridge = this;
 	PollerInitParams.ReplicationSystemInternal = GetReplicationSystem()->GetReplicationSystemInternal();
 
 	FObjectPoller Poller(PollerInitParams);
-	Poller.PollObjects(ObjectsConsideredForPolling);
+	Poller.PollAndCopyObjects(ObjectsConsideredForPolling);
 	
 	FObjectPoller::FPreUpdateAndPollStats Stats = Poller.GetPollStats();
 
diff --git a/Engine/Source/Runtime/Experimental/Iris/Core/Private/Iris/ReplicationSystem/Polling/ObjectPoller.cpp b/Engine/Source/Runtime/Experimental/Iris/Core/Private/Iris/ReplicationSystem/Polling/ObjectPoller.cpp
--- a/Engine/Source/Runtime/Experimental/Iris/Core/Private/Iris/ReplicationSystem/Polling/ObjectPoller.cpp
+++ b/Engine/Source/Runtime/Experimental/Iris/Core/Private/Iris/ReplicationSystem/Polling/ObjectPoller.cpp
@@ -25,7 +25,7 @@ FObjectPoller::FObjectPoller(const FInitParams& InitParams)
 	, NetStatsContext(nullptr)
 	, ReplicatedInstances(LocalNetRefHandleManager.GetReplicatedInstances())
 	, AccumulatedDirtyObjects(ReplicationSystemInternal->GetDirtyNetObjectTracker().GetAccumulatedDirtyNetObjects())
-	, DirtyObjectsToCopy(LocalNetRefHandleManager.GetDirtyObjectsToCopy())
+	, DirtyObjectsToQuantize(LocalNetRefHandleManager.GetDirtyObjectsToQuantize())
 {
 	GarbageCollectionAffectedObjects = MakeNetBitArrayView(ObjectReplicationBridge->GarbageCollectionAffectedObjects);
 
@@ -67,7 +67,7 @@ void FObjectPoller::CallPreUpdate(FInternalNetRefIndex ObjectIndex)
 	}
 }
 
-void FObjectPoller::PollObjects(const FNetBitArrayView& ObjectsConsideredForPolling)
+void FObjectPoller::PollAndCopyObjects(const FNetBitArrayView& ObjectsConsideredForPolling)
 {
 	FDirtyObjectsAccessor DirtyObjectsAccessor(ReplicationSystemInternal->GetDirtyNetObjectTracker());
 	DirtyObjectsThisFrame = DirtyObjectsAccessor.GetDirtyNetObjects();
@@ -76,7 +76,7 @@ void FObjectPoller::PollObjects(const FNetBitArrayView& ObjectsConsideredForPoll
 
 	if (IsIrisPushModelEnabled())
 	{
-		IRIS_PROFILER_SCOPE_VERBOSE(PollPushBased);
+		IRIS_PROFILER_SCOPE_VERBOSE(PollAndCopyPushBased);
 		ObjectsConsideredForPolling.ForAllSetBits([this](FInternalNetRefIndex Objectindex)
 		{
 			PushModelPollObject(Objectindex);
@@ -84,7 +84,7 @@ void FObjectPoller::PollObjects(const FNetBitArrayView& ObjectsConsideredForPoll
 	}
 	else
 	{
-		IRIS_PROFILER_SCOPE_VERBOSE(ForcePoll);
+		IRIS_PROFILER_SCOPE_VERBOSE(ForcePollAndCopy);
 		ObjectsConsideredForPolling.ForAllSetBits([this](FInternalNetRefIndex Objectindex)
 		{
 			ForcePollObject(Objectindex);
@@ -94,7 +94,7 @@ void FObjectPoller::PollObjects(const FNetBitArrayView& ObjectsConsideredForPoll
 	NetStatsContext = nullptr;
 }
 
-void FObjectPoller::PollSingleObject(FNetRefHandle Handle)
+void FObjectPoller::PollAndCopySingleObject(FNetRefHandle Handle)
 {
 	if (uint32 InternalObjectIndex = LocalNetRefHandleManager.GetInternalIndex(Handle))
 	{
@@ -129,7 +129,7 @@ void FObjectPoller::ForcePollObject(FInternalNetRefIndex ObjectIndex)
 	// Poll properties if the instance protocol requires it
 	if (EnumHasAnyFlags(ObjectData.InstanceProtocol->InstanceTraits, EReplicationInstanceProtocolTraits::NeedsPoll))
 	{
-		IRIS_PROFILER_SCOPE_VERBOSE(Poll);
+		IRIS_PROFILER_SCOPE_VERBOSE(PollAndCopy);
 		UE_NET_IRIS_STATS_TIMER(Timer, NetStatsContext);
 
 		const bool bIsGCAffectedObject = GarbageCollectionAffectedObjects.GetBit(ObjectIndex);
@@ -140,12 +140,12 @@ void FObjectPoller::ForcePollObject(FInternalNetRefIndex ObjectIndex)
 		PollOptions |= bIsGCAffectedObject ? EReplicationFragmentPollFlags::ForceRefreshCachedObjectReferencesAfterGC : EReplicationFragmentPollFlags::None;
 
 		const bool bWasAlreadyDirty = DirtyObjectsThisFrame.IsBitSet(ObjectIndex);
-		const bool bPollFoundDirty = FReplicationInstanceOperations::PollAndRefreshCachedPropertyData(ObjectData.InstanceProtocol, PollOptions);
+		const bool bPollFoundDirty = FReplicationInstanceOperations::PollAndCopyPropertyData(ObjectData.InstanceProtocol, PollOptions);
 		if (bWasAlreadyDirty || bPollFoundDirty)
 		{
 			UE_NET_IRIS_STATS_ADD_TIME_AND_COUNT_FOR_OBJECT(Timer, Poll, ObjectIndex);
 
-			DirtyObjectsToCopy.SetBit(ObjectIndex);
+			DirtyObjectsToQuantize.SetBit(ObjectIndex);
 			DirtyObjectsThisFrame.SetBit(ObjectIndex);
 		}
 		else
@@ -156,7 +156,7 @@ void FObjectPoller::ForcePollObject(FInternalNetRefIndex ObjectIndex)
 	}
 	else
 	{
-		DirtyObjectsToCopy.SetBit(ObjectIndex);
+		DirtyObjectsToQuantize.SetBit(ObjectIndex);
 		DirtyObjectsThisFrame.SetBit(ObjectIndex);
 	}
 }
@@ -180,7 +180,7 @@ void FObjectPoller::PushModelPollObject(FInternalNetRefIndex ObjectIndex)
 
 	if (bIsDirtyObject)
 	{
-		DirtyObjectsToCopy.SetBit(ObjectIndex);
+		DirtyObjectsToQuantize.SetBit(ObjectIndex);
 		DirtyObjectsThisFrame.SetBit(ObjectIndex);
 	}
 
@@ -209,14 +209,14 @@ void FObjectPoller::PushModelPollObject(FInternalNetRefIndex ObjectIndex)
 			// We need to do a poll if object is marked as dirty
 			EReplicationFragmentPollFlags PollOptions = EReplicationFragmentPollFlags::PollAllState;
 			PollOptions |= bIsGCAffectedObject ? EReplicationFragmentPollFlags::ForceRefreshCachedObjectReferencesAfterGC : EReplicationFragmentPollFlags::None;
-			bPollFoundDirty = FReplicationInstanceOperations::PollAndRefreshCachedPropertyData(InstanceProtocol, EReplicationFragmentTraits::None, PollOptions);
+			bPollFoundDirty = FReplicationInstanceOperations::PollAndCopyPropertyData(InstanceProtocol, EReplicationFragmentTraits::None, PollOptions);
 			++PollStats.PolledObjectCount;
 		}
 		else if (bIsGCAffectedObject)
 		{
 			// If this object might have been affected by GC, only refresh cached references
 			const EReplicationFragmentTraits RequiredTraits = EReplicationFragmentTraits::HasPushBasedDirtiness;
-			bPollFoundDirty = FReplicationInstanceOperations::PollAndRefreshCachedObjectReferences(InstanceProtocol, RequiredTraits);
+			bPollFoundDirty = FReplicationInstanceOperations::PollAndCopyObjectReferences(InstanceProtocol, RequiredTraits);
 			++PollStats.PolledReferencesObjectCount;
 		}
 	}
@@ -230,7 +230,7 @@ void FObjectPoller::PushModelPollObject(FInternalNetRefIndex ObjectIndex)
 		{
 			// Only states which has push based dirtiness need to be updated as the other states will be polled in full anyway.
 			const EReplicationFragmentTraits RequiredTraits = EReplicationFragmentTraits::HasPushBasedDirtiness;
-			bPollFoundDirty = FReplicationInstanceOperations::PollAndRefreshCachedObjectReferences(InstanceProtocol, RequiredTraits);
+			bPollFoundDirty = FReplicationInstanceOperations::PollAndCopyObjectReferences(InstanceProtocol, RequiredTraits);
 			++PollStats.PolledReferencesObjectCount;
 		}
 
@@ -240,7 +240,7 @@ void FObjectPoller::PushModelPollObject(FInternalNetRefIndex ObjectIndex)
 
 		// If the object is not new or dirty at this point we only need to poll non-push based fragments as we know that pushed based states have not been modified
 		const EReplicationFragmentTraits ExcludeTraits = (bIsDirtyObject || bWantsFullPoll) ? EReplicationFragmentTraits::None : EReplicationFragmentTraits::HasPushBasedDirtiness;
-		bPollFoundDirty |= FReplicationInstanceOperations::PollAndRefreshCachedPropertyData(InstanceProtocol, ExcludeTraits, PollOptions);
+		bPollFoundDirty |= FReplicationInstanceOperations::PollAndCopyPropertyData(InstanceProtocol, ExcludeTraits, PollOptions);
 		++PollStats.PolledObjectCount;
 	}
 
@@ -248,7 +248,7 @@ void FObjectPoller::PushModelPollObject(FInternalNetRefIndex ObjectIndex)
 	{
 		UE_NET_IRIS_STATS_ADD_TIME_AND_COUNT_FOR_OBJECT(Timer, Poll, ObjectIndex);
 
-		DirtyObjectsToCopy.SetBit(ObjectIndex);
+		DirtyObjectsToQuantize.SetBit(ObjectIndex);
 		DirtyObjectsThisFrame.SetBit(ObjectIndex);
 	}
 	else
diff --git a/Engine/Source/Runtime/Experimental/Iris/Core/Private/Iris/ReplicationSystem/Polling/ObjectPoller.h b/Engine/Source/Runtime/Experimental/Iris/Core/Private/Iris/ReplicationSystem/Polling/ObjectPoller.h
--- a/Engine/Source/Runtime/Experimental/Iris/Core/Private/Iris/ReplicationSystem/Polling/ObjectPoller.h
+++ b/Engine/Source/Runtime/Experimental/Iris/Core/Private/Iris/ReplicationSystem/Polling/ObjectPoller.h
@@ -55,11 +55,11 @@ public:
 	 */
 	void PreUpdatePass(const FNetBitArrayView& ObjectsConsideredForPolling);
 
-	/** Poll all the objects whose bit index is set in the array */
-	void PollObjects(const FNetBitArrayView& ObjectsConsideredForPolling);
+	/** Poll all the objects whose bit index is set in the array and copy any dirty data into ReplicationState buffers*/
+	void PollAndCopyObjects(const FNetBitArrayView& ObjectsConsideredForPolling);
 
 	/** Poll a single replicated object */
-	void PollSingleObject(FNetRefHandle Handle);
+	void PollAndCopySingleObject(FNetRefHandle Handle);
 
 private:
 
@@ -83,7 +83,7 @@ private:
 
 	const FNetBitArrayView AccumulatedDirtyObjects;
 
-	FNetBitArrayView DirtyObjectsToCopy;
+	FNetBitArrayView DirtyObjectsToQuantize;
 	FNetBitArrayView DirtyObjectsThisFrame;
 	FNetBitArrayView GarbageCollectionAffectedObjects;
 
diff --git a/Engine/Source/Runtime/Experimental/Iris/Core/Private/Iris/ReplicationSystem/ReplicationBridge.cpp b/Engine/Source/Runtime/Experimental/Iris/Core/Private/Iris/ReplicationSystem/ReplicationBridge.cpp
--- a/Engine/Source/Runtime/Experimental/Iris/Core/Private/Iris/ReplicationSystem/ReplicationBridge.cpp
+++ b/Engine/Source/Runtime/Experimental/Iris/Core/Private/Iris/ReplicationSystem/ReplicationBridge.cpp
@@ -702,10 +702,10 @@ void UReplicationBridge::InternalFlushStateData(UE::Net::FNetSerializationContex
 		CallPreSendUpdateSingleHandle(ObjectData.RefHandle);
 	} 
 
-	FReplicationInstanceOperationsInternal::CopyObjectStateData(ChangeMaskWriter, ChangeMaskCache, *NetRefHandleManager, SerializationContext, InternalObjectIndex);
+	FReplicationInstanceOperationsInternal::QuantizeObjectStateData(ChangeMaskWriter, ChangeMaskCache, *NetRefHandleManager, SerializationContext, InternalObjectIndex);
 
-	// Clear the copy flag since it was done directly here.
-	NetRefHandleManager->GetDirtyObjectsToCopy().ClearBit(InternalObjectIndex);
+	// Clear the quantize flag since it was done directly here.
+	NetRefHandleManager->GetDirtyObjectsToQuantize().ClearBit(InternalObjectIndex);
 
 	// $IRIS TODO:  Should we also clear the DirtyTracker flags for this flushed object ?
 }
@@ -797,7 +797,7 @@ void UReplicationBridge::InternalTearOff(FNetRefHandle Handle)
 
 		if (ObjectData.InstanceProtocol && ObjectData.Protocol->InternalTotalSize > 0U)
 		{
-			FReplicationInstanceOperationsInternal::CopyObjectStateData(ChangeMaskWriter, ChangeMaskCache, *NetRefHandleManager, SerializationContext, InternalObjectIndex);
+			FReplicationInstanceOperationsInternal::QuantizeObjectStateData(ChangeMaskWriter, ChangeMaskCache, *NetRefHandleManager, SerializationContext, InternalObjectIndex);
 		}
 		else
 		{
diff --git a/Engine/Source/Runtime/Experimental/Iris/Core/Private/Iris/ReplicationSystem/ReplicationOperations.cpp b/Engine/Source/Runtime/Experimental/Iris/Core/Private/Iris/ReplicationSystem/ReplicationOperations.cpp
--- a/Engine/Source/Runtime/Experimental/Iris/Core/Private/Iris/ReplicationSystem/ReplicationOperations.cpp
+++ b/Engine/Source/Runtime/Experimental/Iris/Core/Private/Iris/ReplicationSystem/ReplicationOperations.cpp
@@ -551,7 +551,7 @@ void FReplicationStateOperations::DeserializeDeltaWithMask(FNetSerializationCont
 	}
 }
 
-bool FReplicationInstanceOperations::PollAndRefreshCachedPropertyData(const FReplicationInstanceProtocol* InstanceProtocol, EReplicationFragmentTraits ExcludeTraits, EReplicationFragmentPollFlags PollOptions)
+bool FReplicationInstanceOperations::PollAndCopyPropertyData(const FReplicationInstanceProtocol* InstanceProtocol, EReplicationFragmentTraits ExcludeTraits, EReplicationFragmentPollFlags PollOptions)
 {
 	bool bIsStateDirty = false;
 	FReplicationFragment* const * Fragments = InstanceProtocol->Fragments;
@@ -571,12 +571,12 @@ bool FReplicationInstanceOperations::PollAndRefreshCachedPropertyData(const FRep
 	return bIsStateDirty;
 }
 
-bool FReplicationInstanceOperations::PollAndRefreshCachedPropertyData(const FReplicationInstanceProtocol* InstanceProtocol, EReplicationFragmentPollFlags PollOptions)
+bool FReplicationInstanceOperations::PollAndCopyPropertyData(const FReplicationInstanceProtocol* InstanceProtocol, EReplicationFragmentPollFlags PollOptions)
 {
-	return PollAndRefreshCachedPropertyData(InstanceProtocol, EReplicationFragmentTraits::None, PollOptions);
+	return PollAndCopyPropertyData(InstanceProtocol, EReplicationFragmentTraits::None, PollOptions);
 }
 
-bool FReplicationInstanceOperations::PollAndRefreshCachedObjectReferences(const FReplicationInstanceProtocol* InstanceProtocol, EReplicationFragmentTraits RequiredTraits)
+bool FReplicationInstanceOperations::PollAndCopyObjectReferences(const FReplicationInstanceProtocol* InstanceProtocol, EReplicationFragmentTraits RequiredTraits)
 {
 	bool bIsStateDirty = false;
 	const EReplicationFragmentTraits ExpectedTraits = EReplicationFragmentTraits::HasObjectReference | RequiredTraits;
@@ -595,9 +595,9 @@ bool FReplicationInstanceOperations::PollAndRefreshCachedObjectReferences(const
 	return bIsStateDirty;
 }
 
-void FReplicationInstanceOperations::CopyAndQuantize(FNetSerializationContext& Context, uint8* DstObjectStateBuffer, FNetBitStreamWriter* OutChangeMaskWriter, const FReplicationInstanceProtocol* InstanceProtocol, const FReplicationProtocol* Protocol)
+void FReplicationInstanceOperations::Quantize(FNetSerializationContext& Context, uint8* DstObjectStateBuffer, FNetBitStreamWriter* OutChangeMaskWriter, const FReplicationInstanceProtocol* InstanceProtocol, const FReplicationProtocol* Protocol)
 {
-	IRIS_PROFILER_SCOPE_VERBOSE(CopyAndQuantize);
+	IRIS_PROFILER_SCOPE_VERBOSE(Quantize);
 
 	check(InstanceProtocol && InstanceProtocol->FragmentCount == Protocol->ReplicationStateCount);
 
@@ -635,9 +635,9 @@ void FReplicationInstanceOperations::CopyAndQuantize(FNetSerializationContext& C
 	}
 }
 
-void FReplicationInstanceOperations::CopyAndQuantizeIfDirty(FNetSerializationContext& Context, uint8* DstObjectStateBuffer, FNetBitStreamWriter* OutChangeMaskWriter, const FReplicationInstanceProtocol* InstanceProtocol, const FReplicationProtocol* Protocol)
+void FReplicationInstanceOperations::QuantizeIfDirty(FNetSerializationContext& Context, uint8* DstObjectStateBuffer, FNetBitStreamWriter* OutChangeMaskWriter, const FReplicationInstanceProtocol* InstanceProtocol, const FReplicationProtocol* Protocol)
 {
-	IRIS_PROFILER_SCOPE_VERBOSE(CopyAndQuantizeIfDirty);
+	IRIS_PROFILER_SCOPE_VERBOSE(QuantizeIfDirty);
 
 	check(InstanceProtocol && InstanceProtocol->FragmentCount == Protocol->ReplicationStateCount);
 
diff --git a/Engine/Source/Runtime/Experimental/Iris/Core/Private/Iris/ReplicationSystem/ReplicationOperationsInternal.cpp b/Engine/Source/Runtime/Experimental/Iris/Core/Private/Iris/ReplicationSystem/ReplicationOperationsInternal.cpp
--- a/Engine/Source/Runtime/Experimental/Iris/Core/Private/Iris/ReplicationSystem/ReplicationOperationsInternal.cpp
+++ b/Engine/Source/Runtime/Experimental/Iris/Core/Private/Iris/ReplicationSystem/ReplicationOperationsInternal.cpp
@@ -69,13 +69,13 @@ void FReplicationInstanceOperationsInternal::UnbindInstanceProtocol(FReplication
 	}
 }
 
-uint32 FReplicationInstanceOperationsInternal::CopyObjectStateData(FNetBitStreamWriter& ChangeMaskWriter, FChangeMaskCache& Cache, FNetRefHandleManager& NetRefHandleManager, FNetSerializationContext& SerializationContext, uint32 InternalIndex)
+uint32 FReplicationInstanceOperationsInternal::QuantizeObjectStateData(FNetBitStreamWriter& ChangeMaskWriter, FChangeMaskCache& Cache, FNetRefHandleManager& NetRefHandleManager, FNetSerializationContext& SerializationContext, uint32 InternalIndex)
 {
 	if (NetRefHandleManager.IsScopableIndex(InternalIndex))
 	{
 		FNetRefHandleManager::FReplicatedObjectData& Object = NetRefHandleManager.GetReplicatedObjectDataNoCheck(InternalIndex);
 
-		// We cannot copy state data for zero sized objects or objects that no longer has an instance protocol.
+		// We cannot quantize state data for zero sized objects or objects that no longer has an instance protocol.
 		if (Object.InstanceProtocol && Object.Protocol->InternalTotalSize > 0U)
 		{
 			IRIS_PROFILER_PROTOCOL_NAME(Object.Protocol->DebugName->Name);
@@ -88,7 +88,7 @@ uint32 FReplicationInstanceOperationsInternal::CopyObjectStateData(FNetBitStream
 			const bool bUseFullCopyAndQuantize = bCVarForceFullCopyAndQuantize || Object.bNeedsFullCopyAndQuantize;
 			Object.bNeedsFullCopyAndQuantize = 0U;
 
-			// Copy dirty state
+			// Quantize dirty state
 			{
 				// Add entry to cache
 				FChangeMaskCache::FCachedInfo& Info = Cache.AddChangeMaskForObject(InternalIndex, Object.Protocol->ChangeMaskBitCount);
@@ -99,11 +99,11 @@ uint32 FReplicationInstanceOperationsInternal::CopyObjectStateData(FNetBitStream
 				//UE_LOG(LogIris, Log, TEXT("Copying state data for ( InternalIndex: %u ) with NetRefHandle (Id=%u)"), InternalIndex, Object.RefHandle.GetId());
 				if (bUseFullCopyAndQuantize)
 				{
-					FReplicationInstanceOperations::CopyAndQuantize(SerializationContext, NetRefHandleManager.GetReplicatedObjectStateBufferNoCheck(InternalIndex), &ChangeMaskWriter, Object.InstanceProtocol, Object.Protocol);					
+					FReplicationInstanceOperations::Quantize(SerializationContext, NetRefHandleManager.GetReplicatedObjectStateBufferNoCheck(InternalIndex), &ChangeMaskWriter, Object.InstanceProtocol, Object.Protocol);					
 				}
 				else
 				{
-					FReplicationInstanceOperations::CopyAndQuantizeIfDirty(SerializationContext, NetRefHandleManager.GetReplicatedObjectStateBufferNoCheck(InternalIndex), &ChangeMaskWriter, Object.InstanceProtocol, Object.Protocol);
+					FReplicationInstanceOperations::QuantizeIfDirty(SerializationContext, NetRefHandleManager.GetReplicatedObjectStateBufferNoCheck(InternalIndex), &ChangeMaskWriter, Object.InstanceProtocol, Object.Protocol);
 				}
 
 				Info.bHasDirtyChangeMask = MakeNetBitArrayView(ChangeMaskData, ChangeMaskByteCount * 8U).FindFirstOne() != FNetBitArrayView::InvalidIndex;
diff --git a/Engine/Source/Runtime/Experimental/Iris/Core/Private/Iris/ReplicationSystem/ReplicationOperationsInternal.h b/Engine/Source/Runtime/Experimental/Iris/Core/Private/Iris/ReplicationSystem/ReplicationOperationsInternal.h
--- a/Engine/Source/Runtime/Experimental/Iris/Core/Private/Iris/ReplicationSystem/ReplicationOperationsInternal.h
+++ b/Engine/Source/Runtime/Experimental/Iris/Core/Private/Iris/ReplicationSystem/ReplicationOperationsInternal.h
@@ -35,8 +35,8 @@ struct FReplicationInstanceOperationsInternal
 	/** A call to this function will clear the index of the handle into the external statebuffer contained in the Instance Protocol */
 	static IRISCORE_API void UnbindInstanceProtocol(FReplicationInstanceProtocol* InstanceProtocol, const FReplicationProtocol* Protocol);
 
-	/** Copy state data for a single instance */
-	static IRISCORE_API uint32 CopyObjectStateData(FNetBitStreamWriter& ChangeMaskWriter, FChangeMaskCache& Cache, FNetRefHandleManager& NetRefHandleManager, FNetSerializationContext& SerializationContext, uint32 InternalIndex);
+	/** Quantize state data for a single instance */
+	static IRISCORE_API uint32 QuantizeObjectStateData(FNetBitStreamWriter& ChangeMaskWriter, FChangeMaskCache& Cache, FNetRefHandleManager& NetRefHandleManager, FNetSerializationContext& SerializationContext, uint32 InternalIndex);
 
 	/** Reset instance state dirtiness. */
 	static IRISCORE_API void ResetObjectStateDirtiness(FNetRefHandleManager& NetRefHandleManager, uint32 InternalIndex);
diff --git a/Engine/Source/Runtime/Experimental/Iris/Core/Private/Iris/ReplicationSystem/ReplicationSystem.cpp b/Engine/Source/Runtime/Experimental/Iris/Core/Private/Iris/ReplicationSystem/ReplicationSystem.cpp
--- a/Engine/Source/Runtime/Experimental/Iris/Core/Private/Iris/ReplicationSystem/ReplicationSystem.cpp
+++ b/Engine/Source/Runtime/Experimental/Iris/Core/Private/Iris/ReplicationSystem/ReplicationSystem.cpp
@@ -395,15 +395,15 @@ public:
 		ValidConnections.ForAllSetBits(UpdateDirtyChangeMasks);
 	}
 
-	void CopyDirtyStateData()
+	void QuantizeDirtyStateData()
 	{
-		IRIS_PROFILER_SCOPE(FReplicationSystem_CopyDirtyStateData);
+		IRIS_PROFILER_SCOPE(FReplicationSystem_QuantizeDirtyStateData);
 		LLM_SCOPE_BYTAG(IrisState);
 
 		FNetRefHandleManager& NetRefHandleManager = ReplicationSystemInternal.GetNetRefHandleManager();
 		FChangeMaskCache& Cache = ReplicationSystemInternal.GetChangeMaskCache();
 
-		uint32 CopiedObjectCount = 0;
+		uint32 QuantizedObjectCount = 0;
 		
 		// Prepare cache
 		constexpr uint32 ReservedIndexCount = 2048;
@@ -421,18 +421,18 @@ public:
 		SerializationContext.SetNetStatsContext(ReplicationSystemInternal.GetNetTypeStats().GetNetStatsContext());
 
 		// Copy the state data of objects that were dirty this frame.
-		FNetBitArrayView DirtyObjectsToCopy = NetRefHandleManager.GetDirtyObjectsToCopy();
+		FNetBitArrayView DirtyObjectsToQuantize = NetRefHandleManager.GetDirtyObjectsToQuantize();
 
-		auto CopyFunction = [&ChangeMaskWriter, &Cache, &NetRefHandleManager, &CopiedObjectCount, &SerializationContext](uint32 DirtyIndex)
+		auto QuantizeFunction = [&ChangeMaskWriter, &Cache, &NetRefHandleManager, &QuantizedObjectCount, &SerializationContext](uint32 DirtyIndex)
 		{
-			CopiedObjectCount += FReplicationInstanceOperationsInternal::CopyObjectStateData(ChangeMaskWriter, Cache, NetRefHandleManager, SerializationContext, DirtyIndex);
+			QuantizedObjectCount += FReplicationInstanceOperationsInternal::QuantizeObjectStateData(ChangeMaskWriter, Cache, NetRefHandleManager, SerializationContext, DirtyIndex);
 		};
 
-		DirtyObjectsToCopy.ForAllSetBits(CopyFunction);
-		DirtyObjectsToCopy.Reset();
+		DirtyObjectsToQuantize.ForAllSetBits(QuantizeFunction);
+		DirtyObjectsToQuantize.Reset();
 
 		const uint32 ReplicationSystemId = ReplicationSystem->GetId();
-		UE_NET_TRACE_FRAME_STATSCOUNTER(ReplicationSystemId, ReplicationSystem.CopiedObjectCount, CopiedObjectCount, ENetTraceVerbosity::Trace);
+		UE_NET_TRACE_FRAME_STATSCOUNTER(ReplicationSystemId, ReplicationSystem.QuantizedObjectCount, QuantizedObjectCount, ENetTraceVerbosity::Trace);
 	}
 
 	void ResetObjectStateDirtiness()
@@ -686,8 +686,8 @@ void UReplicationSystem::PreSendUpdate(const FSendUpdateParams& Params)
 			// Update conditionals
 			Impl->UpdateConditionals();
 
-			// Copy dirty state data. We need this to happen before both filtering and prioritization
-			Impl->CopyDirtyStateData();
+			// Quantize dirty state data. We need this to happen before both filtering and prioritization
+			Impl->QuantizeDirtyStateData();
 
 			// We must process all attachments to objects going out of scope before we update the scope
 			Impl->ProcessNetObjectAttachmentSendQueue(FNetBlobManager::EProcessMode::ProcessObjectsGoingOutOfScope);
diff --git a/Engine/Source/Runtime/Experimental/Iris/Core/Public/Iris/ReplicationSystem/ObjectReplicationBridge.h b/Engine/Source/Runtime/Experimental/Iris/Core/Public/Iris/ReplicationSystem/ObjectReplicationBridge.h
--- a/Engine/Source/Runtime/Experimental/Iris/Core/Public/Iris/ReplicationSystem/ObjectReplicationBridge.h
+++ b/Engine/Source/Runtime/Experimental/Iris/Core/Public/Iris/ReplicationSystem/ObjectReplicationBridge.h
@@ -327,8 +327,8 @@ private:
 	/** Find any new subobjects created inside PreUpdate and ensure they will be replicated this frame */
 	void ReconcileNewSubObjects(UE::Net::FNetBitArrayView ObjectsConsideredForPolling);
 
-	/** Poll all objects set in the list */
-	void Poll(const UE::Net::FNetBitArrayView ObjectsConsideredForPolling);
+	/** Poll all objects in the list and copy the dirty source data into the ReplicationState buffers*/
+	void PollAndCopy(const UE::Net::FNetBitArrayView ObjectsConsideredForPolling);
 
 	/** Remove mapping between handle and object instance. */
 	void UnregisterInstance(FNetRefHandle RefHandle);
diff --git a/Engine/Source/Runtime/Experimental/Iris/Core/Public/Iris/ReplicationSystem/ReplicationOperations.h b/Engine/Source/Runtime/Experimental/Iris/Core/Public/Iris/ReplicationSystem/ReplicationOperations.h
--- a/Engine/Source/Runtime/Experimental/Iris/Core/Public/Iris/ReplicationSystem/ReplicationOperations.h
+++ b/Engine/Source/Runtime/Experimental/Iris/Core/Public/Iris/ReplicationSystem/ReplicationOperations.h
@@ -75,13 +75,13 @@ struct FReplicationStateOperations
 struct FReplicationInstanceOperations
 {
 	/** Update all registered Fragments that updates dirtiness by polling, except for those with any of the ExcludeTraits. Returns true if a polled state is dirty. */
-	static IRISCORE_API bool PollAndRefreshCachedPropertyData(const FReplicationInstanceProtocol* InstanceProtocol, EReplicationFragmentTraits ExcludeTraits, EReplicationFragmentPollFlags PollOptions = EReplicationFragmentPollFlags::PollAllState);
+	static IRISCORE_API bool PollAndCopyPropertyData(const FReplicationInstanceProtocol* InstanceProtocol, EReplicationFragmentTraits ExcludeTraits, EReplicationFragmentPollFlags PollOptions = EReplicationFragmentPollFlags::PollAllState);
 
 	/** Update all registered Fragments that updates dirtiness by polling.  Returns true if a polled state is dirty. */
-	static IRISCORE_API bool PollAndRefreshCachedPropertyData(const FReplicationInstanceProtocol* InstanceProtocol, EReplicationFragmentPollFlags PollOptions = EReplicationFragmentPollFlags::PollAllState);
+	static IRISCORE_API bool PollAndCopyPropertyData(const FReplicationInstanceProtocol* InstanceProtocol, EReplicationFragmentPollFlags PollOptions = EReplicationFragmentPollFlags::PollAllState);
 
 	/** Update object references in fragments that has object references and additional required traits. Returns true if a polled state is dirty. */
-	static IRISCORE_API bool PollAndRefreshCachedObjectReferences(const FReplicationInstanceProtocol* InstanceProtocol, EReplicationFragmentTraits RequiredTraits);
+	static IRISCORE_API bool PollAndCopyObjectReferences(const FReplicationInstanceProtocol* InstanceProtocol, EReplicationFragmentTraits RequiredTraits);
 
 	/**
 	 * Quantize the state for a replicated object with a given InstanceProtocol using the ReplicationProtocol. 
@@ -89,7 +89,7 @@ struct FReplicationInstanceOperations
 	 * Changemasks will be written to the ChangeMaskWriter. Dirtiness will not be reset.
 	 * @see ResetDirtiness
 	 */
-	static IRISCORE_API void CopyAndQuantize(FNetSerializationContext& Context, uint8* DstObjectStateBuffer, FNetBitStreamWriter* ChangeMaskWriter, const FReplicationInstanceProtocol* InstanceProtocol, const FReplicationProtocol* Protocol);
+	static IRISCORE_API void Quantize(FNetSerializationContext& Context, uint8* DstObjectStateBuffer, FNetBitStreamWriter* ChangeMaskWriter, const FReplicationInstanceProtocol* InstanceProtocol, const FReplicationProtocol* Protocol);
 
 	/**
 	 * Quantize the state for a replicated object with a given InstanceProtocol using the ReplicationProtocol.
@@ -98,7 +98,7 @@ struct FReplicationInstanceOperations
 	 * This variant will only Quantize States marked as dirty
 	 * @see ResetDirtiness
 	 */
-	static IRISCORE_API void CopyAndQuantizeIfDirty(FNetSerializationContext& Context, uint8* DstObjectStateBuffer, FNetBitStreamWriter* ChangeMaskWriter, const FReplicationInstanceProtocol* InstanceProtocol, const FReplicationProtocol* Protocol);
+	static IRISCORE_API void QuantizeIfDirty(FNetSerializationContext& Context, uint8* DstObjectStateBuffer, FNetBitStreamWriter* ChangeMaskWriter, const FReplicationInstanceProtocol* InstanceProtocol, const FReplicationProtocol* Protocol);
 
 	/** Resets dirty tracking stored with the protocol, such as changemasks and init state dirtiness. */
 	static IRISCORE_API void ResetDirtiness(const FReplicationInstanceProtocol* InstanceProtocol, const FReplicationProtocol* Protocol);
