Iris - Fix issue where state changes accrued while going out of scope and later readded to scope was not dirtied correctly

diff --git a/Engine/Plugins/Runtime/ReplicationSystemTestPlugin/Source/Private/Tests/ReplicationSystem/TestCancelPendingDestroy.cpp b/Engine/Plugins/Runtime/ReplicationSystemTestPlugin/Source/Private/Tests/ReplicationSystem/TestCancelPendingDestroy.cpp
--- a/Engine/Plugins/Runtime/ReplicationSystemTestPlugin/Source/Private/Tests/ReplicationSystem/TestCancelPendingDestroy.cpp
+++ b/Engine/Plugins/Runtime/ReplicationSystemTestPlugin/Source/Private/Tests/ReplicationSystem/TestCancelPendingDestroy.cpp
@@ -590,6 +590,74 @@ UE_NET_TEST_FIXTURE(FTestCancelPendingDestroyFixture, TestCancelingCancelPending
 	UE_NET_ASSERT_EQ(Client->GetReplicationBridge()->GetReplicatedObject(ServerSubObject->NetRefHandle), nullptr);
 }
 
+UE_NET_TEST_FIXTURE(FTestCancelPendingDestroyFixture, TestCancelPendingDestroyFromCreatedDoesNotMissChanges)
+{
+	// Add a client
+	FReplicationSystemTestClient* Client0 = CreateClient();
+
+	// Add second client, with not filtering to keep object in scope
+	FReplicationSystemTestClient* Client1 = CreateClient();
+
+	// Spawn object on server
+	UTestReplicatedIrisObject* ServerObject = Server->CreateObject(UTestReplicatedIrisObject::FComponents());
+
+	// Write and send packet
+	Server->PreSendUpdate();
+	Server->SendAndDeliverTo(Client0, DeliverPacket);
+	Server->SendAndDeliverTo(Client1, DeliverPacket);
+	Server->PostSendUpdate();
+
+	UTestReplicatedIrisObject* Client0Object = Cast<UTestReplicatedIrisObject>(Client0->GetReplicationBridge()->GetReplicatedObject(ServerObject->NetRefHandle));
+	UTestReplicatedIrisObject* Client1Object = Cast<UTestReplicatedIrisObject>(Client1->GetReplicationBridge()->GetReplicatedObject(ServerObject->NetRefHandle));
+
+	// Verify that the object was replicated to all clients
+	UE_NET_ASSERT_NE(Client0Object, nullptr);
+	UE_NET_ASSERT_EQ(Client0Object->IntA, ServerObject->IntA);
+	UE_NET_ASSERT_NE(Client1Object, nullptr);
+	UE_NET_ASSERT_EQ(Client1Object->IntA, ServerObject->IntA);
+
+	// Filter out object to cause a PendingDestroy for Client0
+	FNetObjectGroupHandle ExclusionGroupHandle = Server->ReplicationSystem->CreateGroup();
+	Server->ReplicationSystem->AddToGroup(ExclusionGroupHandle, ServerObject->NetRefHandle);
+	Server->ReplicationSystem->AddExclusionFilterGroup(ExclusionGroupHandle);
+
+	Server->ReplicationSystem->SetGroupFilterStatus(ExclusionGroupHandle, Client0->ConnectionIdOnServer, ENetFilterStatus::Disallow);
+	Server->ReplicationSystem->SetGroupFilterStatus(ExclusionGroupHandle, Client1->ConnectionIdOnServer, ENetFilterStatus::Allow);
+	
+	Server->PreSendUpdate();
+	Server->SendTo(Client0);
+	Server->PostSendUpdate();
+
+	// Mark dirty
+	ServerObject->IntA = 3;
+
+	Server->PreSendUpdate();
+	Server->SendAndDeliverTo(Client1, DeliverPacket);
+	Server->PostSendUpdate();
+
+	// Allow replication again, to trigger cancel
+	Server->ReplicationSystem->SetGroupFilterStatus(ExclusionGroupHandle, Client0->ConnectionIdOnServer, ENetFilterStatus::Allow);
+
+	// Drop the packet with the destroy
+	Server->DeliverTo(Client0, false);
+
+	// Write and send a packet
+	Server->PreSendUpdate();
+	Server->SendAndDeliverTo(Client0, DeliverPacket);
+	Server->SendAndDeliverTo(Client1, DeliverPacket);
+	Server->PostSendUpdate();
+
+	Client0Object = Cast<UTestReplicatedIrisObject>(Client0->GetReplicationBridge()->GetReplicatedObject(ServerObject->NetRefHandle));
+	Client1Object = Cast<UTestReplicatedIrisObject>(Client1->GetReplicationBridge()->GetReplicatedObject(ServerObject->NetRefHandle));
+
+	// Verify that the cancel worked and did not miss changed state
+	UE_NET_ASSERT_NE(Client0Object, nullptr);
+	UE_NET_ASSERT_EQ(Client0Object->IntA, ServerObject->IntA);
+	UE_NET_ASSERT_NE(Client1Object, nullptr);
+	UE_NET_ASSERT_EQ(Client1Object->IntA, ServerObject->IntA);
+}
+
+
 UE_NET_TEST_FIXTURE(FTestCancelPendingDestroyFixture, TestCancelingCancelPendingSubObjectDestroyDuringWaitOnDestroyConfirmationWithPacketLoss)
 {
 	// Add a client
diff --git a/Engine/Plugins/Runtime/ReplicationSystemTestPlugin/Source/Private/Tests/ReplicationSystem/TestFlushBeforeDestroy.cpp b/Engine/Plugins/Runtime/ReplicationSystemTestPlugin/Source/Private/Tests/ReplicationSystem/TestFlushBeforeDestroy.cpp
--- a/Engine/Plugins/Runtime/ReplicationSystemTestPlugin/Source/Private/Tests/ReplicationSystem/TestFlushBeforeDestroy.cpp
+++ b/Engine/Plugins/Runtime/ReplicationSystemTestPlugin/Source/Private/Tests/ReplicationSystem/TestFlushBeforeDestroy.cpp
@@ -652,5 +652,76 @@ UE_NET_TEST_FIXTURE(FTestFlushBeforeDestroyFixture, TestImmediateTearOffWithSubO
 	UE_NET_ASSERT_TRUE(Cast<UTestReplicatedIrisObject>(Client->GetReplicationBridge()->GetReplicatedObject(ServerObject->NetRefHandle)) == nullptr);
 }
 
+// Test to recreate a path where we cancel destroy for object pending flush
+UE_NET_TEST_FIXTURE(FTestFlushBeforeDestroyFixture, TestCancelPendingDestroyWaitOnFlushDoesNotMissChanges)
+{
+	UReplicationSystem* ReplicationSystem = Server->ReplicationSystem;
+
+	// Add a client
+	FReplicationSystemTestClient* Client0 = CreateClient();
+	FReplicationSystemTestClient* Client1 = CreateClient();
+
+	RegisterNetBlobHandlers(Client0);
+	RegisterNetBlobHandlers(Client1);
+
+	// Spawn object on server
+	UTestReplicatedIrisObject* ServerObject = Server->CreateObject(0,0);
+
+	// Send and deliver packet
+	Server->PreSendUpdate();
+	Server->SendAndDeliverTo(Client0, true, TEXT("Create Objects"));
+	Server->PostSendUpdate();
+
+	// Store Pointer to objects
+	UTestReplicatedIrisObject* ClientObject = Cast<UTestReplicatedIrisObject>(Client0->GetReplicationBridge()->GetReplicatedObject(ServerObject->NetRefHandle));
+	UE_NET_ASSERT_TRUE(ClientObject != nullptr);
+	UE_NET_ASSERT_EQ(ServerObject->IntA, ClientObject->IntA);
+
+	// Create attachment to force flush behavior by having a rpc in flight
+	{
+		constexpr uint32 PayloadBitCount = 24;
+		const TRefCountPtr<FNetObjectAttachment>& Attachment = MockNetObjectAttachmentHandler->CreateReliableNetObjectAttachment(PayloadBitCount);
+		FNetObjectReference AttachmentTarget = FObjectReferenceCache::MakeNetObjectReference(ServerObject->NetRefHandle);
+		Server->GetReplicationSystem()->QueueNetObjectAttachment(Client0->ConnectionIdOnServer, AttachmentTarget, Attachment);
+	}
+
+	Server->PreSendUpdate();
+	Server->SendTo(Client0, TEXT("Attachment"));
+	Server->PostSendUpdate();
+
+	// Filter out object to cause a flush for Client0
+	FNetObjectGroupHandle ExclusionGroupHandle = Server->ReplicationSystem->CreateGroup();
+	Server->ReplicationSystem->AddToGroup(ExclusionGroupHandle, ServerObject->NetRefHandle);
+	Server->ReplicationSystem->AddExclusionFilterGroup(ExclusionGroupHandle);
+
+	Server->ReplicationSystem->SetGroupFilterStatus(ExclusionGroupHandle, Client0->ConnectionIdOnServer, ENetFilterStatus::Disallow);
+	Server->ReplicationSystem->SetGroupFilterStatus(ExclusionGroupHandle, Client1->ConnectionIdOnServer, ENetFilterStatus::Allow);
+
+	Server->PreSendUpdate();
+	Server->SendTo(Client0, TEXT("Out of scope"));
+	Server->PostSendUpdate();
+
+	// Modify the value of object only
+	++ServerObject->IntA ;
+
+	// Trigger poll + propagate of state
+	Server->PreSendUpdate();
+	Server->PostSendUpdate();
+	
+	// Trigger WaitOnFlush -> Created
+	Server->ReplicationSystem->SetGroupFilterStatus(ExclusionGroupHandle, Client0->ConnectionIdOnServer, ENetFilterStatus::Allow);
+
+	// Drop some packets to stay in state
+	Server->DeliverTo(Client0, false);
+	Server->DeliverTo(Client0, false);
+
+	// Do a normal update, should send state changed that occurred while we where in pending flush state
+	Server->PreSendUpdate();
+	Server->SendAndDeliverTo(Client0, DeliverPacket, TEXT("Expected state"));
+	Server->PostSendUpdate();
+
+	// Verify that ClientObject is torn-off and that the final state was applied
+	UE_NET_ASSERT_EQ(ServerObject->IntA, ClientObject->IntA);
+}
 
 }
diff --git a/Engine/Source/Runtime/Experimental/Iris/Core/Private/Iris/ReplicationSystem/ReplicationWriter.cpp b/Engine/Source/Runtime/Experimental/Iris/Core/Private/Iris/ReplicationSystem/ReplicationWriter.cpp
--- a/Engine/Source/Runtime/Experimental/Iris/Core/Private/Iris/ReplicationSystem/ReplicationWriter.cpp
+++ b/Engine/Source/Runtime/Experimental/Iris/Core/Private/Iris/ReplicationSystem/ReplicationWriter.cpp
@@ -652,6 +652,11 @@ void FReplicationWriter::UpdateScope(const FNetBitArrayView& UpdatedScope)
 				ObjectsPendingDestroy.ClearBit(Index);
 				Info.FlushFlags = GetDefaultFlushFlags();
 				SetState(Index, EReplicatedObjectState::Created);
+
+				// If we have accumulated changes while WaitingOnFlush, we should send them now
+				Info.SubObjectPendingDestroy = 0U;
+				Info.HasDirtyChangeMask |= FNetBitArrayView(Info.GetChangeMaskStoragePointer(), Info.ChangeMaskBitCount).IsAnyBitSet();
+				ObjectsWithDirtyChanges.SetBitValue(Index, Info.HasDirtyChangeMask);
 			}
 		}
 		else if (State == EReplicatedObjectState::WaitOnDestroyConfirmation || State == EReplicatedObjectState::CancelPendingDestroy)
@@ -789,13 +794,18 @@ void FReplicationWriter::InternalUpdateDirtyChangeMasks(const FChangeMaskCache&
 
 	for (const auto& Entry : CachedChangeMasks.Indices)
 	{
-		if (!ObjectsInScope.GetBit(Entry.InternalIndex))
+		FReplicationInfo& Info = ReplicatedObjects[Entry.InternalIndex];
+		if (Info.GetState() == EReplicatedObjectState::Invalid)
 		{
 			continue;
 		}
 
-		MarkObjectDirty(Entry.InternalIndex, "UpdateDirtyChangeMasks");
-		FReplicationInfo& Info = ReplicatedObjects[Entry.InternalIndex];
+		// We want to accumulate dirty changes even if we are going out of scope in case we get readded to scope before repliation has ended.
+		const bool bMarkScopedObjectDirty = ObjectsInScope.GetBit(Entry.InternalIndex);
+		if (bMarkScopedObjectDirty)
+		{
+			MarkObjectDirty(Entry.InternalIndex, "UpdateDirtyChangeMasks");
+		}
 
 		if (Entry.bMarkSubObjectOwnerDirty == 0U)
 		{		
@@ -817,7 +827,7 @@ void FReplicationWriter::InternalUpdateDirtyChangeMasks(const FChangeMaskCache&
 				Changes.Combine(UpdatedChanges, FNetBitArrayView::OrOp);
 
 				// Mark changemask as dirty
-				Info.HasDirtyChangeMask = 1U;
+				Info.HasDirtyChangeMask = bMarkScopedObjectDirty ? 1U : 0U;
 			}
 		}
 		else
