Iris - Don't destroy object until creation has been acked.

diff --git a/Engine/Plugins/Runtime/ReplicationSystemTestPlugin/Source/Private/Tests/ReplicationSystem/TestCancelPendingDestroy.cpp b/Engine/Plugins/Runtime/ReplicationSystemTestPlugin/Source/Private/Tests/ReplicationSystem/TestCancelPendingDestroy.cpp
--- a/Engine/Plugins/Runtime/ReplicationSystemTestPlugin/Source/Private/Tests/ReplicationSystem/TestCancelPendingDestroy.cpp
+++ b/Engine/Plugins/Runtime/ReplicationSystemTestPlugin/Source/Private/Tests/ReplicationSystem/TestCancelPendingDestroy.cpp
@@ -167,6 +167,58 @@ UE_NET_TEST_FIXTURE(FTestCancelPendingDestroyFixture, TestCancelPendingDestroyDu
 	UE_NET_ASSERT_EQ(ClientObject->IntA, ServerObject->IntA);
 }
 
+UE_NET_TEST_FIXTURE(FTestCancelPendingDestroyFixture, TestCancelPendingDestroyDuringWaitOnDestroyConfirmationWithInitialPacketLoss)
+{
+	// Add a client
+	FReplicationSystemTestClient* Client = CreateClient();
+
+	// Spawn object on server
+	UTestReplicatedIrisObject* ServerObject = Server->CreateObject(UTestReplicatedIrisObject::FComponents());
+
+	// Introduce latency by not immediately delivering packets.
+	Server->PreSendUpdate();
+	Server->SendTo(Client, TEXT("Create object"));
+	Server->PostSendUpdate();
+
+	// Filter out object to cause it to end up being PendingDestroy
+	Server->ReplicationSystem->AddToGroup(Server->GetReplicationSystem()->GetNotReplicatedNetObjectGroup(), ServerObject->NetRefHandle);
+
+	Server->PreSendUpdate();
+	Server->SendTo(Client, TEXT("Destroy object"));
+	Server->PostSendUpdate();
+
+	// Remove object from filter to cause object to not destroy it. 
+	Server->ReplicationSystem->RemoveFromGroup(Server->GetReplicationSystem()->GetNotReplicatedNetObjectGroup(), ServerObject->NetRefHandle);
+
+	// Modify object.
+	ServerObject->IntA += 1;
+
+	Server->PreSendUpdate();
+	Server->SendTo(Client, TEXT("Update object"));
+	Server->PostSendUpdate();
+
+	// Drop the initial packet.
+	Server->DeliverTo(Client, DoNotDeliverPacket);
+
+	// Deliver remaining packets, if any.
+	{
+		SIZE_T PacketCount = 0;
+		const auto& ConnectionInfo = Server->GetConnectionInfo(Client->ConnectionIdOnServer);
+		PacketCount = ConnectionInfo.WrittenPackets.Count();
+		for (SIZE_T PacketIt = 0; PacketIt != PacketCount; ++PacketIt)
+		{
+			Server->DeliverTo(Client, DeliverPacket);
+		}
+	}
+
+	Server->UpdateAndSend({ Client });
+
+	// The object should end up being created with the latest state.
+	UTestReplicatedIrisObject* ClientObject = Cast<UTestReplicatedIrisObject>(Client->GetReplicationBridge()->GetReplicatedObject(ServerObject->NetRefHandle));
+	UE_NET_ASSERT_NE(ClientObject, nullptr);
+	UE_NET_ASSERT_EQ(ClientObject->IntA, ServerObject->IntA);
+}
+
 UE_NET_TEST_FIXTURE(FTestCancelPendingDestroyFixture, TestCancelPendingDestroyDuringWaitOnDestroyConfirmationWithoutPacketLoss)
 {
 	// Add a client
diff --git a/Engine/Plugins/Runtime/ReplicationSystemTestPlugin/Source/Private/Tests/ReplicationSystem/TestReplicationSystem.cpp b/Engine/Plugins/Runtime/ReplicationSystemTestPlugin/Source/Private/Tests/ReplicationSystem/TestReplicationSystem.cpp
--- a/Engine/Plugins/Runtime/ReplicationSystemTestPlugin/Source/Private/Tests/ReplicationSystem/TestReplicationSystem.cpp
+++ b/Engine/Plugins/Runtime/ReplicationSystemTestPlugin/Source/Private/Tests/ReplicationSystem/TestReplicationSystem.cpp
@@ -732,35 +732,31 @@ UE_NET_TEST_FIXTURE(FReplicationSystemServerClientTestFixture, AddRemoveFromConn
 	// Disallow group to trigger state change from PendingCreateConfirmation->PendingDestroy
 	ReplicationSystem->SetGroupFilterStatus(Group, ENetFilterStatus::Disallow);	
 
+	// Expect client to create object
+	Server->DeliverTo(Client, true);
+	UE_NET_ASSERT_NE(Client->GetReplicationBridge()->GetReplicatedObject(ServerObject->NetRefHandle), nullptr);
+
 	// Send packet
-	// Expected state to be WaitOnDestroyConfirmation
 	Server->PreSendUpdate();
 	Server->SendTo(Client);
 	Server->PostSendUpdate();
 
-	// Allow group to trigger state to ensure that we restart replication since we have not actually created the object on the client
+	// Allow group to trigger state to ensure that we restart replication
 	ReplicationSystem->SetGroupFilterStatus(Group, ENetFilterStatus::Allow);
 
-	// Expect client to create object
+	// Expect client to destroy object
 	Server->DeliverTo(Client, true);
-	UE_NET_ASSERT_TRUE(Client->GetReplicationBridge()->GetReplicatedObject(ServerObject->NetRefHandle) != nullptr);
-
-	// Expect client to destroy object and server to move to Destroyed state
-	Server->DeliverTo(Client, true);
-	UE_NET_ASSERT_TRUE(Client->GetReplicationBridge()->GetReplicatedObject(ServerObject->NetRefHandle) == nullptr);
+	UE_NET_ASSERT_EQ(Client->GetReplicationBridge()->GetReplicatedObject(ServerObject->NetRefHandle), nullptr);
 
 	// Trigger replication
 	++ServerObject->IntA;
 
 	// Send packet
-	// Invalid -> PendingCreate
-	// PendingCreate->WaitOnCreateConfirmation
-	Server->PreSendUpdate();
-	Server->SendAndDeliverTo(Client, true);
-	Server->PostSendUpdate();
+	// WaitOnDestroyConfirmation -> WaitOnCreateConfirmation
+	Server->UpdateAndSend({ Client });
 
 	// Verify that the object got created again
-	UE_NET_ASSERT_TRUE(Client->GetReplicationBridge()->GetReplicatedObject(ServerObject->NetRefHandle) != nullptr);
+	UE_NET_ASSERT_NE(Client->GetReplicationBridge()->GetReplicatedObject(ServerObject->NetRefHandle), nullptr);
 }
 
 UE_NET_TEST_FIXTURE(FReplicationSystemServerClientTestFixture, TestNetTemporary)
diff --git a/Engine/Source/Runtime/Experimental/Iris/Core/Private/Iris/ReplicationSystem/ReplicationWriter.cpp b/Engine/Source/Runtime/Experimental/Iris/Core/Private/Iris/ReplicationSystem/ReplicationWriter.cpp
--- a/Engine/Source/Runtime/Experimental/Iris/Core/Private/Iris/ReplicationSystem/ReplicationWriter.cpp
+++ b/Engine/Source/Runtime/Experimental/Iris/Core/Private/Iris/ReplicationSystem/ReplicationWriter.cpp
@@ -1633,6 +1633,12 @@ uint32 FReplicationWriter::WriteObjectsPendingDestroy(FNetSerializationContext&
 		FReplicationInfo& Info = GetReplicationInfo(InternalIndex);
 		const FNetRefHandleManager::FReplicatedObjectData& ObjectData = NetRefHandleManager->GetReplicatedObjectDataNoCheck(InternalIndex);
 
+		// Don't send destroy until object creation has been acked.
+		if (!Info.IsCreationConfirmed)
+		{
+			continue;
+		}
+
 		// Already waiting on destroy confirmation
 		if (Info.GetState() == EReplicatedObjectState::WaitOnDestroyConfirmation)
 		{
