Update Iris grid filtering to include an exact culling distance check.

diff --git a/Engine/Plugins/Runtime/ReplicationSystemTestPlugin/Source/Private/Tests/ReplicationSystem/Filtering/TestGridFilter.cpp b/Engine/Plugins/Runtime/ReplicationSystemTestPlugin/Source/Private/Tests/ReplicationSystem/Filtering/TestGridFilter.cpp
--- a/Engine/Plugins/Runtime/ReplicationSystemTestPlugin/Source/Private/Tests/ReplicationSystem/Filtering/TestGridFilter.cpp
+++ b/Engine/Plugins/Runtime/ReplicationSystemTestPlugin/Source/Private/Tests/ReplicationSystem/Filtering/TestGridFilter.cpp
@@ -11,6 +11,7 @@
 #include "Tests/ReplicationSystem/Filtering/TestFilteringObject.h"
 #include "Tests/ReplicationSystem/ReplicationSystemServerClientTestFixture.h"
 
+#include "GenericPlatform/GenericPlatformMath.h"
 
 namespace UE::Net::Private
 {
@@ -107,6 +108,7 @@ UE_NET_TEST_FIXTURE(FTestGridFilterFixture, TestWorldLocGridFilter)
 	Params.bUseClassConfigDynamicFilter = true;
 
 	UReplicatedTestObject* ServerObjectZero = Server->CreateObject(Params);
+	UReplicatedTestObject* ServerObjectNear = Server->CreateObject(Params);
 	UReplicatedTestObject* ServerObjectLimit = Server->CreateObject(Params);
 	UReplicatedTestObject* ServerObjectCulled = Server->CreateObject(Params);
 	UReplicatedTestObject* ServerObjectVeryFar = Server->CreateObject(Params);
@@ -122,6 +124,7 @@ UE_NET_TEST_FIXTURE(FTestGridFilterFixture, TestWorldLocGridFilter)
 
 	// Visible objects
 	ObjectLocs.Add(ServerObjectZero,	FObjectLoc{ FVector::ZeroVector, 1500.f} );
+	ObjectLocs.Add(ServerObjectNear,	FObjectLoc{ FVector(100.f, 100.f, 100.f), 1500.f} );
 	ObjectLocs.Add(ServerObjectLimit,	FObjectLoc{ FVector(DefaultGridConfig->CellSizeX, DefaultGridConfig->CellSizeY, 1500.f), DefaultGridConfig->MaxCullDistance});
 
 	// Culled objects
@@ -141,9 +144,15 @@ UE_NET_TEST_FIXTURE(FTestGridFilterFixture, TestWorldLocGridFilter)
 	Server->ReplicationSystem->SetFilter(ServerObjectCulled->NetRefHandle, WorldLocFilterHandle);
 	Server->ReplicationSystem->SetFilter(ServerObjectVeryFar->NetRefHandle, WorldLocFilterHandle);
 
-	// Add client (view location at 0,0,0)
+	// Add client
 	FReplicationSystemTestClient* Client = CreateClient();
 
+	// Set the view location of the client to (0,0,0)
+	FReplicationView ReplicationView;
+	FReplicationView::FView View = ReplicationView.Views.Emplace_GetRef();
+	View.Pos = FVector(0.0f, 0.0f, 0.0f);
+	Server->ReplicationSystem->SetReplicationView(Client->ConnectionIdOnServer, ReplicationView);
+
 	// Send and deliver packet
 	Server->PreSendUpdate();
 	Server->SendAndDeliverTo(Client, DeliverPacket);
@@ -151,11 +160,11 @@ UE_NET_TEST_FIXTURE(FTestGridFilterFixture, TestWorldLocGridFilter)
 
 	// Test visible objects
 	UE_NET_ASSERT_NE(Client->GetReplicationBridge()->GetReplicatedObject(ServerObjectZero->NetRefHandle), nullptr);
-	UE_NET_ASSERT_NE(Client->GetReplicationBridge()->GetReplicatedObject(ServerObjectLimit->NetRefHandle), nullptr);
+	UE_NET_ASSERT_NE(Client->GetReplicationBridge()->GetReplicatedObject(ServerObjectNear->NetRefHandle), nullptr);
 
 	// Test culled objects
-	// $IRIS TODO:  This fails since the culldistance makes the actor cross over to the same cell as the client. Need to decide if cull distances should be better respected in Iris
-	//UE_NET_ASSERT_EQ(Client->GetReplicationBridge()->GetReplicatedObject(ServerObjectCulled->NetRefHandle), nullptr);
+	UE_NET_ASSERT_EQ(Client->GetReplicationBridge()->GetReplicatedObject(ServerObjectLimit->NetRefHandle), nullptr);
+	UE_NET_ASSERT_EQ(Client->GetReplicationBridge()->GetReplicatedObject(ServerObjectCulled->NetRefHandle), nullptr);
 	UE_NET_ASSERT_EQ(Client->GetReplicationBridge()->GetReplicatedObject(ServerObjectVeryFar->NetRefHandle), nullptr);
 
 	Server->DestroyObject(ServerObjectZero);
@@ -182,9 +191,15 @@ UE_NET_TEST_FIXTURE(FTestGridFilterFixture, TestWorldLocationIsFrequentlyUpdated
 			OutCullDistance = ObjectLoc.CullDistance;
 		});
 
-	// Add client (view location at 0,0,0)
+	// Add client
 	FReplicationSystemTestClient* Client = CreateClient();
 
+	// Set the view location of the client to (0,0,0)
+	FReplicationView ReplicationView;
+	FReplicationView::FView View = ReplicationView.Views.Emplace_GetRef();
+	View.Pos = FVector(0.0f, 0.0f, 0.0f);
+	Server->ReplicationSystem->SetReplicationView(Client->ConnectionIdOnServer, ReplicationView);
+
 	// Spawn object with WorldLocation's on server
 	UObjectReplicationBridge::FCreateNetRefHandleParams Params;
 	Params.bNeedsWorldLocationUpdate = true;
@@ -233,6 +248,7 @@ UE_NET_TEST_FIXTURE(FTestGridFilterFixture, TestFragmentFilter)
 {
 	// Spawn object on server
 	UTestLocationFragmentFilteringObject* ServerObjectZero = Server->CreateObject<UTestLocationFragmentFilteringObject>();
+	UTestLocationFragmentFilteringObject* ServerObjectNear = Server->CreateObject<UTestLocationFragmentFilteringObject>();
 	UTestLocationFragmentFilteringObject* ServerObjectLimit = Server->CreateObject<UTestLocationFragmentFilteringObject>();
 	UTestLocationFragmentFilteringObject* ServerObjectCulled = Server->CreateObject<UTestLocationFragmentFilteringObject>();
 	UTestLocationFragmentFilteringObject* ServerObjectVeryFar = Server->CreateObject<UTestLocationFragmentFilteringObject>();
@@ -241,6 +257,9 @@ UE_NET_TEST_FIXTURE(FTestGridFilterFixture, TestFragmentFilter)
 
 	ServerObjectZero->WorldLocation = FVector::ZeroVector;
 	ServerObjectZero->NetCullDistanceSquared = 1500.f * 1500.f;
+
+	ServerObjectNear->WorldLocation = FVector(100.f, 100.f, 100.f);
+	ServerObjectNear->NetCullDistanceSquared = 1500.f * 1500.f;
 	
 	ServerObjectLimit->WorldLocation = FVector(DefaultGridConfig->CellSizeX, DefaultGridConfig->CellSizeY, 1500.f);
 	ServerObjectLimit->NetCullDistanceSquared = DefaultGridConfig->MaxCullDistance * DefaultGridConfig->MaxCullDistance;
@@ -253,13 +272,20 @@ UE_NET_TEST_FIXTURE(FTestGridFilterFixture, TestFragmentFilter)
 
 	// Apply grid filter
 	Server->ReplicationSystem->SetFilter(ServerObjectZero->NetRefHandle, FragmentLocFilterHandle);
+	Server->ReplicationSystem->SetFilter(ServerObjectNear->NetRefHandle, FragmentLocFilterHandle);
 	Server->ReplicationSystem->SetFilter(ServerObjectLimit->NetRefHandle, FragmentLocFilterHandle);
 	Server->ReplicationSystem->SetFilter(ServerObjectCulled->NetRefHandle, FragmentLocFilterHandle);
 	Server->ReplicationSystem->SetFilter(ServerObjectVeryFar->NetRefHandle, FragmentLocFilterHandle);
 
-	// Add client (view location at 0,0,0)
+	// Add client
 	FReplicationSystemTestClient* Client = CreateClient();
 
+	// Set the view location of the client to (0,0,0)
+	FReplicationView ReplicationView;
+	FReplicationView::FView View = ReplicationView.Views.Emplace_GetRef();
+	View.Pos = FVector(0.0f, 0.0f, 0.0f);
+	Server->ReplicationSystem->SetReplicationView(Client->ConnectionIdOnServer, ReplicationView);
+
 	// Send and deliver packet
 	Server->PreSendUpdate();
 	Server->SendAndDeliverTo(Client, DeliverPacket);
@@ -267,11 +293,11 @@ UE_NET_TEST_FIXTURE(FTestGridFilterFixture, TestFragmentFilter)
 
 	// Test visible objects
 	UE_NET_ASSERT_NE(Client->GetReplicationBridge()->GetReplicatedObject(ServerObjectZero->NetRefHandle), nullptr);
-	UE_NET_ASSERT_NE(Client->GetReplicationBridge()->GetReplicatedObject(ServerObjectLimit->NetRefHandle), nullptr);
+	UE_NET_ASSERT_NE(Client->GetReplicationBridge()->GetReplicatedObject(ServerObjectNear->NetRefHandle), nullptr);
 
 	// Test culled objects
-	// $IRIS TODO:  This fails since the culldistance makes the actor cross over to the same cell as the client. Need to decide if cull distances should be better respected in Iris
-	//UE_NET_ASSERT_EQ(Client->GetReplicationBridge()->GetReplicatedObject(ServerObjectCulled->NetRefHandle), nullptr);
+	UE_NET_ASSERT_EQ(Client->GetReplicationBridge()->GetReplicatedObject(ServerObjectLimit->NetRefHandle), nullptr);
+	UE_NET_ASSERT_EQ(Client->GetReplicationBridge()->GetReplicatedObject(ServerObjectCulled->NetRefHandle), nullptr);
 	UE_NET_ASSERT_EQ(Client->GetReplicationBridge()->GetReplicatedObject(ServerObjectVeryFar->NetRefHandle), nullptr);
 
 	Server->DestroyObject(ServerObjectZero);
@@ -280,5 +306,75 @@ UE_NET_TEST_FIXTURE(FTestGridFilterFixture, TestFragmentFilter)
 	Server->DestroyObject(ServerObjectVeryFar);
 }
 
+#if 0
+UE_NET_TEST_FIXTURE(FTestGridFilterFixture, TestFilterPerformance)
+{
+	static const int32 NUM_CLIENTS = 14;
+	static const int32 NUM_OBJECTS = 2000;
+	static const int32 TEST_ITERATIONS = 100;
+
+	struct FObjectLoc
+	{
+		FVector Loc;
+		float CullDistance;
+	};
+
+	TMap<const UObject*, FObjectLoc> ObjectLocs;
+	Server->GetReplicationBridge()->SetExternalWorldLocationUpdateFunctor([&](FNetRefHandle NetHandle, const UObject* ReplicatedObject, FVector& OutLocation, float& OutCullDistance)
+		{
+			const FObjectLoc& ObjectLoc = ObjectLocs[ReplicatedObject];
+			OutLocation = ObjectLoc.Loc;
+			OutCullDistance = ObjectLoc.CullDistance;
+		});
+
+	// Create client connections.
+	TArray<FReplicationSystemTestClient*> TestClients;
+	for (int32 i = 0; i < NUM_CLIENTS; i++)
+	{
+		FReplicationSystemTestClient* TestClient = CreateClient();
+
+		TestClients.Add(TestClient);
+
+		FReplicationView ReplicationView;
+		FReplicationView::FView View = ReplicationView.Views.Emplace_GetRef();
+		View.Pos = FVector(0.0f, 0.0f, 0.0f);
+		Server->ReplicationSystem->SetReplicationView(TestClient->ConnectionIdOnServer, ReplicationView);
+	}
+
+	// Create objects at random positions inside a cell.
+	const UNetObjectGridFilterConfig* DefaultGridConfig = GetDefault<UNetObjectGridFilterConfig>();
+	FGenericPlatformMath::SRandInit(0);
+	for (int32 i = 0; i < NUM_OBJECTS; i++)
+	{
+		UObjectReplicationBridge::FCreateNetRefHandleParams Params;
+		Params.bCanReceive = true;
+		Params.bNeedsWorldLocationUpdate = true;
+		Params.bUseClassConfigDynamicFilter = true;
+
+		FVector Pos;
+		Pos.X = FGenericPlatformMath::SRand() * DefaultGridConfig->CellSizeX;
+		Pos.Y = FGenericPlatformMath::SRand() * DefaultGridConfig->CellSizeY;
+		Pos.Z = 0.0f;
+
+		UReplicatedTestObject* Object = Server->CreateObject(Params);
+		ObjectLocs.Add(Object, FObjectLoc{ Pos, 1500.f });
+
+		Server->ReplicationSystem->SetFilter(Object->NetRefHandle, WorldLocFilterHandle);
+	}
+
+	// Run server replication multiple iterations.
+	for (int32 i = 0; i < TEST_ITERATIONS; i++)
+	{
+		const double StartTime = FPlatformTime::Seconds();
+
+		Server->PreSendUpdate();
+		for (FReplicationSystemTestClient* TestClient : TestClients)
+		{
+			Server->SendAndDeliverTo(TestClient, DeliverPacket);
+		}
+		Server->PostSendUpdate();
+	}
+}
+#endif
 
 } // end namespace UE::Net::Private
\ No newline at end of file
diff --git a/Engine/Source/Runtime/Experimental/Iris/Core/Private/Iris/ReplicationSystem/Filtering/NetObjectGridFilter.cpp b/Engine/Source/Runtime/Experimental/Iris/Core/Private/Iris/ReplicationSystem/Filtering/NetObjectGridFilter.cpp
--- a/Engine/Source/Runtime/Experimental/Iris/Core/Private/Iris/ReplicationSystem/Filtering/NetObjectGridFilter.cpp
+++ b/Engine/Source/Runtime/Experimental/Iris/Core/Private/Iris/ReplicationSystem/Filtering/NetObjectGridFilter.cpp
@@ -53,7 +53,7 @@ bool UNetObjectGridFilter::AddObject(uint32 ObjectIndex, FNetObjectFilterAddObje
 
 	AddCellInfoForObject(ObjectLocationInfo, Params.InstanceProtocol);
 	
-	if (PerObjectInfo.CullDistance > Config->MaxCullDistance)
+	if (PerObjectInfo.GetCullDistance() > Config->MaxCullDistance)
 	{
 		// Too big an object. We expect it to be costly to move it across cells.
 		RemoveObject(ObjectIndex, ObjectLocationInfo);
@@ -110,8 +110,9 @@ void UNetObjectGridFilter::Filter(FNetObjectFilteringParams& Params)
 		}
 	}
 
-	// Prune old cells and insert still relevant in the new list.
+	if (!Config->bUseExactCullDistance)
 	{
+		// Prune old cells and insert still relevant in the new list.
 		const uint32 MaxFrameCount = Config->ViewPosRelevancyFrameCount;
 		for (const FCellAndTimestamp& PrevCell : PrevCells)
 		{
@@ -144,7 +145,42 @@ void UNetObjectGridFilter::Filter(FNetObjectFilteringParams& Params)
 		{
 			for (const uint32 ObjectIndex : Objects->ObjectIndices)
 			{
-				AllowedObjects.SetBit(ObjectIndex);
+				if (Config->bUseExactCullDistance)
+				{
+					const FObjectLocationInfo& ObjectLocationInfo = static_cast<const FObjectLocationInfo&>(Params.FilteringInfos[ObjectIndex]);
+					const FPerObjectInfo& PerObjectInfo = ObjectInfos[ObjectLocationInfo.GetInfoIndex()];
+
+					for (const FReplicationView::FView& View : Params.View.Views)
+					{
+						const double DistSq = PerObjectInfo.GetCullDistanceSq();
+						const double ObjectToViewDistSq = FVector::DistSquared2D(PerObjectInfo.Position, View.Pos);
+
+						if (ObjectToViewDistSq <= DistSq)
+						{
+							ConnectionInfo.RecentObjectFrameCount.Add(ObjectIndex, Config->ViewPosRelevancyFrameCount);
+						}
+					}
+				}
+				else
+				{
+					AllowedObjects.SetBit(ObjectIndex);
+				}
+			}
+		}
+	}
+
+	if (Config->bUseExactCullDistance)
+	{
+		for (TMap<uint32, uint32>::TIterator It = ConnectionInfo.RecentObjectFrameCount.CreateIterator(); It; ++It)
+		{
+			if (It->Value > 0)
+			{
+				It->Value--;
+				AllowedObjects.SetBit(It->Key);
+			}
+			else
+			{
+				It.RemoveCurrent();
 			}
 		}
 	}
@@ -173,7 +209,7 @@ void UNetObjectGridFilter::AddCellInfoForObject(const FObjectLocationInfo& Objec
 {
 	// Called for completely new objects
 	FPerObjectInfo& PerObjectInfo = ObjectInfos[ObjectLocationInfo.GetInfoIndex()];
-	PerObjectInfo.CullDistance = Config->DefaultCullDistance;
+	PerObjectInfo.SetCullDistance(Config->DefaultCullDistance);
 	UpdatePositionAndCullDistance(ObjectLocationInfo, PerObjectInfo, InstanceProtocol);
 
 	FCellBox NewCellBox;
@@ -343,14 +379,14 @@ void UNetObjectGridFilter::UpdatePositionAndCullDistance(const UNetObjectGridFil
 	// Optionally update cull distance
 	if (NetCullDistanceOverrides->HasCullDistanceOverride(PerObjectInfo.ObjectIndex))
 	{
-		const float CullDistanceSqr = NetCullDistanceOverrides->GetCullDistanceSqr(PerObjectInfo.ObjectIndex);
-		PerObjectInfo.CullDistance = FPlatformMath::Sqrt(CullDistanceSqr);
+		const float CullDistanceSq = NetCullDistanceOverrides->GetCullDistanceSqr(PerObjectInfo.ObjectIndex);
+		PerObjectInfo.SetCullDistanceSq(CullDistanceSq);
 	}
 }
 
 void UNetObjectGridFilter::CalculateCellBox(const UNetObjectGridFilter::FPerObjectInfo& PerObjectInfo, UNetObjectGridFilter::FCellBox& OutCellBox)
 {
-	const float CullDistance = PerObjectInfo.CullDistance;
+	const double CullDistance = PerObjectInfo.GetCullDistance();
 	const FVector Position = PerObjectInfo.Position;
 	FVector MinPosition = Position - CullDistance;
 	FVector MaxPosition = Position + CullDistance;
@@ -434,7 +470,7 @@ void UNetObjectGridWorldLocFilter::UpdateObjectInfo(UNetObjectGridFilter::FPerOb
 	const UE::Net::FWorldLocations::FObjectInfo& CachedObjectInfo = WorldLocations->GetObjectInfo(PerObjectInfo.ObjectIndex);
 	
 	PerObjectInfo.Position = CachedObjectInfo.WorldLocation;
-	PerObjectInfo.CullDistance = CachedObjectInfo.CullDistance;
+	PerObjectInfo.SetCullDistance(CachedObjectInfo.CullDistance);
 }
 
 bool UNetObjectGridWorldLocFilter::BuildObjectInfo(uint32 ObjectIndex, FNetObjectFilterAddObjectParams& Params)
@@ -532,7 +568,7 @@ void UNetObjectGridFragmentLocFilter::UpdateObjectInfo(UNetObjectGridFilter::FPe
 		const UE::Net::FReplicationInstanceProtocol::FFragmentData& FragmentData = FragmentDatas[CullDistanceFragmentInfo->CullDistanceSqrStateIndex];
 		const uint8* CullDistanceSqrAddress = FragmentData.ExternalSrcBuffer + CullDistanceFragmentInfo->CullDistanceSqrStateOffset;
 		const float CullDistanceSqr = *reinterpret_cast<const float*>(CullDistanceSqrAddress);
-		PerObjectInfo.CullDistance = FPlatformMath::Sqrt(CullDistanceSqr);
+		PerObjectInfo.SetCullDistanceSq(CullDistanceSqr);
 	}
 }
 
diff --git a/Engine/Source/Runtime/Experimental/Iris/Core/Public/Iris/ReplicationSystem/Filtering/NetObjectGridFilter.h b/Engine/Source/Runtime/Experimental/Iris/Core/Public/Iris/ReplicationSystem/Filtering/NetObjectGridFilter.h
--- a/Engine/Source/Runtime/Experimental/Iris/Core/Public/Iris/ReplicationSystem/Filtering/NetObjectGridFilter.h
+++ b/Engine/Source/Runtime/Experimental/Iris/Core/Public/Iris/ReplicationSystem/Filtering/NetObjectGridFilter.h
@@ -47,6 +47,10 @@ public:
 	/** Coordinates will be clamped to MinPos and MaxPos. */
 	UPROPERTY(Config)
 	FVector MaxPos = {+0.5f*2097152.0f, +0.5f*2097152.0f, +0.5f*2097152.0f};
+
+	/** If true, use the exact cull distance to determine the objects to create/destroy on the client. Otherwise only use the net grid. */
+	UPROPERTY(Config)
+	bool bUseExactCullDistance = true;
 };
 
 UCLASS(abstract)
@@ -109,8 +113,33 @@ protected:
 	{
 		FVector Position = { 0.f,0.f,0.f };
 		FCellBox CellBox = {};
-		float CullDistance = 0.0f;
 		uint32 ObjectIndex = 0U;
+
+		float GetCullDistance() const
+		{
+			return CullDistance;
+		}
+
+		float GetCullDistanceSq() const
+		{
+			return CullDistanceSq;
+		}
+
+		void SetCullDistance(float Distance)
+		{
+			CullDistance = Distance;
+			CullDistanceSq = Distance * Distance;
+		}
+
+		void SetCullDistanceSq(float DistanceSq)
+		{
+			CullDistance = FPlatformMath::Sqrt(DistanceSq);
+			CullDistanceSq = DistanceSq;
+		}
+
+	private:
+		float CullDistance = 0.0f;
+		float CullDistanceSq = 0.0f;
 	};
 
 	/** Sets the current position of the object based on how we access it's given location. */
@@ -162,6 +191,9 @@ private:
 	{
 		// We don't expect a lot of view positions from a single connection
 		TArray<FCellAndTimestamp, TInlineAllocator<32>> RecentCells;
+		
+		// Objects that have been recently visible to the connection and the frame countdown.
+		TMap<uint32, uint32> RecentObjectFrameCount;
 	};
 
 private:
diff --git a/Engine/Source/Runtime/Experimental/Iris/Core/Public/Iris/ReplicationSystem/ReplicationView.h b/Engine/Source/Runtime/Experimental/Iris/Core/Public/Iris/ReplicationSystem/ReplicationView.h
--- a/Engine/Source/Runtime/Experimental/Iris/Core/Public/Iris/ReplicationSystem/ReplicationView.h
+++ b/Engine/Source/Runtime/Experimental/Iris/Core/Public/Iris/ReplicationSystem/ReplicationView.h
@@ -24,11 +24,11 @@ struct FReplicationView
 		/** The actor that is being directly viewed, usually a pawn. */
 		FNetHandle ViewTarget;
 		/** Where the viewer is looking from */
-		FVector Pos;
+		FVector Pos = FVector::ZeroVector;
 		/** Direction the viewer is looking */
-		FVector Dir;
+		FVector Dir = FVector::ForwardVector;
 		/** The field of view */
-		float FoVRadians;
+		float FoVRadians = UE_HALF_PI;
 	};
 
 	TArray<FView, TInlineAllocator<UE_IRIS_INLINE_VIEWS_PER_CONNECTION>> Views;
