Fixed compilation errors from 31718958

diff --git a/Engine/Plugins/Runtime/ReplicationSystemTestPlugin/Source/Private/Tests/ReplicationState/TestFastArrayReplicationState.h b/Engine/Plugins/Runtime/ReplicationSystemTestPlugin/Source/Private/Tests/ReplicationState/TestFastArrayReplicationState.h
--- a/Engine/Plugins/Runtime/ReplicationSystemTestPlugin/Source/Private/Tests/ReplicationState/TestFastArrayReplicationState.h
+++ b/Engine/Plugins/Runtime/ReplicationSystemTestPlugin/Source/Private/Tests/ReplicationState/TestFastArrayReplicationState.h
@@ -68,7 +68,9 @@ struct FTestFastArrayReplicationState_FastArraySerializer : public FIrisFastArra
 	void PostReplicatedReceive(const FFastArraySerializer::FPostReplicatedReceiveParameters& Parameters)
 	{
 		bHitPostReplicatedReceive = 1U;
+PRAGMA_DISABLE_DEPRECATION_WARNINGS
 		bPostReplicatedReceiveWasHitWithUnresolvedReferences = Parameters.bHasMoreUnmappedReferences;
+PRAGMA_ENABLE_DEPRECATION_WARNINGS
 	}
 
 protected:
diff --git a/Engine/Plugins/Runtime/ReplicationSystemTestPlugin/Source/Private/Tests/Serialization/TestPolymorphicStructNetSerializer.h b/Engine/Plugins/Runtime/ReplicationSystemTestPlugin/Source/Private/Tests/Serialization/TestPolymorphicStructNetSerializer.h
--- a/Engine/Plugins/Runtime/ReplicationSystemTestPlugin/Source/Private/Tests/Serialization/TestPolymorphicStructNetSerializer.h
+++ b/Engine/Plugins/Runtime/ReplicationSystemTestPlugin/Source/Private/Tests/Serialization/TestPolymorphicStructNetSerializer.h
@@ -359,7 +359,9 @@ struct FExamplePolymorphicStructFastArraySerializer : public FIrisFastArraySeria
 	void PostReplicatedReceive(const FFastArraySerializer::FPostReplicatedReceiveParameters& Parameters)
 	{
 		bHitPostReplicatedReceive = 1U;
+PRAGMA_DISABLE_DEPRECATION_WARNINGS
 		bPostReplicatedReceiveWasHitWithUnresolvedReferences = Parameters.bHasMoreUnmappedReferences;
+PRAGMA_ENABLE_DEPRECATION_WARNINGS
 	}
 
 public:
