Iris - Fix ReplicationSystemTest failure due to NetObjectGridFilter updates.

#rnx

diff --git a/Engine/Plugins/Runtime/ReplicationSystemTestPlugin/Source/Private/Tests/ReplicationSystem/Filtering/TestGridFilter.cpp b/Engine/Plugins/Runtime/ReplicationSystemTestPlugin/Source/Private/Tests/ReplicationSystem/Filtering/TestGridFilter.cpp
--- a/Engine/Plugins/Runtime/ReplicationSystemTestPlugin/Source/Private/Tests/ReplicationSystem/Filtering/TestGridFilter.cpp
+++ b/Engine/Plugins/Runtime/ReplicationSystemTestPlugin/Source/Private/Tests/ReplicationSystem/Filtering/TestGridFilter.cpp
@@ -225,7 +225,7 @@ UE_NET_TEST_FIXTURE(FTestGridFilterFixture, TestWorldLocationIsFrequentlyUpdated
 	ObjectLocs.Emplace(ServerObject, FObjectLoc{ FVector(DefaultGridConfig->CellSizeX + 1600.0f,DefaultGridConfig->CellSizeY + 1600.0f, 0), 1500.f });
 
 	// Send and deliver packet
-	for (uint32 LoopIt = 0, LoopEndIt = DefaultGridConfig->ViewPosRelevancyFrameCount; LoopIt <= LoopEndIt; ++LoopIt)
+	for (uint32 LoopIt = 0, LoopEndIt = DefaultGridConfig->ViewPosRelevancyFrameCount + DefaultGridConfig->DefaultFrameCountBeforeCulling; LoopIt <= LoopEndIt; ++LoopIt)
 	{
 		Server->UpdateAndSend({ Client });
 	}
