Iris - Reduce log spam in ReplicationSystemTests

diff --git a/Engine/Plugins/Runtime/ReplicationSystemTestPlugin/Source/Private/Tests/ReplicationSystem/ReplicationSystemServerClientTestFixture.cpp b/Engine/Plugins/Runtime/ReplicationSystemTestPlugin/Source/Private/Tests/ReplicationSystem/ReplicationSystemServerClientTestFixture.cpp
--- a/Engine/Plugins/Runtime/ReplicationSystemTestPlugin/Source/Private/Tests/ReplicationSystem/ReplicationSystemServerClientTestFixture.cpp
+++ b/Engine/Plugins/Runtime/ReplicationSystemTestPlugin/Source/Private/Tests/ReplicationSystem/ReplicationSystemServerClientTestFixture.cpp
@@ -1,6 +1,7 @@
 // Copyright Epic Games, Inc. All Rights Reserved.
 
 #include "ReplicationSystemServerClientTestFixture.h"
+#include "Logging/LogScopedVerbosityOverride.h"
 #include "Iris/Core/IrisLog.h"
 #include "Iris/ReplicationSystem/ReplicationSystem.h"
 #include "Iris/ReplicationSystem/ReplicationSystemInternal.h"
@@ -65,6 +66,8 @@ void FDataStreamTestUtil::AddDataStreamDefinition(const TCHAR* StreamName, const
 FReplicationSystemTestNode::FReplicationSystemTestNode(bool bIsServer, const TCHAR* Name)
 {
 	ReplicationBridge = NewObject<UReplicatedTestObjectBridge>();
+	check(ReplicationBridge != nullptr);
+
 	CreatedObjects.Add(TStrongObjectPtr<UObject>(ReplicationBridge));
 
 	UReplicationSystem::FReplicationSystemParams Params;
@@ -72,15 +75,12 @@ FReplicationSystemTestNode::FReplicationSystemTestNode(bool bIsServer, const TCH
 	Params.bIsServer = bIsServer;
 	Params.bAllowObjectReplication = bIsServer;
 
-	auto IrisLogVerbosity = UE_GET_LOG_VERBOSITY(LogIris);
-	LogIris.SetVerbosity(ELogVerbosity::Error);
+	LOG_SCOPE_VERBOSITY_OVERRIDE(LogIris, ELogVerbosity::Error);
 	ReplicationSystem = FReplicationSystemFactory::CreateReplicationSystem(Params);	
 	if (!bIsServer)
 	{
 		ReplicationBridge->SetCreatedObjectsOnNode(&CreatedObjects);
 	}
-	LogIris.SetVerbosity(IrisLogVerbosity);
-	check(ReplicationBridge != nullptr);
 
 	UE_NET_TRACE_UPDATE_INSTANCE(GetNetTraceId(), bIsServer, Name);
 }
@@ -88,6 +88,8 @@ FReplicationSystemTestNode::FReplicationSystemTestNode(bool bIsServer, const TCH
 FReplicationSystemTestNode::~FReplicationSystemTestNode()
 {
 	const uint32 NetTraceId = ReplicationSystem->GetId();
+
+	LOG_SCOPE_VERBOSITY_OVERRIDE(LogIris, ELogVerbosity::Error);
 	FReplicationSystemFactory::DestroyReplicationSystem(ReplicationSystem);
 	CreatedObjects.Empty();
 
diff --git a/Engine/Plugins/Runtime/ReplicationSystemTestPlugin/Source/Private/Tests/ReplicationSystem/ReplicationSystemTestFixture.h b/Engine/Plugins/Runtime/ReplicationSystemTestPlugin/Source/Private/Tests/ReplicationSystem/ReplicationSystemTestFixture.h
--- a/Engine/Plugins/Runtime/ReplicationSystemTestPlugin/Source/Private/Tests/ReplicationSystem/ReplicationSystemTestFixture.h
+++ b/Engine/Plugins/Runtime/ReplicationSystemTestPlugin/Source/Private/Tests/ReplicationSystem/ReplicationSystemTestFixture.h
@@ -34,7 +34,7 @@ protected:
 		Params.bAllowObjectReplication = true;
 
 		// In a testing environment without configs the creation of the ReplicationSystem can be quite spammy
-		auto IrisLogVerbosity = UE_GET_LOG_VERBOSITY(LogIris);
+		ELogVerbosity::Type IrisLogVerbosity = UE_GET_LOG_VERBOSITY(LogIris);
 		LogIris.SetVerbosity(ELogVerbosity::Error);
 		ReplicationSystem = FReplicationSystemFactory::CreateReplicationSystem(Params);
 		LogIris.SetVerbosity(IrisLogVerbosity);
@@ -44,7 +44,10 @@ protected:
 
 	virtual void TearDown() override
 	{
+		const ELogVerbosity::Type IrisLogVerbosity = UE_GET_LOG_VERBOSITY(LogIris);
+		LogIris.SetVerbosity(ELogVerbosity::Error);
 		FReplicationSystemFactory::DestroyReplicationSystem(ReplicationSystem);
+		LogIris.SetVerbosity(IrisLogVerbosity);
 		CreatedObjects.Empty();
 	}
 
