Iris - Initialize FObjectPollFrequencyLimiter with the appropriate max object count.

#rnx

diff --git a/Engine/Plugins/Runtime/ReplicationSystemTestPlugin/Source/Private/Tests/ReplicationSystem/TestObjectPollFrequencyLimiter.cpp b/Engine/Plugins/Runtime/ReplicationSystemTestPlugin/Source/Private/Tests/ReplicationSystem/TestObjectPollFrequencyLimiter.cpp
--- a/Engine/Plugins/Runtime/ReplicationSystemTestPlugin/Source/Private/Tests/ReplicationSystem/TestObjectPollFrequencyLimiter.cpp
+++ b/Engine/Plugins/Runtime/ReplicationSystemTestPlugin/Source/Private/Tests/ReplicationSystem/TestObjectPollFrequencyLimiter.cpp
@@ -19,7 +19,7 @@ UE_NET_TEST(ObjectPollFrequencyLimiter, ObjectsToPollAreSpreadOutWhenCreatingObj
 	constexpr uint32 MaxObjectCount = PollFrequency + 1U;
 
 	FObjectPollFrequencyLimiter FrequencyLimiter;
-	FrequencyLimiter.Init(MaxObjectCount);
+	FrequencyLimiter.Init(MaxObjectCount + 1);
 
 	FNetBitArray EmptyBitArray(MaxObjectCount + 1);
 	FNetBitArray ScopeBitArray(MaxObjectCount + 1);
@@ -56,7 +56,7 @@ UE_NET_TEST(ObjectPollFrequencyLimiter, ObjectsToPollAreSpreadOutRegardlessOfUpd
 
 	constexpr uint32 MaxObjectCount = 128;
 	FObjectPollFrequencyLimiter FrequencyLimiter;
-	FrequencyLimiter.Init(MaxObjectCount);
+	FrequencyLimiter.Init(MaxObjectCount + 1);
 
 	constexpr uint32 PollFrequency = 8U;
 	// With three times as many objects as the poll frequency we expect three objects to be polled per frame.
