Iris
- Reduced overhead when calling PreReplcate functions as we do it for many objects.

diff --git a/Engine/Plugins/Runtime/ReplicationSystemTestPlugin/Source/Private/Tests/ReplicationSystem/ReplicatedTestObject.cpp b/Engine/Plugins/Runtime/ReplicationSystemTestPlugin/Source/Private/Tests/ReplicationSystem/ReplicatedTestObject.cpp
--- a/Engine/Plugins/Runtime/ReplicationSystemTestPlugin/Source/Private/Tests/ReplicationSystem/ReplicatedTestObject.cpp
+++ b/Engine/Plugins/Runtime/ReplicationSystemTestPlugin/Source/Private/Tests/ReplicationSystem/ReplicatedTestObject.cpp
@@ -249,7 +249,7 @@ void UReplicatedTestObjectBridge::SetExternalWorldLocationUpdateFunctor(TFunctio
 	SetInstanceGetWorldObjectInfoFunction(LocUpdateFunctor);
 }
 
-void UReplicatedTestObjectBridge::SetExternalPreUpdateFunctor(TFunction<void(FNetRefHandle, UObject*, const UReplicationBridge*)> PreUpdateFunctor)
+void UReplicatedTestObjectBridge::SetExternalPreUpdateFunctor(TFunction<void(TArrayView<UObject*>, const UReplicationBridge*)> PreUpdateFunctor)
 {
 	SetInstancePreUpdateFunction(PreUpdateFunctor);
 }
diff --git a/Engine/Plugins/Runtime/ReplicationSystemTestPlugin/Source/Private/Tests/ReplicationSystem/ReplicatedTestObject.h b/Engine/Plugins/Runtime/ReplicationSystemTestPlugin/Source/Private/Tests/ReplicationSystem/ReplicatedTestObject.h
--- a/Engine/Plugins/Runtime/ReplicationSystemTestPlugin/Source/Private/Tests/ReplicationSystem/ReplicatedTestObject.h
+++ b/Engine/Plugins/Runtime/ReplicationSystemTestPlugin/Source/Private/Tests/ReplicationSystem/ReplicatedTestObject.h
@@ -506,7 +506,7 @@ public:
 
 	void SetExternalWorldLocationUpdateFunctor(TFunction<void(FNetRefHandle NetHandle, const UObject* ReplicatedObject, FVector& OutLocation, float& OutCullDistance)> LocUpdateFunctor);
 
-	void SetExternalPreUpdateFunctor(TFunction<void(FNetRefHandle, UObject*, const UReplicationBridge*)> PreUpdateFunctor);
+	void SetExternalPreUpdateFunctor(TFunction<void(TArrayView<UObject*>, const UReplicationBridge*)> PreUpdateFunctor);
 
 	float GetMaxTickRate() const { return Super::GetMaxTickRate(); }
 
diff --git a/Engine/Plugins/Runtime/ReplicationSystemTestPlugin/Source/Private/Tests/ReplicationSystem/TestDirtyNetObjectTracker.cpp b/Engine/Plugins/Runtime/ReplicationSystemTestPlugin/Source/Private/Tests/ReplicationSystem/TestDirtyNetObjectTracker.cpp
--- a/Engine/Plugins/Runtime/ReplicationSystemTestPlugin/Source/Private/Tests/ReplicationSystem/TestDirtyNetObjectTracker.cpp
+++ b/Engine/Plugins/Runtime/ReplicationSystemTestPlugin/Source/Private/Tests/ReplicationSystem/TestDirtyNetObjectTracker.cpp
@@ -232,12 +232,15 @@ UE_NET_TEST_FIXTURE(FReplicationSystemServerClientTestFixture, NetForceUpdateOth
 	// Force ObjectA to replicate
 	Server->ReplicationSystem->ForceNetUpdate(ServerObjectA->NetRefHandle);
 
-	auto PreUpdateObjectA = [&](FNetRefHandle NetHandle, UObject* ReplicatedObject, const UReplicationBridge* ReplicationBridge)
+	auto PreUpdateObjectA = [&](TArrayView<UObject*> Instances, const UReplicationBridge* Bridge)
 	{
-		// Inside ObjectA PreReplicationUpdate, force ObjectB to be replicated
-		if (ServerObjectA == ReplicatedObject)
+		for (UObject* ReplicatedObject : Instances)
 		{
-			Server->ReplicationSystem->ForceNetUpdate(ServerObjectB->NetRefHandle);
+			// Inside ObjectA PreReplicationUpdate, force ObjectB to be replicated
+			if (ServerObjectA == ReplicatedObject)
+			{
+				Server->ReplicationSystem->ForceNetUpdate(ServerObjectB->NetRefHandle);
+			}
 		}
 	};
 
@@ -259,12 +262,15 @@ UE_NET_TEST_FIXTURE(FReplicationSystemServerClientTestFixture, NetForceUpdateOth
 	// Now the ForceNetUpdate on ObjectB is applied and it is replicated
 	UE_NET_ASSERT_EQ(ClientObjectB->IntA, ServerObjectB->IntA);
 
-	auto PreUpdateObjectB = [&](FNetRefHandle NetHandle, UObject* ReplicatedObject, const UReplicationBridge* ReplicationBridge)
+	auto PreUpdateObjectB = [&](TArrayView<UObject*> Instances, const UReplicationBridge* Bridge)
 	{
-		// Inside ObjectB PreReplicationUpdate, force ObjectA to be replicated
-		if (ServerObjectB == ReplicatedObject)
+		for (UObject* ReplicatedObject : Instances)
 		{
-			Server->ReplicationSystem->ForceNetUpdate(ServerObjectA->NetRefHandle);
+			// Inside ObjectB PreReplicationUpdate, force ObjectA to be replicated
+			if (ServerObjectB == ReplicatedObject)
+			{
+				Server->ReplicationSystem->ForceNetUpdate(ServerObjectA->NetRefHandle);
+			}
 		}
 	};
 	Server->GetReplicationBridge()->SetExternalPreUpdateFunctor(PreUpdateObjectB);
@@ -322,12 +328,15 @@ UE_NET_TEST_FIXTURE(FReplicationSystemServerClientTestFixture, DirtyOtherObjectI
 	UE_NET_ASSERT_NE(ClientObjectB, nullptr);
 
 	UTestReplicatedIrisObject* ObjectToDirty = nullptr;
-	auto PreUpdateObject = [&](FNetRefHandle NetHandle, UObject* ReplicatedObject, const UReplicationBridge* ReplicationBridge)
+	auto PreUpdateObject = [&](TArrayView<UObject*> Instances, const UReplicationBridge* Bridge)
 	{
-		// There's only two objects, so when we update one we dirty the other
-		if (ObjectToDirty != ReplicatedObject)
+		for (UObject* ReplicatedObject : Instances)
 		{
-			ObjectToDirty->ObjectReferenceComponents[0]->ModifyIntA();
+			// There's only two objects, so when we update one we dirty the other
+			if (ObjectToDirty != ReplicatedObject)
+			{
+				ObjectToDirty->ObjectReferenceComponents[0]->ModifyIntA();
+			}
 		}
 	};
 	
@@ -373,11 +382,14 @@ UE_NET_TEST_FIXTURE(FReplicationSystemServerClientTestFixture, PushModelMarkSelf
 	UTestReplicatedIrisObject* ClientObject = Cast<UTestReplicatedIrisObject>(Client->GetReplicationBridge()->GetReplicatedObject(ServerObject->NetRefHandle));
 	UE_NET_ASSERT_NE(ClientObject, nullptr);
 
-	auto PreUpdateObject = [&](FNetRefHandle NetHandle, UObject* InObject, const UReplicationBridge* ReplicationBridge)
+	auto PreUpdateObject = [&](TArrayView<UObject*> Instances, const UReplicationBridge* Bridge)
 	{
-		if (InObject == ServerObject)
+		for (UObject* InObject : Instances)
 		{
-			ServerObject->ObjectReferenceComponents[0]->ModifyIntA();
+			if (InObject == ServerObject)
+			{
+				ServerObject->ObjectReferenceComponents[0]->ModifyIntA();
+			}
 		}
 	};
 
diff --git a/Engine/Plugins/Runtime/ReplicationSystemTestPlugin/Source/Private/Tests/ReplicationSystem/TestPreUpdateCallback.cpp b/Engine/Plugins/Runtime/ReplicationSystemTestPlugin/Source/Private/Tests/ReplicationSystem/TestPreUpdateCallback.cpp
--- a/Engine/Plugins/Runtime/ReplicationSystemTestPlugin/Source/Private/Tests/ReplicationSystem/TestPreUpdateCallback.cpp
+++ b/Engine/Plugins/Runtime/ReplicationSystemTestPlugin/Source/Private/Tests/ReplicationSystem/TestPreUpdateCallback.cpp
@@ -39,21 +39,24 @@ UE_NET_TEST_FIXTURE(FReplicationSystemServerClientTestFixture, CreateObjectsInsi
 	UE_NET_ASSERT_NE(ClientObject, nullptr);
 
 	// Add a PreUpdate where we create a new subobject
-	auto PreUpdateObject = [&](FNetRefHandle NetHandle, UObject* ReplicatedObject, const UReplicationBridge* ReplicationBridge)
+	auto PreUpdateObject = [&](TArrayView<UObject*> Instances, const UReplicationBridge* Bridge)
 	{
-		if (ServerObject == ReplicatedObject)
+		for (UObject* ReplicatedObject : Instances)
 		{
-			if (ServerSubObject == nullptr)
+			if (ServerObject == ReplicatedObject)
 			{
-				// Dirty a property
-				ServerObject->IntA = 0xBB;
+				if (ServerSubObject == nullptr)
+				{
+					// Dirty a property
+					ServerObject->IntA = 0xBB;
 
-				// Create a subobject
-				ServerSubObject = Server->CreateSubObject<UTestReplicatedIrisObject>(ServerObject->NetRefHandle);
+					// Create a subobject
+					ServerSubObject = Server->CreateSubObject<UTestReplicatedIrisObject>(ServerObject->NetRefHandle);
 
-				// Dirty this subobject
-				ServerSubObject->IntA = 0xBB;
+					// Dirty this subobject
+					ServerSubObject->IntA = 0xBB;
 
+				}
 			}
 		}
 	};
diff --git a/Engine/Source/Runtime/Engine/Private/Net/Iris/ReplicationSystem/ActorReplicationBridge.cpp b/Engine/Source/Runtime/Engine/Private/Net/Iris/ReplicationSystem/ActorReplicationBridge.cpp
--- a/Engine/Source/Runtime/Engine/Private/Net/Iris/ReplicationSystem/ActorReplicationBridge.cpp
+++ b/Engine/Source/Runtime/Engine/Private/Net/Iris/ReplicationSystem/ActorReplicationBridge.cpp
@@ -63,17 +63,16 @@ bool IsActorValidForIrisReplication(const AActor* Actor)
 	return IsValid(Actor) && !Actor->IsActorBeingDestroyed() && !Actor->IsUnreachable();
 }
 
-void ActorReplicationBridgePreUpdateFunction(FNetRefHandle Handle, UObject* Instance, const UReplicationBridge* Bridge)
+void ActorReplicationBridgePreUpdateFunction(TArrayView<UObject*> Instances, const UReplicationBridge* Bridge)
 {
-	AActor* Actor = Cast<AActor>(Instance);
-	if (IsActorValidForIrisReplication(Actor))
-	{
-		/** $IRIS TODO:
-		 * Here we need to call something that either fakes CallPrereplication or we need to mimic the interfaces of NetDriver/PropertyChangeTracker
-		 * however we do think any state changes should be pushed to the network system rather than polled
-		 * One example is the gatherCurrent() movement, movementdata should have its own ReplicationState that is updated when changed. (or updated using a helper method)
-		 */
-		Actor->CallPreReplication(CastChecked<const UActorReplicationBridge>(Bridge)->GetNetDriver());
+	UNetDriver* NetDriver = CastChecked<const UActorReplicationBridge>(Bridge)->GetNetDriver();
+	for (UObject* Instance : Instances)
+	{
+		AActor* Actor = Cast<AActor>(Instance);
+		if (IsActorValidForIrisReplication(Actor))
+		{
+			Actor->CallPreReplication(NetDriver);
+		}
 	}
 }
 
diff --git a/Engine/Source/Runtime/Experimental/Iris/Core/Private/Iris/ReplicationSystem/NetRefHandleManager.cpp b/Engine/Source/Runtime/Experimental/Iris/Core/Private/Iris/ReplicationSystem/NetRefHandleManager.cpp
--- a/Engine/Source/Runtime/Experimental/Iris/Core/Private/Iris/ReplicationSystem/NetRefHandleManager.cpp
+++ b/Engine/Source/Runtime/Experimental/Iris/Core/Private/Iris/ReplicationSystem/NetRefHandleManager.cpp
@@ -37,6 +37,7 @@ FNetRefHandleManager::FNetRefHandleManager(FReplicationProtocolManager& InReplic
 , ObjectsWithDependentObjectsInternalIndices(MaxActiveObjectCount)
 , DestroyedStartupObjectInternalIndices(MaxActiveObjectCount)
 , WantToBeDormantInternalIndices(MaxActiveObjectCount)
+, ObjectsWithPreUpdate(MaxActiveObjectCount)
 , NextStaticHandleIndex(1) // Index 0 is always reserved, for both static and dynamic handles
 , NextDynamicHandleIndex(1)
 , ReplicationProtocolManager(InReplicationProtocolManager)
@@ -147,6 +148,7 @@ FInternalNetRefIndex FNetRefHandleManager::InternalCreateNetObject(const FNetRef
 		Data.NetHandle = GlobalHandle;
 		Data.Protocol = ReplicationProtocol;
 		Data.InstanceProtocol = nullptr;
+		ObjectsWithPreUpdate.ClearBit(InternalIndex);
 		ReplicatedObjectStateBuffers[InternalIndex] = nullptr;
 		Data.ReceiveStateBuffer = nullptr;
 		Data.bShouldPropagateChangedStates = 1U;
@@ -197,6 +199,8 @@ void FNetRefHandleManager::AttachInstanceProtocol(FInternalNetRefIndex InternalI
 
 		check(ReplicatedInstances[InternalIndex] == nullptr);
 		ReplicatedInstances[InternalIndex] = Instance;
+
+		ObjectsWithPreUpdate.SetBitValue(InternalIndex, EnumHasAnyFlags(InstanceProtocol->InstanceTraits, EReplicationInstanceProtocolTraits::NeedsPreSendUpdate));
 	}
 }
 
@@ -217,6 +221,8 @@ const FReplicationInstanceProtocol* FNetRefHandleManager::DetachInstanceProtocol
 		
 		Data.InstanceProtocol = nullptr;
 		ReplicatedInstances[InternalIndex] = nullptr;
+
+		ObjectsWithPreUpdate.ClearBit(InternalIndex);
 		
 		return InstanceProtocol;
 	}
diff --git a/Engine/Source/Runtime/Experimental/Iris/Core/Private/Iris/ReplicationSystem/NetRefHandleManager.h b/Engine/Source/Runtime/Experimental/Iris/Core/Private/Iris/ReplicationSystem/NetRefHandleManager.h
--- a/Engine/Source/Runtime/Experimental/Iris/Core/Private/Iris/ReplicationSystem/NetRefHandleManager.h
+++ b/Engine/Source/Runtime/Experimental/Iris/Core/Private/Iris/ReplicationSystem/NetRefHandleManager.h
@@ -300,6 +300,9 @@ public:
 	/** Return the largest internal index that has been used. */
 	uint32 GetLargestInternalIndex() const { return LargestInternalIndex; };
 
+	/** Get Objects that is flagged for PreUpdate (aka PreReplication) */
+	FNetBitArrayView GetObjectsWithPreUpdate() const { return MakeNetBitArrayView(ObjectsWithPreUpdate); }
+
 public:
 
 	// Iterate over all dependent objects and their dependent objects
@@ -411,6 +414,9 @@ private:
 	// Bitset marking internal indices that wants to be dormant
 	FNetBitArray WantToBeDormantInternalIndices;
 
+	// Bitset marking internal indices that is flagged as requiring PreUpdate
+	FNetBitArray ObjectsWithPreUpdate;
+
 	// Both are marked in DestroyedStartupObjectsInternalIndicies
 	// When one is destroyed both mapping will be cleared
 	// Map OriginalIndex -> DestructionInfoIndex
diff --git a/Engine/Source/Runtime/Experimental/Iris/Core/Private/Iris/ReplicationSystem/Polling/ObjectPoller.cpp b/Engine/Source/Runtime/Experimental/Iris/Core/Private/Iris/ReplicationSystem/Polling/ObjectPoller.cpp
--- a/Engine/Source/Runtime/Experimental/Iris/Core/Private/Iris/ReplicationSystem/Polling/ObjectPoller.cpp
+++ b/Engine/Source/Runtime/Experimental/Iris/Core/Private/Iris/ReplicationSystem/Polling/ObjectPoller.cpp
@@ -14,6 +14,7 @@
 #include "Iris/Stats/NetStatsContext.h"
 
 #include "Net/Core/Trace/NetDebugName.h"
+#include "HAL/IConsoleManager.h"
 
 namespace UE::Net::Private
 {
@@ -35,35 +36,48 @@ FObjectPoller::FObjectPoller(const FInitParams& InitParams)
 void FObjectPoller::PreUpdatePass(const FNetBitArrayView& ObjectsConsideredForPolling)
 {
 	IRIS_PROFILER_SCOPE_VERBOSE(PreUpdatePass);
-	NetStatsContext = ReplicationSystemInternal->GetNetTypeStats().GetNetStatsContext();
 
-	ObjectsConsideredForPolling.ForAllSetBits([this](FInternalNetRefIndex Objectindex)
-	{
-		CallPreUpdate(Objectindex);
-	});
+	constexpr uint32 PreUpdateBatchCount = 128U;
 
-	NetStatsContext = nullptr;
-}
+	UObject* BatchedObjects[PreUpdateBatchCount];
+	uint32 BatchedObjectCount = 0U;
 
-void FObjectPoller::CallPreUpdate(FInternalNetRefIndex ObjectIndex)
-{
-	FNetRefHandleManager::FReplicatedObjectData& ObjectData = LocalNetRefHandleManager.GetReplicatedObjectDataNoCheck(ObjectIndex);
-	if (UNLIKELY(ObjectData.InstanceProtocol == nullptr))
+	if (!ObjectReplicationBridge->PreUpdateInstanceFunction)
 	{
 		return;
 	}
+	
+	auto BatchedPreUpdate = [this, &BatchedObjects, &BatchedObjectCount](FInternalNetRefIndex Objectindex)
+	{
+		// Flush if needed
+		if (BatchedObjectCount == PreUpdateBatchCount)
+		{
+			ObjectReplicationBridge->PreUpdateInstanceFunction(MakeArrayView<UObject*>(BatchedObjects, BatchedObjectCount), ObjectReplicationBridge);
+			PollStats.PreUpdatedObjectCount += BatchedObjectCount;
+			BatchedObjectCount = 0U;
+		}
+		UObject* Instance = ReplicatedInstances[Objectindex];
+		BatchedObjects[BatchedObjectCount] = Instance;
+		BatchedObjectCount += Instance ? 1U : 0U;
+	};
 
-	IRIS_PROFILER_PROTOCOL_NAME(ObjectData.Protocol->DebugName->Name);
+	FNetBitArrayView::ForAllSetBits(ObjectsConsideredForPolling, LocalNetRefHandleManager.GetObjectsWithPreUpdate(), FNetBitArrayView::AndOp, BatchedPreUpdate);
 
-	// Call per-instance PreUpdate function
-	if (ObjectReplicationBridge->PreUpdateInstanceFunction && EnumHasAnyFlags(ObjectData.InstanceProtocol->InstanceTraits, EReplicationInstanceProtocolTraits::NeedsPreSendUpdate))
+	// Flush last batch
+	if (BatchedObjectCount > 0)
 	{
-		UE_NET_IRIS_STATS_TIMER(Timer, NetStatsContext);
-	
-		ObjectReplicationBridge->PreUpdateInstanceFunction(ObjectData.RefHandle, ReplicatedInstances[ObjectIndex], ObjectReplicationBridge);
+		ObjectReplicationBridge->PreUpdateInstanceFunction(MakeArrayView(BatchedObjects, BatchedObjectCount), ObjectReplicationBridge);
+		PollStats.PreUpdatedObjectCount += BatchedObjectCount;
+	}
+}
+
+void FObjectPoller::CallPreUpdate(FInternalNetRefIndex ObjectIndex)
+{
+	UObject* Instance = ReplicatedInstances[ObjectIndex];
+	if (Instance && LocalNetRefHandleManager.GetObjectsWithPreUpdate().GetBit(ObjectIndex))
+	{
+		ObjectReplicationBridge->PreUpdateInstanceFunction(MakeArrayView<UObject*>(&Instance, 1U), ObjectReplicationBridge);	
 		++PollStats.PreUpdatedObjectCount;
-		
-		UE_NET_IRIS_STATS_ADD_TIME_AND_COUNT_FOR_OBJECT(Timer, PreUpdate, ObjectIndex);
 	}
 }
 
diff --git a/Engine/Source/Runtime/Experimental/Iris/Core/Public/Iris/ReplicationSystem/ObjectReplicationBridge.h b/Engine/Source/Runtime/Experimental/Iris/Core/Public/Iris/ReplicationSystem/ObjectReplicationBridge.h
--- a/Engine/Source/Runtime/Experimental/Iris/Core/Public/Iris/ReplicationSystem/ObjectReplicationBridge.h
+++ b/Engine/Source/Runtime/Experimental/Iris/Core/Public/Iris/ReplicationSystem/ObjectReplicationBridge.h
@@ -8,6 +8,7 @@
 #include "Iris/ReplicationSystem/ReplicationBridge.h"
 #include "Delegates/IDelegateInstance.h"
 #include "Iris/ReplicationSystem/ReplicationSystemTypes.h"
+#include "Containers/ArrayView.h"
 
 #include "ObjectReplicationBridge.generated.h"
 
@@ -262,7 +263,7 @@ protected:
 	IRISCORE_API UObject* GetObjectFromReferenceHandle(FNetRefHandle RefHandle) const;
 
 	/** Helper method that calls provided PreUpdateFunction and polls state data for all replicated instances with the NeedsPoll trait. */
-	using FInstancePreUpdateFunction = TFunction<void(FNetRefHandle, UObject*, const UReplicationBridge*)>;
+	using FInstancePreUpdateFunction = TFunction<void(TArrayView<UObject*>, const UReplicationBridge*)>;
 
 	/** Set the function that we should call before copying state data. */
 	IRISCORE_API void SetInstancePreUpdateFunction(FInstancePreUpdateFunction InPreUpdateFunction);
