Iris - Fixed Debug runtime issue with UObjects with identical names.

#rnx

diff --git a/Engine/Plugins/Runtime/ReplicationSystemTestPlugin/Source/Private/Tests/ReplicationSystem/TestDependentObjects.cpp b/Engine/Plugins/Runtime/ReplicationSystemTestPlugin/Source/Private/Tests/ReplicationSystem/TestDependentObjects.cpp
--- a/Engine/Plugins/Runtime/ReplicationSystemTestPlugin/Source/Private/Tests/ReplicationSystem/TestDependentObjects.cpp
+++ b/Engine/Plugins/Runtime/ReplicationSystemTestPlugin/Source/Private/Tests/ReplicationSystem/TestDependentObjects.cpp
@@ -59,7 +59,7 @@ private:
 
 		{
 			FNetObjectFilterDefinition& MockDefinition = NewFilterDefinitions.Emplace_GetRef();
-			MockDefinition.FilterName = "Mock";
+			MockDefinition.FilterName = "MockFilter";
 			MockDefinition.ClassName = "/Script/ReplicationSystemTestPlugin.MockNetObjectFilter";
 			MockDefinition.ConfigClassName = "/Script/ReplicationSystemTestPlugin.MockNetObjectFilterConfig";
 		}
@@ -84,8 +84,8 @@ private:
 	void InitFilterHandles()
 	{
 		NotRoutedFilterHandle = Server->GetReplicationSystem()->GetFilterHandle("NotRouted");
-		MockFilterHandle = Server->GetReplicationSystem()->GetFilterHandle("Mock");
-		MockFilter = Cast<UMockNetObjectFilter>(Server->GetReplicationSystem()->GetFilter("Mock"));
+		MockFilterHandle = Server->GetReplicationSystem()->GetFilterHandle("MockFilter");
+		MockFilter = Cast<UMockNetObjectFilter>(Server->GetReplicationSystem()->GetFilter("MockFilter"));
 	}
 
 private:
