Iris - Object scope hysteresis support.  Keep dynamically filtered out objects around for a specified amount of frames. Configure behavior via hysteresis profiles that determine the frame timeout per class. The filter config for a specific class can then mention the hysteresis profile in order to get the appropriate behavior.  Config example is present in ReplicationFilteringConfig.h
Note: Hysteresis is enabled by default with a frame count of 6.

diff --git a/Engine/Config/BaseEngine.ini b/Engine/Config/BaseEngine.ini
--- a/Engine/Config/BaseEngine.ini
+++ b/Engine/Config/BaseEngine.ini
@@ -1322,6 +1322,10 @@ MaxPartCount=4096
 +NetObjectFilterDefinitions=(FilterName=Spatial, ClassName=/Script/IrisCore.NetObjectGridWorldLocFilter, ConfigClassName=/Script/IrisCore.NetObjectGridFilterConfig)
 +NetObjectFilterDefinitions=(FilterName=NotRouted, ClassName=/Script/IrisCore.FilterOutNetObjectFilter, ConfigClassName=/Script/IrisCore.FilterOutNetObjectFilterConfig)
 
+[/Script/IrisCore.ReplicationFilteringConfig]
+bEnableObjectScopeHysteresis=true
+DefaultHysteresisFrameCount=6
+
 [/Script/IrisCore.ObjectReplicationBridgeConfig]
 RequiredNetDriverChannelClassName=/Script/Engine.DataStreamChannel
 
diff --git a/Engine/Plugins/Runtime/ReplicationSystemTestPlugin/Source/Private/Tests/ReplicationSystem/Filtering/TestFiltering.cpp b/Engine/Plugins/Runtime/ReplicationSystemTestPlugin/Source/Private/Tests/ReplicationSystem/Filtering/TestFiltering.cpp
--- a/Engine/Plugins/Runtime/ReplicationSystemTestPlugin/Source/Private/Tests/ReplicationSystem/Filtering/TestFiltering.cpp
+++ b/Engine/Plugins/Runtime/ReplicationSystemTestPlugin/Source/Private/Tests/ReplicationSystem/Filtering/TestFiltering.cpp
@@ -2,9 +2,11 @@
 
 #include "NetworkAutomationTest.h"
 #include "NetworkAutomationTestMacros.h"
+#include "Iris/IrisConstants.h"
 #include "Iris/ReplicationSystem/ReplicationSystem.h"
 #include "Iris/ReplicationSystem/Filtering/NetObjectFilter.h"
 #include "Iris/ReplicationSystem/Filtering/NetObjectFilterDefinitions.h"
+#include "Iris/ReplicationSystem/Filtering/ReplicationFilteringConfig.h"
 #include "Tests/ReplicationSystem/Filtering/MockNetObjectFilter.h"
 #include "Tests/ReplicationSystem/Filtering/TestFilteringObject.h"
 #include "Tests/ReplicationSystem/ReplicationSystemServerClientTestFixture.h"
@@ -18,6 +20,7 @@ protected:
 	virtual void SetUp() override
 	{
 		InitNetObjectFilterDefinitions();
+		InitObjectScopeHysteresisProfiles();
 		FReplicationSystemServerClientTestFixture::SetUp();
 		InitMockNetObjectFilter();
 	}
@@ -25,6 +28,7 @@ protected:
 	virtual void TearDown() override
 	{
 		FReplicationSystemServerClientTestFixture::TearDown();
+		RestoreObjectScopeHysteresisProfiles();
 		RestoreFilterDefinitions();
 	}
 
@@ -41,6 +45,87 @@ protected:
 		MockNetObjectFilter->ResetFunctionCallStatus();
 	}
 
+	uint32 GetHysteresisFrameCount(const char* ProfileName) const
+	{
+		const UReplicationFilteringConfig* Config = GetDefault<UReplicationFilteringConfig>();
+		if (const FObjectScopeHysteresisProfile* Profile = Config->GetHysteresisProfiles().FindByKey(ProfileName))
+		{
+			return Profile->HysteresisFrameCount;
+		}
+
+		return Config->GetDefaultHysteresisFrameCount();
+	}
+
+	class FScopedDefaultHysteresisFrameCount
+	{
+	public:
+		FScopedDefaultHysteresisFrameCount(uint8 DefaultHysteresisFrameCount)
+		{
+			const UClass* NetObjectFilteringConfiglass = UReplicationFilteringConfig::StaticClass();
+			if (const FProperty* Property = NetObjectFilteringConfiglass->FindPropertyByName("DefaultHysteresisFrameCount"))
+			{
+				UReplicationFilteringConfig* FilteringConfig = GetMutableDefault<UReplicationFilteringConfig>();
+				Property->CopyCompleteValue(&PrevValue, (void*)(UPTRINT(FilteringConfig) + Property->GetOffset_ForInternal()));
+				Property->CopyCompleteValue((void*)(UPTRINT(FilteringConfig) + Property->GetOffset_ForInternal()), &DefaultHysteresisFrameCount);
+				bPrevValueIsValid = true;
+			}
+		}
+
+		~FScopedDefaultHysteresisFrameCount()
+		{
+			if (!bPrevValueIsValid)
+			{
+				return;
+			}
+
+			const UClass* NetObjectFilteringConfiglass = UReplicationFilteringConfig::StaticClass();
+			if (const FProperty* Property = NetObjectFilteringConfiglass->FindPropertyByName("DefaultHysteresisFrameCount"))
+			{
+				UReplicationFilteringConfig* FilteringConfig = GetMutableDefault<UReplicationFilteringConfig>();
+				Property->CopyCompleteValue((void*)(UPTRINT(FilteringConfig) + Property->GetOffset_ForInternal()), &PrevValue);
+				bPrevValueIsValid = true;
+			}
+		}
+
+		uint8 PrevValue = 0;
+		bool bPrevValueIsValid = false;
+	};
+
+	class FScopedHysteresisUpdateConnectionThrottling
+	{
+	public:
+		FScopedHysteresisUpdateConnectionThrottling(uint8 HysteresisUpdateConnectionThrottling)
+		{
+			const UClass* NetObjectFilteringConfiglass = UReplicationFilteringConfig::StaticClass();
+			if (const FProperty* Property = NetObjectFilteringConfiglass->FindPropertyByName("HysteresisUpdateConnectionThrottling"))
+			{
+				UReplicationFilteringConfig* FilteringConfig = GetMutableDefault<UReplicationFilteringConfig>();
+				Property->CopyCompleteValue(&PrevValue, (void*)(UPTRINT(FilteringConfig) + Property->GetOffset_ForInternal()));
+				Property->CopyCompleteValue((void*)(UPTRINT(FilteringConfig) + Property->GetOffset_ForInternal()), &HysteresisUpdateConnectionThrottling);
+				bPrevValueIsValid = true;
+			}
+		}
+
+		~FScopedHysteresisUpdateConnectionThrottling()
+		{
+			if (!bPrevValueIsValid)
+			{
+				return;
+			}
+
+			const UClass* NetObjectFilteringConfiglass = UReplicationFilteringConfig::StaticClass();
+			if (const FProperty* Property = NetObjectFilteringConfiglass->FindPropertyByName("HysteresisUpdateConnectionThrottling"))
+			{
+				UReplicationFilteringConfig* FilteringConfig = GetMutableDefault<UReplicationFilteringConfig>();
+				Property->CopyCompleteValue((void*)(UPTRINT(FilteringConfig) + Property->GetOffset_ForInternal()), &PrevValue);
+				bPrevValueIsValid = true;
+			}
+		}
+
+		uint8 PrevValue = 0;
+		bool bPrevValueIsValid = false;
+	};
+
 private:
 	void InitNetObjectFilterDefinitions()
 	{
@@ -77,18 +162,64 @@ private:
 		MockNetObjectFilter = nullptr;
 	}
 
+	void InitObjectScopeHysteresisProfiles()
+	{
+		const UClass* NetObjectFilteringConfiglass = UReplicationFilteringConfig::StaticClass();
+		const FProperty* ProfilesProperty = NetObjectFilteringConfiglass->FindPropertyByName("HysteresisProfiles");
+		check(ProfilesProperty != nullptr);
+
+		// Save CDO state.
+		UReplicationFilteringConfig* FilteringConfig = GetMutableDefault<UReplicationFilteringConfig>();
+		ProfilesProperty->CopyCompleteValue(&OriginalObjectScopeHysteresisProfiles, (void*)(UPTRINT(FilteringConfig) + ProfilesProperty->GetOffset_ForInternal()));
+
+		// Modify profiles to what we need
+		TArray<FObjectScopeHysteresisProfile> NewProfiles;
+		{
+			{
+				FObjectScopeHysteresisProfile& Profile = NewProfiles.Emplace_GetRef();
+				Profile.FilterProfileName = "FiveFrames";
+				Profile.HysteresisFrameCount = 5;
+			}
+
+			{
+				FObjectScopeHysteresisProfile& Profile = NewProfiles.Emplace_GetRef();
+				Profile.FilterProfileName = "OneFrame";
+				Profile.HysteresisFrameCount = 1;
+			}
+
+			{
+				FObjectScopeHysteresisProfile& Profile = NewProfiles.Emplace_GetRef();
+				Profile.FilterProfileName = "ZeroFrames";
+				Profile.HysteresisFrameCount = 0;
+			}
+		}
+
+		ProfilesProperty->CopyCompleteValue((void*)(UPTRINT(FilteringConfig) + ProfilesProperty->GetOffset_ForInternal()), &NewProfiles);
+	}
+
+	void RestoreObjectScopeHysteresisProfiles()
+	{
+		// Restore CDO state from the saved state.
+		const UClass* NetObjectFilterDefinitionsClass = UReplicationFilteringConfig::StaticClass();
+		const FProperty* ProfilesProperty = NetObjectFilterDefinitionsClass->FindPropertyByName("HysteresisProfiles");
+		UReplicationFilteringConfig* FilteringConfig = GetMutableDefault<UReplicationFilteringConfig>();
+		ProfilesProperty->CopyCompleteValue((void*)(UPTRINT(FilteringConfig) + ProfilesProperty->GetOffset_ForInternal()), &OriginalObjectScopeHysteresisProfiles);
+		OriginalObjectScopeHysteresisProfiles.Empty();
+	}
+
 	void InitMockNetObjectFilter()
 	{
 		MockNetObjectFilter = CastChecked<UMockNetObjectFilter>(Server->GetReplicationSystem()->GetFilter(GetMockFilterName()));
 		MockFilterHandle = Server->GetReplicationSystem()->GetFilterHandle(GetMockFilterName());
 	}
-
 protected:
 	UMockNetObjectFilter* MockNetObjectFilter;
 	FNetObjectFilterHandle MockFilterHandle;
+	FName ObjectHysteresisProfileName;
 
 private:
 	TArray<FNetObjectFilterDefinition> OriginalFilterDefinitions;
+	TArray<FObjectScopeHysteresisProfile> OriginalObjectScopeHysteresisProfiles;
 };
 
 
@@ -1035,7 +1166,7 @@ UE_NET_TEST_FIXTURE(FTestFilteringFixture, NestedDependentObjectIsFilteredAsPare
 	Server->SendAndDeliverTo(Client, DeliverPacket);
 	Server->PostSendUpdate();
 
-	// We expect the object to exist and the future dependent and future dependent objects not to exist
+	// We expect the object to exist and the future dependent and future nested dependent objects not to exist
 	UE_NET_ASSERT_NE(Client->GetReplicationBridge()->GetReplicatedObject(ServerObject->NetRefHandle), nullptr);
 	UE_NET_ASSERT_EQ(Client->GetReplicationBridge()->GetReplicatedObject(ServerFutureDependentObject->NetRefHandle), nullptr);
 	UE_NET_ASSERT_EQ(Client->GetReplicationBridge()->GetReplicatedObject(ServerFutureNestedDependentObject->NetRefHandle), nullptr);
@@ -2274,4 +2405,300 @@ UE_NET_TEST_FIXTURE(FTestFilteringFixture, LateAddedConnectionWorksWithComplexGr
 	UE_NET_ASSERT_NE(LateAddedClient->GetReplicationBridge()->GetReplicatedObject(ServerSubObjects[3]->NetRefHandle), nullptr);
 }
 
+// Dynamic filtering should cause hysteresis to kick in for a filtered out object with a filter profile.
+UE_NET_TEST_FIXTURE(FTestFilteringFixture, HysteresisKicksInForDynamicallyFilteredOutObjectWithFilterProfile)
+{
+	// Add client
+	FReplicationSystemTestClient* Client = CreateClient();
+
+	SetDynamicFilterStatus(ENetFilterStatus::Allow);
+
+	// Spawn object on server and set filter and filter profile for hysteresis
+	UReplicatedTestObject* ServerObject = Server->CreateObject({.IrisComponentCount = 0});
+	Server->ReplicationSystem->SetFilter(ServerObject->NetRefHandle, MockFilterHandle, /* .FilterProfile = */ "FiveFrames");
+
+	Server->UpdateAndSend({ Client });
+
+	SetDynamicFilterStatus(ENetFilterStatus::Disallow);
+
+	const uint32 HysteresisFrameCount = GetHysteresisFrameCount("FiveFrames");
+	UE_NET_ASSERT_EQ(HysteresisFrameCount, 5U);
+	for (uint32 It = 0, EndIt = HysteresisFrameCount; It < EndIt; ++It)
+	{
+		Server->UpdateAndSend({ Client });
+	}
+
+	// At this point the object should still exist on the client due to hysteresis
+	UE_NET_ASSERT_NE(Client->GetReplicationBridge()->GetReplicatedObject(ServerObject->NetRefHandle), nullptr);
+
+	Server->UpdateAndSend({ Client });
+
+	UE_NET_ASSERT_EQ(Client->GetReplicationBridge()->GetReplicatedObject(ServerObject->NetRefHandle), nullptr);
+}
+
+// Dynamic filtering should cause hysteresis to kick in for a filtered out object without a filter profile, thus using default hysteresis frame count.
+UE_NET_TEST_FIXTURE(FTestFilteringFixture, HysteresisKicksInForDynamicallyFilteredOutObjectWithoutFilterProfile)
+{
+	constexpr uint32 DefaultHystersisFrameCount = 3;
+	FScopedDefaultHysteresisFrameCount ScopedDefaultHysteresisFrameCount(DefaultHystersisFrameCount);
+
+	// Add client
+	FReplicationSystemTestClient* Client = CreateClient();
+
+	SetDynamicFilterStatus(ENetFilterStatus::Allow);
+
+	// Spawn object on server and set filter and filter profile for hysteresis
+	UReplicatedTestObject* ServerObject = Server->CreateObject({ .IrisComponentCount = 0 });
+	Server->ReplicationSystem->SetFilter(ServerObject->NetRefHandle, MockFilterHandle);
+
+	Server->UpdateAndSend({ Client });
+
+	SetDynamicFilterStatus(ENetFilterStatus::Disallow);
+
+	for (uint32 It = 0, EndIt = DefaultHystersisFrameCount; It < EndIt; ++It)
+	{
+		Server->UpdateAndSend({ Client });
+	}
+
+	// At this point the object should still exist on the client due to hysteresis
+	UE_NET_ASSERT_NE(Client->GetReplicationBridge()->GetReplicatedObject(ServerObject->NetRefHandle), nullptr);
+
+	Server->UpdateAndSend({ Client });
+
+	UE_NET_ASSERT_EQ(Client->GetReplicationBridge()->GetReplicatedObject(ServerObject->NetRefHandle), nullptr);
+}
+
+// Owner filtering changes should not cause hysteresis to kick in.
+UE_NET_TEST_FIXTURE(FTestFilteringFixture, HysteresisDoesNotKickInForOwnerFilteredObject)
+{
+	constexpr uint32 DefaultHystersisFrameCount = 3;
+	FScopedDefaultHysteresisFrameCount ScopedDefaultHysteresisFrameCount(DefaultHystersisFrameCount);
+
+	// Add client
+	FReplicationSystemTestClient* Client = CreateClient();
+
+	// Spawn object on server
+	UReplicatedTestObject* ServerObject = Server->CreateObject({ .IrisComponentCount = 0 });
+	Server->ReplicationSystem->SetFilter(ServerObject->NetRefHandle, ToOwnerFilterHandle);
+	Server->ReplicationSystem->SetOwningNetConnection(ServerObject->NetRefHandle, Client->ConnectionIdOnServer);
+
+	// Send and deliver packet
+	Server->UpdateAndSend({ Client });
+
+	// Object should have been created on the client
+	UE_NET_ASSERT_NE(Client->GetReplicationBridge()->GetReplicatedObject(ServerObject->NetRefHandle), nullptr);
+
+	// Switch owner and make sure objects gets immediately destroyed on the client
+	Server->ReplicationSystem->SetOwningNetConnection(ServerObject->NetRefHandle, InvalidConnectionId);
+
+	// Send and deliver packet
+	Server->UpdateAndSend({ Client });
+
+	UE_NET_ASSERT_EQ(Client->GetReplicationBridge()->GetReplicatedObject(ServerObject->NetRefHandle), nullptr);
+}
+
+// Exclusion group filtering changes should not cause hysteresis to kick in.
+UE_NET_TEST_FIXTURE(FTestFilteringFixture, HysteresisDoesNotKickInForExclusionGroupFilteredObject)
+{
+	constexpr uint32 DefaultHystersisFrameCount = 3;
+	FScopedDefaultHysteresisFrameCount ScopedDefaultHysteresisFrameCount(DefaultHystersisFrameCount);
+
+	// Add client
+	FReplicationSystemTestClient* Client = CreateClient();
+
+	// Spawn object on server
+	UReplicatedTestObject* ServerObject = Server->CreateObject({ .IrisComponentCount = 0 });
+
+	// Add to exclusion group that allows replication to all connections.
+	FNetObjectGroupHandle GroupHandle = Server->ReplicationSystem->CreateGroup();
+	Server->ReplicationSystem->AddToGroup(GroupHandle, ServerObject->NetRefHandle);
+	Server->ReplicationSystem->AddExclusionFilterGroup(GroupHandle);
+	Server->ReplicationSystem->SetGroupFilterStatus(GroupHandle, ENetFilterStatus::Allow);
+
+	// Send and deliver packet
+	Server->UpdateAndSend({ Client });
+
+	// Object should have been created on the client
+	UE_NET_ASSERT_NE(Client->GetReplicationBridge()->GetReplicatedObject(ServerObject->NetRefHandle), nullptr);
+
+	// Disallow the group to be replicated.
+	Server->ReplicationSystem->SetGroupFilterStatus(GroupHandle, ENetFilterStatus::Disallow);
+
+	// Send and deliver packet
+	Server->UpdateAndSend({ Client });
+
+	// Client object should be destroyed
+	UE_NET_ASSERT_EQ(Client->GetReplicationBridge()->GetReplicatedObject(ServerObject->NetRefHandle), nullptr);
+}
+
+// If an object is filtered out from the start we should not start replicating it at all
+UE_NET_TEST_FIXTURE(FTestFilteringFixture, HysteresisDoesNotKickInForNewlyCreatedObject)
+{
+	// Add client
+	FReplicationSystemTestClient* Client = CreateClient();
+
+	SetDynamicFilterStatus(ENetFilterStatus::Disallow);
+
+	// Spawn object on server and set filter and filter profile for hysteresis
+	UReplicatedTestObject* ServerObject = Server->CreateObject({ .IrisComponentCount = 0 });
+	Server->ReplicationSystem->SetFilter(ServerObject->NetRefHandle, MockFilterHandle, /* .FilterProfile = */ "FiveFrames");
+
+	Server->UpdateAndSend({ Client });
+
+	UE_NET_ASSERT_NE(Client->GetReplicationBridge()->GetReplicatedObject(ServerObject->NetRefHandle), nullptr);
+}
+
+// Destroyed objects are expected to be destroyed as quickly as possible.
+UE_NET_TEST_FIXTURE(FTestFilteringFixture, HysteresisDoesNotKickInForDestroyedObject)
+{
+	// Add client
+	FReplicationSystemTestClient* Client = CreateClient();
+
+	SetDynamicFilterStatus(ENetFilterStatus::Allow);
+
+	// Spawn object on server and set filter and filter profile for hysteresis
+	UReplicatedTestObject* ServerObject = Server->CreateObject({ .IrisComponentCount = 0 });
+	Server->ReplicationSystem->SetFilter(ServerObject->NetRefHandle, MockFilterHandle, /* .FilterProfile = */ "FiveFrames");
+
+	Server->UpdateAndSend({ Client });
+
+	// Destroy object and make sure it immediately gets destroyed on the client as well.
+	const FNetRefHandle ServerNetRefHandle = ServerObject->NetRefHandle;
+	Server->DestroyObject(ServerObject);
+	ServerObject = nullptr;
+
+	Server->UpdateAndSend({ Client });
+
+	UE_NET_ASSERT_EQ(Client->GetReplicationBridge()->GetReplicatedObject(ServerNetRefHandle), nullptr);
+}
+
+// Make sure that connection throttling does not cause objects to be filtered out too soon. Also verify throttling occurs.
+UE_NET_TEST_FIXTURE(FTestFilteringFixture, HysteresisConnectionThrottlingWorksAsExpected)
+{
+	constexpr uint8 ConnectionThrottlingFrameCount = 5;
+	FScopedHysteresisUpdateConnectionThrottling ConnectionThrottling(ConnectionThrottlingFrameCount);
+
+	// Add client
+	FReplicationSystemTestClient* Client = CreateClient();
+
+	// Figure out hysteresis update frame.
+	{
+		SetDynamicFilterStatus(ENetFilterStatus::Allow);
+
+		// Spawn object on server and set filter and filter profile for hysteresis
+		UReplicatedTestObject* ServerObject = Server->CreateObject({ .IrisComponentCount = 0 });
+		Server->ReplicationSystem->SetFilter(ServerObject->NetRefHandle, MockFilterHandle, /* .FilterProfile = */ "OneFrame");
+
+		Server->UpdateAndSend({ Client });
+
+		SetDynamicFilterStatus(ENetFilterStatus::Disallow);
+
+		// As we have an hysteresis of one frame we will detect immediately when the trottling is updated
+		for (uint32 It = 0, EndIt = ConnectionThrottlingFrameCount; It < EndIt; ++It)
+		{
+			Server->UpdateAndSend({ Client });
+			if (Client->GetReplicationBridge()->GetReplicatedObject(ServerObject->NetRefHandle) == nullptr)
+			{
+				break;
+			}
+		}
+
+		// Object must have been destroyed on the client by now.
+		UE_NET_ASSERT_EQ(Client->GetReplicationBridge()->GetReplicatedObject(ServerObject->NetRefHandle), nullptr);
+	}
+
+	// Make sure object is kept alive for at least the expected frame count
+	{
+		SetDynamicFilterStatus(ENetFilterStatus::Allow);
+
+		UReplicatedTestObject* ServerObject = Server->CreateObject({ .IrisComponentCount = 0 });
+		Server->ReplicationSystem->SetFilter(ServerObject->NetRefHandle, MockFilterHandle, /* .FilterProfile = */ "FiveFrames");
+
+		// Advance up to one frame before we expect hysteresis update
+		for (uint32 It = 0, EndIt = ConnectionThrottlingFrameCount - 1; It < EndIt; ++It)
+		{
+			Server->UpdateAndSend({ Client });
+		}
+
+		// Filter out object. Hysteresis update should be performed but the object should not be filtered out immediately as not enough frames have passed.
+		SetDynamicFilterStatus(ENetFilterStatus::Disallow);
+
+		uint32 WaitFrameCount = 0;
+		for (uint32 It = 0, EndIt = 2U * ConnectionThrottlingFrameCount; It < EndIt; ++It)
+		{
+			++WaitFrameCount;
+			Server->UpdateAndSend({ Client });
+
+			if (Client->GetReplicationBridge()->GetReplicatedObject(ServerObject->NetRefHandle) == nullptr)
+			{
+				break;
+			}
+		}
+
+		// Object must have been destroyed on the client by now.
+		UE_NET_ASSERT_EQ(Client->GetReplicationBridge()->GetReplicatedObject(ServerObject->NetRefHandle), nullptr);
+
+		// FiveFames profile means at least five frames of waiting
+		UE_NET_ASSERT_GT(WaitFrameCount, 5U);
+
+		// This assert assumes ConnectionThrottlingFrameCount is five as well. If it's four we'd expect 4+4+1 frames.
+		UE_NET_ASSERT_EQ(WaitFrameCount, ConnectionThrottlingFrameCount + 1U);
+	}
+}
+
+// Test that a lot of objects can be filtered out on the same frame.
+UE_NET_TEST_FIXTURE(FTestFilteringFixture, LotsOfObjectsCanBeFilteredOutViaHysteresisInOneFrame)
+{
+	constexpr uint32 HighObjectCount = 65;
+
+	// Add client
+	FReplicationSystemTestClient* Client = CreateClient();
+
+	SetDynamicFilterStatus(ENetFilterStatus::Allow);
+
+	// Spawn lots of objects on server
+	UReplicatedTestObject* ServerObjects[HighObjectCount];
+	for (UReplicatedTestObject*& ServerObject : ServerObjects)
+	{
+		ServerObject = Server->CreateObject({ .IrisComponentCount = 0 });
+		Server->ReplicationSystem->SetFilter(ServerObject->NetRefHandle, MockFilterHandle, /* .FilterProfile = */ "OneFrame");
+	}
+
+	// Send and deliver packets until we believe all objects have been created on the client
+	for (uint32 It = 0; It < HighObjectCount; ++It)
+	{
+		Server->UpdateAndSend({ Client });
+		if (Client->GetReplicationBridge()->GetReplicatedObject(ServerObjects[HighObjectCount - 1]->NetRefHandle) != nullptr)
+		{
+			bool bAllObjectsCreated = true;
+
+			for (const UReplicatedTestObject* ServerObject : ServerObjects)
+			{
+				if (Client->GetReplicationBridge()->GetReplicatedObject(ServerObject->NetRefHandle) == nullptr)
+				{
+					bAllObjectsCreated = false;
+					break;
+				}
+			}
+			
+			if (bAllObjectsCreated)
+			{
+				break;
+			}
+		}
+	}
+
+	SetDynamicFilterStatus(ENetFilterStatus::Disallow);
+
+	// Send and deliver packet. Need to update twice to have the objects filtered out.
+	Server->UpdateAndSend({ Client });
+	Server->UpdateAndSend({ Client });
+
+	// All client objects should be destroyed
+	for (const UReplicatedTestObject* ServerObject : ServerObjects)
+	{
+		UE_NET_ASSERT_EQ(Client->GetReplicationBridge()->GetReplicatedObject(ServerObject->NetRefHandle), nullptr);
+	}
+}
+
 } // end namespace UE::Net::Private
diff --git a/Engine/Source/Runtime/Experimental/Iris/Core/Private/Iris/ReplicationSystem/Filtering/ObjectScopeHysteresisUpdater.cpp b/Engine/Source/Runtime/Experimental/Iris/Core/Private/Iris/ReplicationSystem/Filtering/ObjectScopeHysteresisUpdater.cpp
new file mode 100644
--- /dev/null
+++ b/Engine/Source/Runtime/Experimental/Iris/Core/Private/Iris/ReplicationSystem/Filtering/ObjectScopeHysteresisUpdater.cpp
@@ -0,0 +1,160 @@
+// Copyright Epic Games, Inc. All Rights Reserved.
+
+#include "Iris/ReplicationSystem/Filtering/ObjectScopeHysteresisUpdater.h"
+#include "Iris/Core/IrisProfiler.h"
+#include "Containers/Map.h"
+
+namespace UE::Net::Private
+{
+
+void FObjectScopeHysteresisUpdater::Init(uint32 MaxObjectCount)
+{
+	ObjectsToUpdate.Init(MaxObjectCount);
+}
+
+void FObjectScopeHysteresisUpdater::Deinit()
+{
+	FrameCounters.Empty();
+	LocalIndexToNetRefIndex.Empty();
+	NetRefIndexToLocalIndex.Empty();
+	UsedLocalIndices.Empty();
+	ObjectsToUpdate.Empty();
+}
+
+void FObjectScopeHysteresisUpdater::SetHysteresisFrameCount(FInternalNetRefIndex NetRefIndex, uint16 HysteresisFrameCount)
+{
+	const FLocalIndex LocalIndex = GetOrCreateLocalIndex(NetRefIndex);
+	FrameCounters[LocalIndex] = HysteresisFrameCount;
+}
+
+void FObjectScopeHysteresisUpdater::RemoveHysteresis(FInternalNetRefIndex NetRefIndex)
+{
+	if (FLocalIndex* LocalIndex = NetRefIndexToLocalIndex.Find(NetRefIndex))
+	{
+		FreeLocalIndex(*LocalIndex);
+	}
+}
+
+/** Removes all objects in the bitarray vfrom hysteresis. */
+void FObjectScopeHysteresisUpdater::RemoveHysteresis(const FNetBitArrayView& ObjectsToRemove)
+{
+	IRIS_PROFILER_SCOPE(FObjectScopeHysteresisUpdater_RemoveHysteresis);
+	FNetBitArrayView::ForAllSetBits(MakeNetBitArrayView(ObjectsToUpdate), ObjectsToRemove, FNetBitArrayView::AndOp, [this](uint32 NetRefIndex)
+		{
+			this->RemoveHysteresis(NetRefIndex);
+		}
+	);
+}
+
+void FObjectScopeHysteresisUpdater::Update(uint8 FramesSinceLastUpdate, TArray<FInternalNetRefIndex>& OutObjectsToFilterOut)
+{
+	IRIS_PROFILER_SCOPE(FObjectScopeHysteresisUpdater_Update);
+
+	ensure(FramesSinceLastUpdate > 0 && FramesSinceLastUpdate <= 128);
+
+	typedef FNetBitArrayView::StorageWordType WordType;
+	constexpr uint32 WordBitCount = FNetBitArrayView::WordBitCount;
+
+	const uint16 FilterOutCompareValue = (65536U - FramesSinceLastUpdate) & 65535U;
+	uint16* CountersData = FrameCounters.GetData();
+
+	TArray<FInternalNetRefIndex, TInlineAllocator<32>> ObjectsToRemoveFromUpdate;
+
+	const WordType* LocalIndicesData = UsedLocalIndices.GetData();
+	for (FLocalIndex ObjectIt = 0, ObjectEndIt = MaxLocalIndex, IndexOffset = 0; ObjectIt < ObjectEndIt; ObjectIt += WordBitCount, ++LocalIndicesData, IndexOffset += 32U)
+	{
+		// Skip ranges with no objects to update
+		WordType LocalIndicesWord = *LocalIndicesData;
+		if (!LocalIndicesWord)
+		{
+			continue;
+		}
+
+		WordType LocalIndicesMask = LocalIndicesWord;
+		for (WordType LocalIndexOffset = 0; LocalIndexOffset < WordBitCount; LocalIndexOffset += 4U, LocalIndicesMask >>= 4U)
+		{
+			if (!(LocalIndicesMask & 15U))
+			{
+				continue;
+			}
+
+			uint16 Counters[4];
+			Counters[0] = CountersData[IndexOffset + LocalIndexOffset + 0];
+			Counters[1] = CountersData[IndexOffset + LocalIndexOffset + 1];
+			Counters[2] = CountersData[IndexOffset + LocalIndexOffset + 2];
+			Counters[3] = CountersData[IndexOffset + LocalIndexOffset + 3];
+
+			Counters[0] -= FramesSinceLastUpdate;
+			Counters[1] -= FramesSinceLastUpdate;
+			Counters[2] -= FramesSinceLastUpdate;
+			Counters[3] -= FramesSinceLastUpdate;
+
+			// We can update the counters regardless of whether the objects were kept in scope or not. We're not expecting calls to Set/Remove while updating.
+			CountersData[IndexOffset + LocalIndexOffset + 0] = Counters[0];
+			CountersData[IndexOffset + LocalIndexOffset + 1] = Counters[1];
+			CountersData[IndexOffset + LocalIndexOffset + 2] = Counters[2];
+			CountersData[IndexOffset + LocalIndexOffset + 3] = Counters[3];
+
+			// For counter values < FramesSinceLastUpdate the object should remain in scope
+			for (uint32 Offset : {0, 1, 2, 3})
+			{
+				if (LocalIndicesMask & (1U << Offset))
+				{
+					if (Counters[Offset] >= FilterOutCompareValue)
+					{
+						// Object should now stay filtered out. Remove from updates.
+						ObjectsToRemoveFromUpdate.Add(IndexOffset + LocalIndexOffset + Offset);
+					}
+				}
+			}
+		}
+
+		if (!ObjectsToRemoveFromUpdate.IsEmpty())
+		{
+			OutObjectsToFilterOut.Reserve(OutObjectsToFilterOut.Num() + ObjectsToRemoveFromUpdate.Num());
+			for (FLocalIndex LocalIndex : ObjectsToRemoveFromUpdate)
+			{
+				OutObjectsToFilterOut.Add(LocalIndexToNetRefIndex[LocalIndex]);
+				FreeLocalIndex(LocalIndex);
+			}
+
+			ObjectsToRemoveFromUpdate.Reset();
+		}
+	}
+}
+
+FObjectScopeHysteresisUpdater::FLocalIndex FObjectScopeHysteresisUpdater::GetOrCreateLocalIndex(FInternalNetRefIndex NetRefIndex)
+{
+	if (FLocalIndex* LocalIndex = NetRefIndexToLocalIndex.Find(NetRefIndex))
+	{
+		return *LocalIndex;
+	}
+
+	uint32 LocalIndex = UsedLocalIndices.FindFirstZero();
+	if (LocalIndex == FNetBitArray::InvalidIndex)
+	{
+		UsedLocalIndices.AddBits(LocalIndexGrowCount);
+		LocalIndexToNetRefIndex.AddZeroed(LocalIndexGrowCount);
+		FrameCounters.AddZeroed(LocalIndexGrowCount);
+		LocalIndex = UsedLocalIndices.FindFirstZero();
+	}
+
+	UsedLocalIndices.SetBit(LocalIndex);
+	LocalIndexToNetRefIndex[LocalIndex] = NetRefIndex;
+	NetRefIndexToLocalIndex.Add(NetRefIndex, LocalIndex);
+	ObjectsToUpdate.SetBit(NetRefIndex);
+	MaxLocalIndex = FPlatformMath::Max(MaxLocalIndex, LocalIndex + 1U);
+
+	return LocalIndex;
+}
+
+void FObjectScopeHysteresisUpdater::FreeLocalIndex(FLocalIndex LocalIndex)
+{
+	UsedLocalIndices.ClearBit(LocalIndex);
+	const FInternalNetRefIndex NetRefIndex = LocalIndexToNetRefIndex[LocalIndex];
+	NetRefIndexToLocalIndex.Remove(NetRefIndex);
+	ObjectsToUpdate.ClearBit(NetRefIndex);
+	// Intentionally not updating LocalIndexToNetRefIndex since it isn't accessed for unset local indices.
+}
+
+}
diff --git a/Engine/Source/Runtime/Experimental/Iris/Core/Private/Iris/ReplicationSystem/Filtering/ObjectScopeHysteresisUpdater.h b/Engine/Source/Runtime/Experimental/Iris/Core/Private/Iris/ReplicationSystem/Filtering/ObjectScopeHysteresisUpdater.h
new file mode 100644
--- /dev/null
+++ b/Engine/Source/Runtime/Experimental/Iris/Core/Private/Iris/ReplicationSystem/Filtering/ObjectScopeHysteresisUpdater.h
@@ -0,0 +1,79 @@
+// Copyright Epic Games, Inc. All Rights Reserved.
+
+#pragma once
+
+#include "CoreTypes.h"
+#include "Containers/Array.h"
+#include "Containers/Map.h"
+#include "Net/Core/NetBitArray.h"
+
+namespace UE::Net::Private
+{
+	typedef uint32 FInternalNetRefIndex;
+}
+
+namespace UE::Net::Private
+{
+
+class FObjectScopeHysteresisUpdater
+{
+public:
+	void Init(uint32 MaxObjectCount);
+	void Deinit();
+
+	/** Sets an hysteresis frame count such that an object will be kept in scope until such many frames has passed. */
+	void SetHysteresisFrameCount(FInternalNetRefIndex NetRefIndex, uint16 HysteresisFrameCount);
+
+	/** Remove an object from hysteresis update. Needed when an object goes out of scope. An object is automatically removed from hysteresis updates when a previously set frame count has expired. */
+	void RemoveHysteresis(FInternalNetRefIndex NetRefIndex);
+
+	/** Removes all objects in the bitarray from hysteresis. */
+	void RemoveHysteresis(const FNetBitArrayView& ObjectsToRemove);
+
+	/** Removes all objects in the array from hysteresis. */
+	void RemoveHysteresis(const TArrayView<const uint32>& ObjectsToRemove);
+
+	/** Adjusts the connection relevant objects based on hysteresis frame counts. */
+	void Update(uint8 FramesSinceLastUpdate, TArray<FInternalNetRefIndex>& OutObjectsToFilterOut);
+
+	bool HasObjectsToUpdate() const;
+
+	/** Returns the bitarray of objects affected by hysteresis */
+	FNetBitArrayView GetUpdatedObjects() const;
+
+private:
+	enum : unsigned
+	{
+		LocalIndexGrowCount = 256U,
+	};
+
+	typedef uint32 FLocalIndex;
+
+	FLocalIndex GetOrCreateLocalIndex(FInternalNetRefIndex NetRefIndex);
+	void FreeLocalIndex(FLocalIndex LocalIndex);
+
+	// Internally there's a mapping between NetRefIndex and a LocalIndex to keep updates in as small a range as possible. This is the highest local index + 1 that needs to be updated.
+	FLocalIndex MaxLocalIndex = 0;
+	// Per LocalIndex how many frames left to update before filtering out the object.
+	TArray<uint16> FrameCounters;
+	// Lookup table for local index to InternalNetRefIndex.
+	TArray<FInternalNetRefIndex> LocalIndexToNetRefIndex;
+	// Lookup map for InternalNetRefIndex to LocalIndex to be able to figure out whether an object already is assigned a LocalIndex or not.
+	TMap<FInternalNetRefIndex, FLocalIndex> NetRefIndexToLocalIndex;
+	// A set bit indicates that the LocalIndex is used. Only stores MaxLocalIndex bits. The bitarray grows as needed.
+	FNetBitArray UsedLocalIndices;
+	// A set bit indicates that the InternalNetRefIndex is being updated. Can always hold the MaxObjectCount passed to Init().
+	FNetBitArray ObjectsToUpdate;
+};
+
+inline bool FObjectScopeHysteresisUpdater::HasObjectsToUpdate() const
+{
+	return !NetRefIndexToLocalIndex.IsEmpty();
+}
+
+inline FNetBitArrayView FObjectScopeHysteresisUpdater::GetUpdatedObjects() const
+{
+	return MakeNetBitArrayView(ObjectsToUpdate);
+}
+
+}
diff --git a/Engine/Source/Runtime/Experimental/Iris/Core/Private/Iris/ReplicationSystem/Filtering/ReplicationFiltering.cpp b/Engine/Source/Runtime/Experimental/Iris/Core/Private/Iris/ReplicationSystem/Filtering/ReplicationFiltering.cpp
--- a/Engine/Source/Runtime/Experimental/Iris/Core/Private/Iris/ReplicationSystem/Filtering/ReplicationFiltering.cpp
+++ b/Engine/Source/Runtime/Experimental/Iris/Core/Private/Iris/ReplicationSystem/Filtering/ReplicationFiltering.cpp
@@ -6,6 +6,7 @@
 
 #include "Iris/IrisConfigInternal.h"
 #include "Iris/IrisConstants.h"
+#include "Iris/Core/BitTwiddling.h"
 #include "Iris/Core/IrisCsv.h"
 #include "Iris/Core/IrisLog.h"
 #include "Iris/Core/IrisProfiler.h"
@@ -17,7 +18,7 @@
 #include "Iris/ReplicationSystem/DeltaCompression/DeltaCompressionBaselineInvalidationTracker.h"
 #include "Iris/ReplicationSystem/Filtering/NetObjectGroups.h"
 #include "Iris/ReplicationSystem/Filtering/NetObjectFilterDefinitions.h"
-
+#include "Iris/ReplicationSystem/Filtering/ReplicationFilteringConfig.h"
 #include <limits>
 
 namespace UE::Net::Private
@@ -186,6 +187,8 @@ void FReplicationFiltering::Init(FReplicationFilteringInitParams& Params)
 	check(Params.Connections != nullptr);
 	check(Params.Connections->GetMaxConnectionCount() <= std::numeric_limits<decltype(ObjectIndexToOwningConnection)::ElementType>::max());
 
+	Config = GetDefault<UReplicationFilteringConfig>();
+
 	ReplicationSystem = Params.ReplicationSystem;
 
 	Connections = Params.Connections;
@@ -233,6 +236,9 @@ void FReplicationFiltering::Init(FReplicationFilteringInitParams& Params)
 	// Owning connections
 	ObjectIndexToOwningConnection.SetNumZeroed(MaxObjectCount);
 
+	// Object scope hysteresis
+	ObjectScopeHysteresisFrameCounts.SetNumZeroed(MaxObjectCount);
+
 	// Dynamic filters
 	{
 		NetObjectFilteringInfos.SetNumUninitialized(Params.MaxObjectCount);
@@ -244,6 +250,7 @@ void FReplicationFiltering::Init(FReplicationFilteringInitParams& Params)
 	}
 
 	InitFilters();
+	InitObjectScopeHysteresis();
 }
 
 void FReplicationFiltering::Filter()
@@ -251,6 +258,9 @@ void FReplicationFiltering::Filter()
 #if UE_NET_IRIS_CSV_STATS
 	CSV_SCOPED_TIMING_STAT(Iris, Filter_PrePoll);
 #endif
+
+	++FrameIndex;
+
 	ResetRemovedConnections();
 
 	InitNewConnections();
@@ -265,6 +275,8 @@ void FReplicationFiltering::Filter()
 
 	UpdateSubObjectFilters();
 
+	PreUpdateObjectScopeHysteresis();
+
 	if (HasDynamicFilters())
 	{
 		UpdateDynamicFilters();
@@ -425,6 +437,7 @@ bool FReplicationFiltering::SetFilter(FInternalNetRefIndex ObjectIndex, FNetObje
 			FilterInfo.FilteredObjects.SetBit(ObjIndex);
 			this->ObjectIndexToDynamicFilterIndex[ObjIndex] = static_cast<uint8>(FilterIndex);
 			this->DynamicFilterEnabledObjects.SetBit(ObjIndex);
+			this->ObjectScopeHysteresisFrameCounts[ObjIndex] = this->GetObjectScopeHysteresisFrameCount(FilterConfigProfile);
 			return true;
 		}
 		
@@ -604,11 +617,7 @@ void FReplicationFiltering::RemoveConnection(uint32 ConnectionId)
 	
 	// Reset connection info
 	FPerConnectionInfo& ConnectionInfo = ConnectionInfos[ConnectionId];
-	ConnectionInfo.ConnectionFilteredObjects.Empty();
-	ConnectionInfo.GroupExcludedObjects.Empty();
-	ConnectionInfo.ObjectsInScopeBeforeDynamicFiltering.Empty();
-	ConnectionInfo.GroupIncludedObjects.Empty();
-	ConnectionInfo.ObjectsInScope.Empty();
+	ConnectionInfo.Deinit();
 
 	for (FFilterInfo& Info : DynamicFilterInfos)
 	{
@@ -663,6 +672,7 @@ void FReplicationFiltering::InitNewConnections()
 		{
 			ConnectionInfo.DynamicFilteredOutObjects.Init(MaxObjectCount);
 			ConnectionInfo.InProgressDynamicFilteredOutObjects.Init(MaxObjectCount);
+			ConnectionInfo.HysteresisUpdater.Init(MaxObjectCount);
 		}
 
 		// Update group exclusion filtering
@@ -867,7 +877,7 @@ void FReplicationFiltering::UpdateObjectsInScope()
 			uint32 DeletedObjects = (PrevExistingObjects & ~ExistingObjects);
 			for ( ; DeletedObjects; )
 			{
-				const uint32 LeastSignificantBit = DeletedObjects & uint32(-int32(DeletedObjects));
+				const uint32 LeastSignificantBit = GetLeastSignificantBit(DeletedObjects);
 				DeletedObjects ^= LeastSignificantBit;
 
 				const uint32 ObjectIndex = BitOffset + FPlatformMath::CountTrailingZeros(LeastSignificantBit);
@@ -891,7 +901,7 @@ void FReplicationFiltering::UpdateObjectsInScope()
 			{
 				for ( ; AddedSubObjects; )
 				{
-					const uint32 LeastSignificantBit = AddedSubObjects & uint32(-int32(AddedSubObjects));
+					const uint32 LeastSignificantBit = GetLeastSignificantBit(AddedSubObjects);
 					AddedSubObjects ^= LeastSignificantBit;
 
 					const uint32 ObjectIndex = BitOffset + FPlatformMath::CountTrailingZeros(LeastSignificantBit);
@@ -1244,8 +1254,13 @@ void FReplicationFiltering::UpdateDynamicFiltering()
 {
 	IRIS_PROFILER_SCOPE(FReplicationFiltering_UpdateDynamicFiltering);
 
-	// Working set
-	TArray<uint32, TInlineAllocator<256>> DisabledDependentObjects;
+	// Working set. Use a reasonable amount of stack.
+	constexpr uint32 TotalObjectCountOnStack = 1024;
+	constexpr uint32 FilteredOutObjectCountOnStack = 256;
+	TArray<uint32, TInlineAllocator<TotalObjectCountOnStack - FilteredOutObjectCountOnStack>> FilteredOutObjects;
+	TArray<uint32, TInlineAllocator<FilteredOutObjectCountOnStack>> FilteredOutDependentObjects;
+	TArray<uint32> FilteredOutByHysteresisObjects;
+	FilteredOutByHysteresisObjects.Reserve(256);
 
 	uint32* AllowedObjectsData = static_cast<uint32*>(FMemory_Alloca(WordCountForObjectBitArrays * sizeof(uint32)));
 	FNetBitArrayView AllowedObjects(AllowedObjectsData, MaxObjectCount, FNetBitArrayView::NoResetNoValidate);
@@ -1266,7 +1281,6 @@ void FReplicationFiltering::UpdateDynamicFiltering()
 
 		ConnectionInfo.InProgressDynamicFilteredOutObjects.Reset();
 		uint32* InProgressDynamicFilteredOutObjectsData = ConnectionInfo.InProgressDynamicFilteredOutObjects.GetData();
-		FNetBitArrayView InProgressDynamicFilteredOutObjects = MakeNetBitArrayView(ConnectionInfo.InProgressDynamicFilteredOutObjects);
 
 		/*
 		 * Apply dynamic filters.
@@ -1303,6 +1317,16 @@ void FReplicationFiltering::UpdateDynamicFiltering()
 					InProgressDynamicFilteredOutObjectsData[WordIt] = (InProgressDynamicFilteredOutObjectsData[WordIt] & ~FilteredObjects) | (~FilterAllowedObjects & FilteredObjects);
 				}
 			}
+
+			// Objects in hysteresis should not be filtered out. 
+			if (HysteresisState.Mode == EHysteresisProcessingMode::Enabled)
+			{
+				if (ConnectionInfo.HysteresisUpdater.HasObjectsToUpdate())
+				{
+					FNetBitArrayView InProgressDynamicFilteredOutObjects = MakeNetBitArrayView(ConnectionInfo.InProgressDynamicFilteredOutObjects);
+					InProgressDynamicFilteredOutObjects.Combine(ConnectionInfo.HysteresisUpdater.GetUpdatedObjects(), FNetBitArrayView::AndNotOp);
+				}
+			}
 		}
 
 		/*
@@ -1310,12 +1334,13 @@ void FReplicationFiltering::UpdateDynamicFiltering()
 		 * and try only to do the more expensive operations, such as subobject management,
 		 * for objects that have changed filter status since the previous frame.
 		 */
+		FNetBitArrayView DynamicFilteredOutObjects = MakeNetBitArrayView(ConnectionInfo.DynamicFilteredOutObjects);
 		{
 			IRIS_PROFILER_SCOPE(FReplicationFiltering_OnFilterStatusChanged);
-			DisabledDependentObjects.Reset();
+			FilteredOutObjects.Reset();
+			FilteredOutDependentObjects.Reset();
 
 			const uint32* DynamicFilterEnabledObjectsData = DynamicFilterEnabledObjects.GetData();
-			FNetBitArrayView DynamicFilteredOutObjects = MakeNetBitArrayView(ConnectionInfo.DynamicFilteredOutObjects);
 			uint32* DynamicFilteredOutObjectsData = ConnectionInfo.DynamicFilteredOutObjects.GetData();
 			for (uint32 WordIt = 0, WordEndIt = WordCountForObjectBitArrays; WordIt != WordEndIt; ++WordIt)
 			{
@@ -1336,8 +1361,8 @@ void FReplicationFiltering::UpdateDynamicFiltering()
 					continue;
 				}
 
-				// Update dynamic filter, but preserve subobject information.
-				DynamicFilteredOutObjectsData[WordIt] = (PrevFilteredOutObjects & SubObjects) | NewFilteredOutObjects;
+				// Update dynamic filter with no longer filtered out objects, but preserve subobject information. Objects that were just filtered out will be dealt with after hysteresis.
+				DynamicFilteredOutObjectsData[WordIt] = (PrevFilteredOutObjects & SubObjects) | (NewFilteredOutObjects & PrevFilteredOutObjects);
 
 				const uint32 BitOffset = WordIt*32U;
 
@@ -1346,25 +1371,25 @@ void FReplicationFiltering::UpdateDynamicFiltering()
 				for (uint32 DisabledObjects = NewFilteredOutObjects & ObjectsToProcess; DisabledObjects; )
 				{
 					// Extract one set bit from DisabledObjects. The code below will extract the least significant bit set and then clear it.
-					const uint32 LeastSignificantBit = DisabledObjects & uint32(-int32(DisabledObjects));
+					const uint32 LeastSignificantBit = GetLeastSignificantBit(DisabledObjects);
 					DisabledObjects ^= LeastSignificantBit;
 
+					// Store disabled objects for later processing to minimize performance impact of hysteresis
 					const uint32 ObjectIndex = BitOffset + FPlatformMath::CountTrailingZeros(LeastSignificantBit);
-					for (const FInternalNetRefIndex SubObjectIndex : NetRefHandleManager->GetSubObjects(ObjectIndex))
-					{						
-						DynamicFilteredOutObjects.SetBit(SubObjectIndex);
-					}
+
+					// Store objects for later processing after hysteresis 
+					FilteredOutObjects.Add(ObjectIndex);
 
 					// Store dependent objects for later processing.
 					if (DependentObjects & LeastSignificantBit)
 					{
-						DisabledDependentObjects.Add(ObjectIndex);
+						FilteredOutDependentObjects.Add(ObjectIndex);
 					}
 				}
 
 				for (uint32 EnabledObjects = ~NewFilteredOutObjects & ObjectsToProcess; EnabledObjects; )
 				{
-					const uint32 LeastSignificantBit = EnabledObjects & uint32(-int32(EnabledObjects));
+					const uint32 LeastSignificantBit = GetLeastSignificantBit(EnabledObjects);
 					EnabledObjects ^= LeastSignificantBit;
 
 					const uint32 ObjectIndex = BitOffset + FPlatformMath::CountTrailingZeros(LeastSignificantBit);
@@ -1376,6 +1401,67 @@ void FReplicationFiltering::UpdateDynamicFiltering()
 			}
 		}
 
+		// Update hysteresis
+		if (HysteresisState.Mode == EHysteresisProcessingMode::Enabled)
+		{
+			IRIS_CSV_PROFILER_SCOPE(Iris, FReplicationFiltering_ConnectionHysteresis);
+
+			// Time to update?
+			const uint32 ConnIdMod = ConnId % HysteresisState.ConnectionIdStride;
+			if (ConnIdMod == HysteresisState.ConnectionStartId)
+			{
+				FilteredOutByHysteresisObjects.Reset();
+				ConnectionInfo.HysteresisUpdater.Update(static_cast<uint8>(HysteresisState.ConnectionIdStride), FilteredOutByHysteresisObjects);
+
+				// Unconditionally filter out objects that are no longer kept in scope due to hysteresis.
+				for (const uint32 ObjectIndex : FilteredOutByHysteresisObjects)
+				{
+					DynamicFilteredOutObjects.SetBit(ObjectIndex);
+					for (const FInternalNetRefIndex SubObjectIndex : NetRefHandleManager->GetSubObjects(ObjectIndex))
+					{
+						DynamicFilteredOutObjects.SetBit(SubObjectIndex);
+					}
+				}
+			}
+
+			// Add filtered out objects to hysteresis if supported and carry out the filtering if not.
+			{
+				// Add the amount of frames that have passed since last update since the update is going to subtract the connection update throttling value.
+				const uint32 AdjustHysteresisForUpdateThrottling = (HysteresisState.ConnectionStartId + HysteresisState.ConnectionIdStride - ConnIdMod) % HysteresisState.ConnectionIdStride;
+				for (const uint32 ObjectIndex : FilteredOutObjects)
+				{
+					if (const uint32 HysteresisFrameCount = ObjectScopeHysteresisFrameCounts[ObjectIndex])
+					{
+						// We need to adjust the hysteresis frame count to account for when it will be updated. The -1 stems from the fact that updating won't happen until next frame at the earliest.
+						const uint16 TotalHysteresisFrameCount = static_cast<uint16>(HysteresisFrameCount - 1 + AdjustHysteresisForUpdateThrottling);
+						ConnectionInfo.HysteresisUpdater.SetHysteresisFrameCount(ObjectIndex, TotalHysteresisFrameCount);
+					}
+					else
+					{
+						// Filter out
+						DynamicFilteredOutObjects.SetBit(ObjectIndex);
+						for (const FInternalNetRefIndex SubObjectIndex : NetRefHandleManager->GetSubObjects(ObjectIndex))
+						{
+							DynamicFilteredOutObjects.SetBit(SubObjectIndex);
+						}
+					}
+				}
+			}
+		}
+		else
+		{
+			IRIS_PROFILER_SCOPE(FReplicationFiltering_ConnectionFilterOut);
+			// Filter out objects
+			for (const uint32 ObjectIndex : FilteredOutObjects)
+			{
+				DynamicFilteredOutObjects.SetBit(ObjectIndex);
+				for (const FInternalNetRefIndex SubObjectIndex : NetRefHandleManager->GetSubObjects(ObjectIndex))
+				{
+					DynamicFilteredOutObjects.SetBit(SubObjectIndex);
+				}
+			}
+		}
+
 		// Update the entire scope for the connection
 		{
 			IRIS_PROFILER_SCOPE(FReplicationFiltering_UpdateConnectionScope);
@@ -1398,7 +1484,7 @@ void FReplicationFiltering::UpdateDynamicFiltering()
 		// If any object that has a dependency that isn't filtered out we must re-enable the dependent object.
 		{
 			IRIS_PROFILER_SCOPE(FReplicationFiltering_UpdateDependentObjects);
-			for (const uint32 DependentObjectIndex : DisabledDependentObjects)
+			for (const uint32 DependentObjectIndex : FilteredOutDependentObjects)
 			{
 				if (GetDependentObjectFilterStatus(NetRefHandleManager, ConnectionInfo.ObjectsInScope, DependentObjectIndex) == ENetFilterStatus::Allow)
 				{
@@ -1435,6 +1521,8 @@ void FReplicationFiltering::PostUpdateDynamicFiltering()
 			Info.Filter->PostFilter(PostFilteringParams);
 		}
 	}
+
+	PostUpdateObjectScopeHysteresis();
 }
 
 void FReplicationFiltering::NotifyFiltersOfDirtyObjects()
@@ -2058,6 +2146,18 @@ bool FReplicationFiltering::ClearGroupInclusionFilterEffectsForObject(uint32 Obj
 		
 		GroupIncludedObjects.ClearBit(ObjectIndex);
 
+		// Dynamically filtered objects are subject to hysteresis.
+		if (HysteresisState.Mode == EHysteresisProcessingMode::Enabled)
+		{
+			const uint8 HysteresisFrameCount = DynamicFilterEnabledObjects.GetBit(ObjectIndex) ? ObjectScopeHysteresisFrameCounts[ObjectIndex] : uint8(0);
+			if (HysteresisFrameCount)
+			{
+				const uint32 ConnIdMod = ConnectionId % HysteresisState.ConnectionIdStride;
+				const uint32 AdjustHysteresisForUpdateThrottling = (HysteresisState.ConnectionStartId + HysteresisState.ConnectionIdStride - ConnIdMod) % HysteresisState.ConnectionIdStride;
+				ConnectionInfo.HysteresisUpdater.SetHysteresisFrameCount(ObjectIndex, static_cast<uint16>(HysteresisFrameCount + AdjustHysteresisForUpdateThrottling));
+			}
+		}
+
 		for (const FInternalNetRefIndex SubObjectIndex : NetRefHandleManager->GetSubObjects(ObjectIndex))
 		{
 			GroupIncludedObjects.ClearBit(SubObjectIndex);
@@ -2407,6 +2507,13 @@ void FReplicationFiltering::InitFilters()
 	}
 }
 
+void FReplicationFiltering::InitObjectScopeHysteresis()
+{
+	HysteresisState.Mode = Config->IsObjectScopeHysteresisEnabled() ? EHysteresisProcessingMode::Enabled : EHysteresisProcessingMode::Disabled;
+	// Always allocated and maintained regardless of whether the feature is enabled or not.
+	HysteresisState.ObjectsToClear.SetNumBits(MaxObjectCount);
+}
+
 void FReplicationFiltering::RemoveFromDynamicFilter(uint32 ObjectIndex, uint32 FilterIndex)
 {
 	UE_LOG(LogIrisFiltering, Verbose, TEXT("RemoveFromDynamicFilter removing %s from Dynamic Filter %s"), 
@@ -2421,6 +2528,9 @@ void FReplicationFiltering::RemoveFromDynamicFilter(uint32 ObjectIndex, uint32 F
 
 	DynamicFilterEnabledObjects.ClearBit(ObjectIndex);
 	ObjectsRequiringDynamicFilterUpdate.SetBit(ObjectIndex);
+
+	// Remove from hysteresis
+	HysteresisState.ClearFromHysteresis(ObjectIndex);
 }
 
 void FReplicationFiltering::InvalidateBaselinesForObject(uint32 ObjectIndex, uint32 NewOwningConnectionId, uint32 PrevOwningConnectionId)
@@ -2504,6 +2614,72 @@ void FReplicationFiltering::BuildObjectsInFilterList(FNetBitArrayView OutObjects
 	}
 }
 
+
+void FReplicationFiltering::PreUpdateObjectScopeHysteresis()
+{
+	if (HysteresisState.Mode == EHysteresisProcessingMode::Enabled)
+	{
+		HysteresisState.ConnectionStartId = (FrameIndex % Config->GetHysteresisUpdateConnectionThrottling());
+		HysteresisState.ConnectionIdStride = Config->GetHysteresisUpdateConnectionThrottling();
+	}
+
+	ClearObjectsFromHysteresis();
+}
+
+void FReplicationFiltering::PostUpdateObjectScopeHysteresis()
+{
+	HysteresisState.ObjectsToClearCount = 0;
+	HysteresisState.ObjectsToClear.Reset();
+}
+
+void FReplicationFiltering::ClearObjectsFromHysteresis()
+{
+	if (!HysteresisState.ObjectsToClearCount)
+	{
+		return;
+	}
+
+	IRIS_PROFILER_SCOPE(FObjectScopeHysteresisUpdater_ClearObjectsFromHysteresis);
+
+	// $IRIS TODO: Use more optimal path for low counts of objects to clear, for example passing an ArrayView instead.
+	ValidConnections.ForAllSetBits([this](uint32 ConnectionId)
+		{
+			FPerConnectionInfo& ConnectionInfo = this->ConnectionInfos[ConnectionId];
+			ConnectionInfo.HysteresisUpdater.RemoveHysteresis(MakeNetBitArrayView(this->HysteresisState.ObjectsToClear));
+		});
+}
+
+uint8 FReplicationFiltering::GetObjectScopeHysteresisFrameCount(FName ProfileName) const
+{
+	if (const FObjectScopeHysteresisProfile* Profile = Config->GetHysteresisProfiles().FindByKey(ProfileName))
+	{
+		return Profile->HysteresisFrameCount;
+	}
+
+	return Config->GetDefaultHysteresisFrameCount();
+}
+
+void FReplicationFiltering::FPerConnectionInfo::Deinit()
+{
+	ConnectionFilteredObjects.Empty();
+	GroupExcludedObjects.Empty();
+	ObjectsInScopeBeforeDynamicFiltering.Empty();
+	GroupIncludedObjects.Empty();
+	ObjectsInScope.Empty();
+	DynamicFilteredOutObjects.Empty();
+	InProgressDynamicFilteredOutObjects.Empty();
+	HysteresisUpdater.Deinit();
+}
+
+//*************************************************************************************************
+// FObjectScopeHysteresisState
+//*************************************************************************************************
+void FReplicationFiltering::FObjectScopeHysteresisState::ClearFromHysteresis(FInternalNetRefIndex NetRefIndex)
+{
+	ObjectsToClear.SetBit(NetRefIndex);
+	++ObjectsToClearCount;
+}
+
 //*************************************************************************************************
 // FNetObjectFilteringInfoAccessor
 //*************************************************************************************************
diff --git a/Engine/Source/Runtime/Experimental/Iris/Core/Private/Iris/ReplicationSystem/Filtering/ReplicationFiltering.h b/Engine/Source/Runtime/Experimental/Iris/Core/Private/Iris/ReplicationSystem/Filtering/ReplicationFiltering.h
--- a/Engine/Source/Runtime/Experimental/Iris/Core/Private/Iris/ReplicationSystem/Filtering/ReplicationFiltering.h
+++ b/Engine/Source/Runtime/Experimental/Iris/Core/Private/Iris/ReplicationSystem/Filtering/ReplicationFiltering.h
@@ -5,10 +5,13 @@
 #include "CoreTypes.h"
 #include "Net/Core/NetBitArray.h"
 #include "Iris/ReplicationSystem/Filtering/NetObjectFilter.h"
+#include "Iris/ReplicationSystem/Filtering/ObjectScopeHysteresisUpdater.h"
+#include "Iris/ReplicationSystem/Filtering/ReplicationFilteringConfig.h"
 #include "Iris/ReplicationSystem/NetObjectGroupHandle.h"
 #include "Containers/Array.h"
 #include "UObject/StrongObjectPtr.h"
 
+class UReplicationFilteringConfig;
 class UReplicationSystem;
 namespace UE::Net
 {
@@ -146,6 +149,8 @@ public:
 private:
 	struct FPerConnectionInfo
 	{
+		void Deinit();
+
 		// Objects filtered depending on owning connection or user set connection filtering
 		FNetBitArray ConnectionFilteredObjects;
 		// Objects filtered out due to one or more exclusion groups it belongs to is filtered out
@@ -162,6 +167,9 @@ private:
 
 		// List of objects currently filtered out after processing dynamic filter passes and inclusion groups
 		FNetBitArray InProgressDynamicFilteredOutObjects;
+
+		// Updater of hysteresis for objects being dynamically filtered out
+		FObjectScopeHysteresisUpdater HysteresisUpdater;
 	};
 
 	struct FPerObjectInfo
@@ -194,6 +202,7 @@ private:
 	static void StaticChecks();
 
 	void InitFilters();
+	void InitObjectScopeHysteresis();
 
 	void InitNewConnections();
 	void ResetRemovedConnections();
@@ -208,6 +217,10 @@ private:
 	void UpdateDynamicFiltering();
 	void PostUpdateDynamicFiltering();
 
+	void PreUpdateObjectScopeHysteresis();
+	void PostUpdateObjectScopeHysteresis();
+	void ClearObjectsFromHysteresis();
+
 	/** Build the list of always relevant objects + objects that are currently relevant to at least one connection. */
 	void FilterNonRelevantObjects();
 
@@ -255,13 +268,45 @@ private:
 	/** Returns all the filtering infos. */
 	TArrayView<FNetObjectFilteringInfo> GetNetObjectFilteringInfos();
 
+	uint8 GetObjectScopeHysteresisFrameCount(FName Profile) const;
+
 private:
+	enum EHysteresisProcessingMode : uint32
+	{
+		Disabled,
+		Enabled,
+	};
+
+	// Scope hysteresis state. Hysteresis is applied to objects going out of scope for objects that so desire.
+	struct FObjectScopeHysteresisState
+	{
+	public:
+		void ClearFromHysteresis(FInternalNetRefIndex NetRefIndex);
+
+		// Processing mode
+		EHysteresisProcessingMode Mode = EHysteresisProcessingMode::Disabled;
+		// Which connection ID to start with for updating.
+		uint32 ConnectionStartId = 0;
+		// Stride for connection update throttling.
+		uint32 ConnectionIdStride = 1;
+
+		// Approximate number of objects that should be cleared from hysteresis.
+		uint32 ObjectsToClearCount = 0;
+
+		// Objects to clear from hysteresis due to being destroyed or removed from dynamic filtering.
+		FNetBitArray ObjectsToClear;
+	};
+
 	// Used for ObjectIndexToDynamicFilterIndex lookup
 	static constexpr uint8 InvalidDynamicFilterIndex = 255U;
 
+	// Config
+	TObjectPtr<const UReplicationFilteringConfig> Config = nullptr;
+
 	// General
 	TObjectPtr<UReplicationSystem> ReplicationSystem = nullptr;
 	const FNetRefHandleManager* NetRefHandleManager = nullptr;
+	uint32 FrameIndex = 0;
 
 	// Groups
 	FNetObjectGroups* Groups = nullptr;
@@ -298,6 +343,9 @@ private:
 	TArray<FPerGroupInfo> GroupInfos;
 	uint32 MaxGroupCount = 0;
 
+	// Hysteresis frame counts for dynamically filtered objects
+	TArray<uint8> ObjectScopeHysteresisFrameCounts;
+
 	/** NetObjectGroups used for filtering out objects. */
 	FNetBitArray ExclusionFilterGroups;
 
@@ -328,6 +376,9 @@ private:
 	FNetBitArray DynamicFilterEnabledObjects;
 	FNetBitArray ObjectsRequiringDynamicFilterUpdate;
 
+	// Object scope hystereris
+	FObjectScopeHysteresisState HysteresisState;
+
 	uint32 bHasNewConnection : 1;
 	uint32 bHasRemovedConnection : 1;
 	uint32 bHasDirtyConnectionFilter: 1;
diff --git a/Engine/Source/Runtime/Experimental/Iris/Core/Public/Iris/ReplicationSystem/Filtering/ReplicationFilteringConfig.h b/Engine/Source/Runtime/Experimental/Iris/Core/Public/Iris/ReplicationSystem/Filtering/ReplicationFilteringConfig.h
new file mode 100644
--- /dev/null
+++ b/Engine/Source/Runtime/Experimental/Iris/Core/Public/Iris/ReplicationSystem/Filtering/ReplicationFilteringConfig.h
@@ -0,0 +1,97 @@
+// Copyright Epic Games, Inc. All Rights Reserved.
+
+#pragma once
+
+#include "UObject/ObjectMacros.h"
+#include "Math/UnrealMathUtility.h"
+#include "ReplicationFilteringConfig.generated.h"
+
+USTRUCT()
+struct FObjectScopeHysteresisProfile
+{
+	GENERATED_BODY()
+
+public:
+	bool operator==(FName Name) const
+	{
+		return FilterProfileName == Name;
+	}
+
+	/** The config name used to map to this profile */
+	UPROPERTY()
+	FName FilterProfileName;
+
+	/** The number of frames to keep the object in scope after it has been filtered out by dynamic filtering. */
+	UPROPERTY()
+	uint8 HysteresisFrameCount = 0;
+};
+
+/**
+ * Object scope hysteresis support. Keep dynamically filtered out objects around for a specified amount of frames. 
+ * Configure behavior via hysteresis profiles that determine the frame timeout per class.
+ * The filter config for a specific class can then mention the hysteresis profile in order to get the appropriate behavior. 
+ *
+ * Example:
+ * [/Script/IrisCore.ReplicationFilteringConfig]
+ * bEnableObjectScopeHysteresis=true
+ * DefaultHysteresisFrameCount=4
+ * HysteresisUpdateConnectionThrottling=4
+ * !HysteresisProfiles=ClearArray
+ * +FilterProfiles=(FilterProfileName=PawnFilterProfile, HysteresisFrameCount=30)
+ * 
+ * [/Script/ IrisCore.ObjectReplicationBridgeConfig]
+ * +FilterConfigs=(ClassName=/Script/Engine.Pawn, DynamicFilterName=Spatial, FilterProfile=PawnFilterProfile)
+ */
+
+UCLASS(transient, config = Engine)
+class UReplicationFilteringConfig final : public UObject
+{
+	GENERATED_BODY()
+
+public:
+	bool IsObjectScopeHysteresisEnabled() const;
+	uint8 GetDefaultHysteresisFrameCount() const;
+	uint8 GetHysteresisUpdateConnectionThrottling() const;
+
+	const TArray<FObjectScopeHysteresisProfile>& GetHysteresisProfiles() const;
+
+private:
+	UPROPERTY(Config)
+	/** If enabled a dynamically filtered out object will not be considered out of scope for a particular number of frames. */
+	bool bEnableObjectScopeHysteresis = true;
+
+	UPROPERTY(Config)
+	/** How many frames a dynamically filtered out object should still be considered in scope by default. Can be overridden with HysteresisClassConfigs. */ 
+	uint8 DefaultHysteresisFrameCount = 0;
+
+	UPROPERTY(Config)
+	/**
+	 * Update every Nth connection each frame. If 1 then every connection will be updated every frame, if 2 then half of the connections will be updated per frame and so on.
+	 * Keep this number low. The value will be clamped to 128. Due to the nature of the throttling objects may linger for N-1 extra frames before considered out of scope.
+	 */
+	uint8 HysteresisUpdateConnectionThrottling = 1;
+
+	/** Specialized configuration profiles */
+	UPROPERTY(Config)
+	TArray<FObjectScopeHysteresisProfile> HysteresisProfiles;
+};
+
+inline bool UReplicationFilteringConfig::IsObjectScopeHysteresisEnabled() const
+{
+	return bEnableObjectScopeHysteresis;
+}
+
+inline uint8 UReplicationFilteringConfig::GetDefaultHysteresisFrameCount() const
+{
+	return DefaultHysteresisFrameCount;
+}
+
+inline uint8 UReplicationFilteringConfig::GetHysteresisUpdateConnectionThrottling() const
+{
+	return FMath::Clamp<uint8>(HysteresisUpdateConnectionThrottling, 1U, 128U);
+}
+
+inline const TArray<FObjectScopeHysteresisProfile>& UReplicationFilteringConfig::GetHysteresisProfiles() const
+{
+	return HysteresisProfiles;
+}
