Iris - Object scope hysteresis support.  Keep dynamically filtered out objects around for a specified amount of frames. Configure behavior via hysteresis profiles that determine the frame timeout per class. The filter config for a specific class can then mention the hysteresis profile in order to get the appropriate behavior.  Config example is present in ReplicationFilteringConfig.h
Note: Hysteresis is enabled by default with a frame count of 6, configured in BaseEngine.ini.

diff --git a/Engine/Config/BaseEngine.ini b/Engine/Config/BaseEngine.ini
--- a/Engine/Config/BaseEngine.ini
+++ b/Engine/Config/BaseEngine.ini
@@ -1322,6 +1322,10 @@ MaxPartCount=4096
 +NetObjectFilterDefinitions=(FilterName=Spatial, ClassName=/Script/IrisCore.NetObjectGridWorldLocFilter, ConfigClassName=/Script/IrisCore.NetObjectGridFilterConfig)
 +NetObjectFilterDefinitions=(FilterName=NotRouted, ClassName=/Script/IrisCore.FilterOutNetObjectFilter, ConfigClassName=/Script/IrisCore.FilterOutNetObjectFilterConfig)
 
+[/Script/IrisCore.ReplicationFilteringConfig]
+bEnableObjectScopeHysteresis=true
+DefaultHysteresisFrameCount=6
+
 [/Script/IrisCore.ObjectReplicationBridgeConfig]
 RequiredNetDriverChannelClassName=/Script/Engine.DataStreamChannel
 
diff --git a/Engine/Plugins/Runtime/ReplicationSystemTestPlugin/Source/Private/Tests/ReplicationSystem/Filtering/TestFiltering.cpp b/Engine/Plugins/Runtime/ReplicationSystemTestPlugin/Source/Private/Tests/ReplicationSystem/Filtering/TestFiltering.cpp
--- a/Engine/Plugins/Runtime/ReplicationSystemTestPlugin/Source/Private/Tests/ReplicationSystem/Filtering/TestFiltering.cpp
+++ b/Engine/Plugins/Runtime/ReplicationSystemTestPlugin/Source/Private/Tests/ReplicationSystem/Filtering/TestFiltering.cpp
@@ -2,9 +2,11 @@
 
 #include "NetworkAutomationTest.h"
 #include "NetworkAutomationTestMacros.h"
+#include "Iris/IrisConstants.h"
 #include "Iris/ReplicationSystem/ReplicationSystem.h"
 #include "Iris/ReplicationSystem/Filtering/NetObjectFilter.h"
 #include "Iris/ReplicationSystem/Filtering/NetObjectFilterDefinitions.h"
+#include "Iris/ReplicationSystem/Filtering/ReplicationFilteringConfig.h"
 #include "Tests/ReplicationSystem/Filtering/MockNetObjectFilter.h"
 #include "Tests/ReplicationSystem/Filtering/TestFilteringObject.h"
 #include "Tests/ReplicationSystem/ReplicationSystemServerClientTestFixture.h"
@@ -18,6 +20,7 @@ protected:
 	virtual void SetUp() override
 	{
 		InitNetObjectFilterDefinitions();
+		InitObjectScopeHysteresisProfiles();
 		FReplicationSystemServerClientTestFixture::SetUp();
 		InitMockNetObjectFilter();
 	}
@@ -25,6 +28,7 @@ protected:
 	virtual void TearDown() override
 	{
 		FReplicationSystemServerClientTestFixture::TearDown();
+		RestoreObjectScopeHysteresisProfiles();
 		RestoreFilterDefinitions();
 	}
 
@@ -41,6 +45,87 @@ protected:
 		MockNetObjectFilter->ResetFunctionCallStatus();
 	}
 
+	uint32 GetHysteresisFrameCount(const char* ProfileName) const
+	{
+		const UReplicationFilteringConfig* Config = GetDefault<UReplicationFilteringConfig>();
+		if (const FObjectScopeHysteresisProfile* Profile = Config->GetHysteresisProfiles().FindByKey(ProfileName))
+		{
+			return Profile->HysteresisFrameCount;
+		}
+
+		return Config->GetDefaultHysteresisFrameCount();
+	}
+
+	class FScopedDefaultHysteresisFrameCount
+	{
+	public:
+		FScopedDefaultHysteresisFrameCount(uint8 DefaultHysteresisFrameCount)
+		{
+			const UClass* NetObjectFilteringConfiglass = UReplicationFilteringConfig::StaticClass();
+			if (const FProperty* Property = NetObjectFilteringConfiglass->FindPropertyByName("DefaultHysteresisFrameCount"))
+			{
+				UReplicationFilteringConfig* FilteringConfig = GetMutableDefault<UReplicationFilteringConfig>();
+				Property->CopyCompleteValue(&PrevValue, (void*)(UPTRINT(FilteringConfig) + Property->GetOffset_ForInternal()));
+				Property->CopyCompleteValue((void*)(UPTRINT(FilteringConfig) + Property->GetOffset_ForInternal()), &DefaultHysteresisFrameCount);
+				bPrevValueIsValid = true;
+			}
+		}
+
+		~FScopedDefaultHysteresisFrameCount()
+		{
+			if (!bPrevValueIsValid)
+			{
+				return;
+			}
+
+			const UClass* NetObjectFilteringConfiglass = UReplicationFilteringConfig::StaticClass();
+			if (const FProperty* Property = NetObjectFilteringConfiglass->FindPropertyByName("DefaultHysteresisFrameCount"))
+			{
+				UReplicationFilteringConfig* FilteringConfig = GetMutableDefault<UReplicationFilteringConfig>();
+				Property->CopyCompleteValue((void*)(UPTRINT(FilteringConfig) + Property->GetOffset_ForInternal()), &PrevValue);
+				bPrevValueIsValid = true;
+			}
+		}
+
+		uint8 PrevValue = 0;
+		bool bPrevValueIsValid = false;
+	};
+
+	class FScopedHysteresisUpdateConnectionThrottling
+	{
+	public:
+		FScopedHysteresisUpdateConnectionThrottling(uint8 HysteresisUpdateConnectionThrottling)
+		{
+			const UClass* NetObjectFilteringConfiglass = UReplicationFilteringConfig::StaticClass();
+			if (const FProperty* Property = NetObjectFilteringConfiglass->FindPropertyByName("HysteresisUpdateConnectionThrottling"))
+			{
+				UReplicationFilteringConfig* FilteringConfig = GetMutableDefault<UReplicationFilteringConfig>();
+				Property->CopyCompleteValue(&PrevValue, (void*)(UPTRINT(FilteringConfig) + Property->GetOffset_ForInternal()));
+				Property->CopyCompleteValue((void*)(UPTRINT(FilteringConfig) + Property->GetOffset_ForInternal()), &HysteresisUpdateConnectionThrottling);
+				bPrevValueIsValid = true;
+			}
+		}
+
+		~FScopedHysteresisUpdateConnectionThrottling()
+		{
+			if (!bPrevValueIsValid)
+			{
+				return;
+			}
+
+			const UClass* NetObjectFilteringConfiglass = UReplicationFilteringConfig::StaticClass();
+			if (const FProperty* Property = NetObjectFilteringConfiglass->FindPropertyByName("HysteresisUpdateConnectionThrottling"))
+			{
+				UReplicationFilteringConfig* FilteringConfig = GetMutableDefault<UReplicationFilteringConfig>();
+				Property->CopyCompleteValue((void*)(UPTRINT(FilteringConfig) + Property->GetOffset_ForInternal()), &PrevValue);
+				bPrevValueIsValid = true;
+			}
+		}
+
+		uint8 PrevValue = 0;
+		bool bPrevValueIsValid = false;
+	};
+
 private:
 	void InitNetObjectFilterDefinitions()
 	{
@@ -77,18 +162,74 @@ private:
 		MockNetObjectFilter = nullptr;
 	}
 
+	void InitObjectScopeHysteresisProfiles()
+	{
+		const UClass* ReplicationFilteringConfigClass = UReplicationFilteringConfig::StaticClass();
+		const FProperty* ProfilesProperty = ReplicationFilteringConfigClass->FindPropertyByName("HysteresisProfiles");
+		check(ProfilesProperty != nullptr);
+		const FProperty* EnableObjectScopeHysteresisProperty = ReplicationFilteringConfigClass->FindPropertyByName("bEnableObjectScopeHysteresis");
+		check(EnableObjectScopeHysteresisProperty != nullptr);
+
+		// Save CDO state.
+		UReplicationFilteringConfig* FilteringConfig = GetMutableDefault<UReplicationFilteringConfig>();
+		ProfilesProperty->CopyCompleteValue(&OriginalObjectScopeHysteresisProfiles, (void*)(UPTRINT(FilteringConfig) + ProfilesProperty->GetOffset_ForInternal()));
+
+		// Modify profiles to what we need
+		TArray<FObjectScopeHysteresisProfile> NewProfiles;
+		{
+			{
+				FObjectScopeHysteresisProfile& Profile = NewProfiles.Emplace_GetRef();
+				Profile.FilterProfileName = "FiveFrames";
+				Profile.HysteresisFrameCount = 5;
+			}
+
+			{
+				FObjectScopeHysteresisProfile& Profile = NewProfiles.Emplace_GetRef();
+				Profile.FilterProfileName = "OneFrame";
+				Profile.HysteresisFrameCount = 1;
+			}
+
+			{
+				FObjectScopeHysteresisProfile& Profile = NewProfiles.Emplace_GetRef();
+				Profile.FilterProfileName = "ZeroFrames";
+				Profile.HysteresisFrameCount = 0;
+			}
+		}
+
+		ProfilesProperty->CopyCompleteValue((void*)(UPTRINT(FilteringConfig) + ProfilesProperty->GetOffset_ForInternal()), &NewProfiles);
+
+		bOriginalIsObjectScopeHysteresisEnabled = FilteringConfig->IsObjectScopeHysteresisEnabled();
+		const bool bEnableObjectScopeHysteresis = true;
+		EnableObjectScopeHysteresisProperty->CopyCompleteValue((void*)(UPTRINT(FilteringConfig) + EnableObjectScopeHysteresisProperty->GetOffset_ForInternal()), &bEnableObjectScopeHysteresis);
+	}
+
+	void RestoreObjectScopeHysteresisProfiles()
+	{
+		// Restore CDO state from the saved state.
+		const UClass* ReplicationFilteringConfigClass = UReplicationFilteringConfig::StaticClass();
+		const FProperty* ProfilesProperty = ReplicationFilteringConfigClass->FindPropertyByName("HysteresisProfiles");
+		const FProperty* EnableObjectScopeHysteresisProperty = ReplicationFilteringConfigClass->FindPropertyByName("bEnableObjectScopeHysteresis");
+
+		UReplicationFilteringConfig* FilteringConfig = GetMutableDefault<UReplicationFilteringConfig>();
+		ProfilesProperty->CopyCompleteValue((void*)(UPTRINT(FilteringConfig) + ProfilesProperty->GetOffset_ForInternal()), &OriginalObjectScopeHysteresisProfiles);
+		OriginalObjectScopeHysteresisProfiles.Empty();
+		EnableObjectScopeHysteresisProperty->CopyCompleteValue((void*)(UPTRINT(FilteringConfig) + EnableObjectScopeHysteresisProperty->GetOffset_ForInternal()), &bOriginalIsObjectScopeHysteresisEnabled);
+	}
+
 	void InitMockNetObjectFilter()
 	{
 		MockNetObjectFilter = CastChecked<UMockNetObjectFilter>(Server->GetReplicationSystem()->GetFilter(GetMockFilterName()));
 		MockFilterHandle = Server->GetReplicationSystem()->GetFilterHandle(GetMockFilterName());
 	}
-
 protected:
 	UMockNetObjectFilter* MockNetObjectFilter;
 	FNetObjectFilterHandle MockFilterHandle;
+	FName ObjectHysteresisProfileName;
 
 private:
 	TArray<FNetObjectFilterDefinition> OriginalFilterDefinitions;
+	TArray<FObjectScopeHysteresisProfile> OriginalObjectScopeHysteresisProfiles;
+	bool bOriginalIsObjectScopeHysteresisEnabled = false;
 };
 
 
@@ -1035,7 +1176,7 @@ UE_NET_TEST_FIXTURE(FTestFilteringFixture, NestedDependentObjectIsFilteredAsPare
 	Server->SendAndDeliverTo(Client, DeliverPacket);
 	Server->PostSendUpdate();
 
-	// We expect the object to exist and the future dependent and future dependent objects not to exist
+	// We expect the object to exist and the future dependent and future nested dependent objects not to exist
 	UE_NET_ASSERT_NE(Client->GetReplicationBridge()->GetReplicatedObject(ServerObject->NetRefHandle), nullptr);
 	UE_NET_ASSERT_EQ(Client->GetReplicationBridge()->GetReplicatedObject(ServerFutureDependentObject->NetRefHandle), nullptr);
 	UE_NET_ASSERT_EQ(Client->GetReplicationBridge()->GetReplicatedObject(ServerFutureNestedDependentObject->NetRefHandle), nullptr);
@@ -2274,4 +2415,625 @@ UE_NET_TEST_FIXTURE(FTestFilteringFixture, LateAddedConnectionWorksWithComplexGr
 	UE_NET_ASSERT_NE(LateAddedClient->GetReplicationBridge()->GetReplicatedObject(ServerSubObjects[3]->NetRefHandle), nullptr);
 }
 
+UE_NET_TEST_FIXTURE(FTestFilteringFixture, HysteresisIsEnabled)
+{
+	UE_NET_ASSERT_TRUE_MSG(GetDefault<UReplicationFilteringConfig>()->IsObjectScopeHysteresisEnabled(), "Error: Hysteresis is disabled. All hysteresis tests will fail.");
+}
+
+// Dynamic filtering should cause hysteresis to kick in for a filtered out object with a filter profile.
+UE_NET_TEST_FIXTURE(FTestFilteringFixture, HysteresisKicksInForDynamicallyFilteredOutObjectWithFilterProfile)
+{
+	// Add client
+	FReplicationSystemTestClient* Client = CreateClient();
+
+	SetDynamicFilterStatus(ENetFilterStatus::Allow);
+
+	// Spawn object on server and set filter and filter profile for hysteresis
+	UReplicatedTestObject* ServerObject = Server->CreateObject({.IrisComponentCount = 0});
+	Server->ReplicationSystem->SetFilter(ServerObject->NetRefHandle, MockFilterHandle, /* .FilterProfile = */ "FiveFrames");
+
+	Server->UpdateAndSend({ Client });
+
+	SetDynamicFilterStatus(ENetFilterStatus::Disallow);
+
+	const uint32 HysteresisFrameCount = GetHysteresisFrameCount("FiveFrames");
+	UE_NET_ASSERT_EQ(HysteresisFrameCount, 5U);
+	for (uint32 It = 0, EndIt = HysteresisFrameCount; It < EndIt; ++It)
+	{
+		Server->UpdateAndSend({ Client });
+	}
+
+	// At this point the object should still exist on the client due to hysteresis
+	UE_NET_ASSERT_NE(Client->GetReplicationBridge()->GetReplicatedObject(ServerObject->NetRefHandle), nullptr);
+
+	Server->UpdateAndSend({ Client });
+
+	// Hysteresis frame count has passed. The object shuld now be destroyed.
+	UE_NET_ASSERT_EQ(Client->GetReplicationBridge()->GetReplicatedObject(ServerObject->NetRefHandle), nullptr);
+}
+
+// Dynamic filtering should cause hysteresis to kick in for a filtered out object without a filter profile, thus using default hysteresis frame count.
+UE_NET_TEST_FIXTURE(FTestFilteringFixture, HysteresisKicksInForDynamicallyFilteredOutObjectWithoutFilterProfile)
+{
+	constexpr uint32 DefaultHysteresisFrameCount = 3;
+	FScopedDefaultHysteresisFrameCount ScopedDefaultHysteresisFrameCount(DefaultHysteresisFrameCount);
+
+	// Add client
+	FReplicationSystemTestClient* Client = CreateClient();
+
+	SetDynamicFilterStatus(ENetFilterStatus::Allow);
+
+	// Spawn object on server and set filter and filter profile for hysteresis
+	UReplicatedTestObject* ServerObject = Server->CreateObject({ .IrisComponentCount = 0 });
+	Server->ReplicationSystem->SetFilter(ServerObject->NetRefHandle, MockFilterHandle);
+
+	Server->UpdateAndSend({ Client });
+
+	SetDynamicFilterStatus(ENetFilterStatus::Disallow);
+
+	for (uint32 It = 0, EndIt = DefaultHysteresisFrameCount; It < EndIt; ++It)
+	{
+		Server->UpdateAndSend({ Client });
+	}
+
+	// At this point the object should still exist on the client due to hysteresis
+	UE_NET_ASSERT_NE(Client->GetReplicationBridge()->GetReplicatedObject(ServerObject->NetRefHandle), nullptr);
+
+	Server->UpdateAndSend({ Client });
+
+	// Hysteresis frame count has passed. The object shuld now be destroyed.
+	UE_NET_ASSERT_EQ(Client->GetReplicationBridge()->GetReplicatedObject(ServerObject->NetRefHandle), nullptr);
+}
+
+// Owner filtering changes should not cause hysteresis to kick in.
+UE_NET_TEST_FIXTURE(FTestFilteringFixture, HysteresisDoesNotKickInForOwnerFilteredObject)
+{
+	constexpr uint32 DefaultHysteresisFrameCount = 3;
+	FScopedDefaultHysteresisFrameCount ScopedDefaultHysteresisFrameCount(DefaultHysteresisFrameCount);
+
+	// Add client
+	FReplicationSystemTestClient* Client = CreateClient();
+
+	// Spawn object on server
+	UReplicatedTestObject* ServerObject = Server->CreateObject({ .IrisComponentCount = 0 });
+	Server->ReplicationSystem->SetFilter(ServerObject->NetRefHandle, ToOwnerFilterHandle);
+	Server->ReplicationSystem->SetOwningNetConnection(ServerObject->NetRefHandle, Client->ConnectionIdOnServer);
+
+	// Send and deliver packet
+	Server->UpdateAndSend({ Client });
+
+	// Object should have been created on the client
+	UE_NET_ASSERT_NE(Client->GetReplicationBridge()->GetReplicatedObject(ServerObject->NetRefHandle), nullptr);
+
+	// Switch owner and make sure objects gets immediately destroyed on the client
+	Server->ReplicationSystem->SetOwningNetConnection(ServerObject->NetRefHandle, InvalidConnectionId);
+
+	// Send and deliver packet
+	Server->UpdateAndSend({ Client });
+
+	// Owner filtering should not cause hysteresis to kick in so the owner change should cause the client object to be destroyed immediately.
+	UE_NET_ASSERT_EQ(Client->GetReplicationBridge()->GetReplicatedObject(ServerObject->NetRefHandle), nullptr);
+}
+
+// Exclusion group filtering changes should not cause hysteresis to kick in.
+UE_NET_TEST_FIXTURE(FTestFilteringFixture, HysteresisDoesNotKickInForExclusionGroupFilteredObject)
+{
+	constexpr uint32 DefaultHysteresisFrameCount = 3;
+	FScopedDefaultHysteresisFrameCount ScopedDefaultHysteresisFrameCount(DefaultHysteresisFrameCount);
+
+	// Add client
+	FReplicationSystemTestClient* Client = CreateClient();
+
+	// Spawn object on server
+	UReplicatedTestObject* ServerObject = Server->CreateObject({ .IrisComponentCount = 0 });
+
+	// Add to exclusion group that allows replication to all connections.
+	FNetObjectGroupHandle GroupHandle = Server->ReplicationSystem->CreateGroup();
+	Server->ReplicationSystem->AddToGroup(GroupHandle, ServerObject->NetRefHandle);
+	Server->ReplicationSystem->AddExclusionFilterGroup(GroupHandle);
+	Server->ReplicationSystem->SetGroupFilterStatus(GroupHandle, ENetFilterStatus::Allow);
+
+	// Send and deliver packet
+	Server->UpdateAndSend({ Client });
+
+	// Object should have been created on the client
+	UE_NET_ASSERT_NE(Client->GetReplicationBridge()->GetReplicatedObject(ServerObject->NetRefHandle), nullptr);
+
+	// Disallow the group to be replicated.
+	Server->ReplicationSystem->SetGroupFilterStatus(GroupHandle, ENetFilterStatus::Disallow);
+
+	// Send and deliver packet
+	Server->UpdateAndSend({ Client });
+
+	// Exclusion group filtering changes should not cause hysteresis to kick in so the client object to be destroyed immediately when the group disallows replication.
+	UE_NET_ASSERT_EQ(Client->GetReplicationBridge()->GetReplicatedObject(ServerObject->NetRefHandle), nullptr);
+}
+
+// If an object is filtered out from the start we should not start replicating it at all
+UE_NET_TEST_FIXTURE(FTestFilteringFixture, HysteresisDoesNotKickInForNewlyCreatedObject)
+{
+	// Add client
+	FReplicationSystemTestClient* Client = CreateClient();
+
+	SetDynamicFilterStatus(ENetFilterStatus::Disallow);
+
+	// Spawn object on server and set filter and filter profile for hysteresis
+	UReplicatedTestObject* ServerObject = Server->CreateObject({ .IrisComponentCount = 0 });
+	Server->ReplicationSystem->SetFilter(ServerObject->NetRefHandle, MockFilterHandle, /* .FilterProfile = */ "FiveFrames");
+
+	Server->UpdateAndSend({ Client });
+
+	// The object was filtered out from the start so hysteresis should not cause it to start replicating.
+	UE_NET_ASSERT_NE(Client->GetReplicationBridge()->GetReplicatedObject(ServerObject->NetRefHandle), nullptr);
+}
+
+// Destroyed objects are expected to be destroyed as quickly as possible.
+UE_NET_TEST_FIXTURE(FTestFilteringFixture, HysteresisDoesNotKickInForDestroyedObject)
+{
+	// Add client
+	FReplicationSystemTestClient* Client = CreateClient();
+
+	SetDynamicFilterStatus(ENetFilterStatus::Allow);
+
+	// Spawn object on server and set filter and filter profile for hysteresis
+	UReplicatedTestObject* ServerObject = Server->CreateObject({ .IrisComponentCount = 0 });
+	Server->ReplicationSystem->SetFilter(ServerObject->NetRefHandle, MockFilterHandle, /* .FilterProfile = */ "FiveFrames");
+
+	Server->UpdateAndSend({ Client });
+
+	// Destroy object and make sure it immediately gets destroyed on the client as well.
+	const FNetRefHandle ServerNetRefHandle = ServerObject->NetRefHandle;
+	Server->DestroyObject(ServerObject);
+	ServerObject = nullptr;
+
+	Server->UpdateAndSend({ Client });
+
+	// The object was destroyed on the server and should be destroyed as soon as possible on the client as well.
+	UE_NET_ASSERT_EQ(Client->GetReplicationBridge()->GetReplicatedObject(ServerNetRefHandle), nullptr);
+}
+
+// Test case where dependent objects are filtered out yet should be replicated due to their parent being in scope.
+UE_NET_TEST_FIXTURE(FTestFilteringFixture, HysteresisDoesNotKickInForDependentObjectWithReplicatedParent)
+{
+	constexpr uint32 DefaultHysteresisFrameCount = 2;
+	FScopedDefaultHysteresisFrameCount ScopedDefaultHysteresisFrameCount(DefaultHysteresisFrameCount);
+
+	// Add client
+	FReplicationSystemTestClient* Client = CreateClient();
+
+	// Only the dependent object will have a filter set in this test. Filter out by default.
+	SetDynamicFilterStatus(ENetFilterStatus::Disallow);
+
+	UReplicatedTestObject* ServerObject = Server->CreateObject({ .IrisComponentCount = 0 });
+	UReplicatedTestObject* ServerDependentObject = Server->CreateObject({ .IrisComponentCount = 0 });
+	Server->ReplicationSystem->SetFilter(ServerDependentObject->NetRefHandle, MockFilterHandle);
+
+	Server->UpdateAndSend({ Client });
+
+	// Add the dependency
+	Server->ReplicationBridge->AddDependentObject(ServerObject->NetRefHandle, ServerDependentObject->NetRefHandle);
+
+	Server->UpdateAndSend({ Client });
+
+	// Make sure the dependent object has been created 
+	UE_NET_ASSERT_NE(Client->GetReplicationBridge()->GetReplicatedObject(ServerDependentObject->NetRefHandle), nullptr);
+
+	// Perform send update for a few frames and make sure the dependent object stays relevant.
+	for (uint32 It = 0; It != DefaultHysteresisFrameCount; ++It)
+	{
+		Server->UpdateAndSend({ Client });
+		UE_NET_ASSERT_NE(Client->GetReplicationBridge()->GetReplicatedObject(ServerDependentObject->NetRefHandle), nullptr);
+	}
+}
+
+// Test case where dependent objects are filtered out yet should be replicated due to their parent being in scope.
+UE_NET_TEST_FIXTURE(FTestFilteringFixture, HysteresisKicksInForDependentObjectWhenParentIsFilteredOut)
+{
+	constexpr uint32 DefaultHysteresisFrameCount = 3;
+	FScopedDefaultHysteresisFrameCount ScopedDefaultHysteresisFrameCount(DefaultHysteresisFrameCount);
+
+	// Add client
+	FReplicationSystemTestClient* Client = CreateClient();
+
+	// Both objects have a filter set in this test
+	SetDynamicFilterStatus(ENetFilterStatus::Allow);
+
+	UReplicatedTestObject* ServerObject = Server->CreateObject({ .IrisComponentCount = 0 });
+	Server->ReplicationSystem->SetFilter(ServerObject->NetRefHandle, MockFilterHandle, /* .FilterProfile = */ "ZeroFrames");
+
+	UReplicatedTestObject* ServerDependentObject = Server->CreateObject({ .IrisComponentCount = 0 });
+	Server->ReplicationSystem->SetFilter(ServerDependentObject->NetRefHandle, MockFilterHandle);
+
+	Server->ReplicationBridge->AddDependentObject(ServerObject->NetRefHandle, ServerDependentObject->NetRefHandle);
+
+	Server->UpdateAndSend({ Client });
+
+	// Make sure tall objects have been created 
+	UE_NET_ASSERT_NE(Client->GetReplicationBridge()->GetReplicatedObject(ServerObject->NetRefHandle), nullptr);
+	UE_NET_ASSERT_NE(Client->GetReplicationBridge()->GetReplicatedObject(ServerDependentObject->NetRefHandle), nullptr);
+
+	// Filter out both objects.
+	SetDynamicFilterStatus(ENetFilterStatus::Disallow);
+
+	// Perform send update for a few frames and make sure the dependent object stays relevant.
+	for (uint32 It = 0; It != DefaultHysteresisFrameCount; ++It)
+	{
+		Server->UpdateAndSend({ Client });
+		UE_NET_ASSERT_NE(Client->GetReplicationBridge()->GetReplicatedObject(ServerDependentObject->NetRefHandle), nullptr);
+	}
+
+	// Eventually the dependent object should be filtered out
+	Server->UpdateAndSend({ Client });
+	UE_NET_ASSERT_EQ(Client->GetReplicationBridge()->GetReplicatedObject(ServerDependentObject->NetRefHandle), nullptr);
+
+	// Make sure it stays filtered out
+	for (uint32 It = 0; It != DefaultHysteresisFrameCount + 1; ++It)
+	{
+		Server->UpdateAndSend({ Client });
+		UE_NET_ASSERT_EQ(Client->GetReplicationBridge()->GetReplicatedObject(ServerDependentObject->NetRefHandle), nullptr);
+	}
+}
+
+// Verify dependent object is replicated after parent is filtered out and ends up in hysteresis
+UE_NET_TEST_FIXTURE(FTestFilteringFixture, DependentObjectIsReplicatedWhenParentIsInHysteresis)
+{
+	constexpr uint32 DefaultHysteresisFrameCount = 3;
+	FScopedDefaultHysteresisFrameCount ScopedDefaultHysteresisFrameCount(DefaultHysteresisFrameCount);
+
+	// Add client
+	FReplicationSystemTestClient* Client = CreateClient();
+
+	// Both objects have a filter set in this test
+	SetDynamicFilterStatus(ENetFilterStatus::Allow);
+
+	UReplicatedTestObject* ServerObject = Server->CreateObject({ .IrisComponentCount = 0 });
+	Server->ReplicationSystem->SetFilter(ServerObject->NetRefHandle, MockFilterHandle);
+
+	UReplicatedTestObject* ServerDependentObject = Server->CreateObject({ .IrisComponentCount = 0 });
+	Server->ReplicationSystem->SetFilter(ServerDependentObject->NetRefHandle, MockFilterHandle, /* .FilterProfile = */ "ZeroFrames");
+
+	Server->ReplicationBridge->AddDependentObject(ServerObject->NetRefHandle, ServerDependentObject->NetRefHandle);
+
+	Server->UpdateAndSend({ Client });
+
+	// Make sure tall objects have been created 
+	UE_NET_ASSERT_NE(Client->GetReplicationBridge()->GetReplicatedObject(ServerObject->NetRefHandle), nullptr);
+	UE_NET_ASSERT_NE(Client->GetReplicationBridge()->GetReplicatedObject(ServerDependentObject->NetRefHandle), nullptr);
+
+	// Filter out both objects.
+	SetDynamicFilterStatus(ENetFilterStatus::Disallow);
+
+	// Perform send update for a few frames and make sure both objects stays relevant.
+	for (uint32 It = 0; It != DefaultHysteresisFrameCount; ++It)
+	{
+		Server->UpdateAndSend({ Client });
+		UE_NET_ASSERT_NE(Client->GetReplicationBridge()->GetReplicatedObject(ServerObject->NetRefHandle), nullptr);
+		UE_NET_ASSERT_NE(Client->GetReplicationBridge()->GetReplicatedObject(ServerDependentObject->NetRefHandle), nullptr);
+	}
+
+	// Eventually the both objects should be filtered out
+	Server->UpdateAndSend({ Client });
+	UE_NET_ASSERT_EQ(Client->GetReplicationBridge()->GetReplicatedObject(ServerObject->NetRefHandle), nullptr);
+	UE_NET_ASSERT_EQ(Client->GetReplicationBridge()->GetReplicatedObject(ServerDependentObject->NetRefHandle), nullptr);
+
+	// Make sure they stay filtered out
+	for (uint32 It = 0; It != DefaultHysteresisFrameCount + 1; ++It)
+	{
+		Server->UpdateAndSend({ Client });
+		UE_NET_ASSERT_EQ(Client->GetReplicationBridge()->GetReplicatedObject(ServerObject->NetRefHandle), nullptr);
+		UE_NET_ASSERT_EQ(Client->GetReplicationBridge()->GetReplicatedObject(ServerDependentObject->NetRefHandle), nullptr);
+	}
+}
+
+// Verify dependent object with hysteresis is replicated as long as the parent object with hysteresis is.
+UE_NET_TEST_FIXTURE(FTestFilteringFixture, DependentObjectWithHysteresisIsReplicatedWhenParentIsInHysteresis)
+{
+	constexpr uint32 TwoFrameHysteresis = 2;
+	FScopedDefaultHysteresisFrameCount ScopedDefaultHysteresisFrameCount(TwoFrameHysteresis);
+
+	// Add client
+	FReplicationSystemTestClient* Client = CreateClient();
+
+	// Both objects have a filter set in this test
+	SetDynamicFilterStatus(ENetFilterStatus::Allow);
+
+	UReplicatedTestObject* ServerObject = Server->CreateObject({ .IrisComponentCount = 0 });
+	Server->ReplicationSystem->SetFilter(ServerObject->NetRefHandle, MockFilterHandle);
+
+	UReplicatedTestObject* ServerDependentObject = Server->CreateObject({ .IrisComponentCount = 0 });
+	Server->ReplicationSystem->SetFilter(ServerDependentObject->NetRefHandle, MockFilterHandle, /* .FilterProfile = */ "FiveFrames");
+
+	Server->ReplicationBridge->AddDependentObject(ServerObject->NetRefHandle, ServerDependentObject->NetRefHandle);
+
+	Server->UpdateAndSend({ Client });
+
+	// Make sure tall objects have been created 
+	UE_NET_ASSERT_NE(Client->GetReplicationBridge()->GetReplicatedObject(ServerObject->NetRefHandle), nullptr);
+	UE_NET_ASSERT_NE(Client->GetReplicationBridge()->GetReplicatedObject(ServerDependentObject->NetRefHandle), nullptr);
+
+	// Filter out both objects.
+	SetDynamicFilterStatus(ENetFilterStatus::Disallow);
+
+	// Perform send update for a few frames and make sure both objects stays relevant.
+	for (uint32 It = 0; It != TwoFrameHysteresis; ++It)
+	{
+		Server->UpdateAndSend({ Client });
+		UE_NET_ASSERT_NE(Client->GetReplicationBridge()->GetReplicatedObject(ServerObject->NetRefHandle), nullptr);
+		UE_NET_ASSERT_NE(Client->GetReplicationBridge()->GetReplicatedObject(ServerDependentObject->NetRefHandle), nullptr);
+	}
+
+	// Eventually the parent object should be filtered out
+	Server->UpdateAndSend({ Client });
+	UE_NET_ASSERT_EQ(Client->GetReplicationBridge()->GetReplicatedObject(ServerObject->NetRefHandle), nullptr);
+
+	// At this point the dependent object should still exist as 5 > 2
+	UE_NET_ASSERT_NE(Client->GetReplicationBridge()->GetReplicatedObject(ServerDependentObject->NetRefHandle), nullptr);
+
+	// Make sure dependent object sticks around for a bit
+	for (uint32 It = 0; It != 5 - TwoFrameHysteresis - 1; ++It)
+	{
+		Server->UpdateAndSend({ Client });
+		UE_NET_ASSERT_NE(Client->GetReplicationBridge()->GetReplicatedObject(ServerDependentObject->NetRefHandle), nullptr);
+	}
+
+	// We don't know exactly when hysteresis kicks in, if it starts when parent hysteresis starts or when parent is finally filtered out.
+	// BUT at the very least it should be filtered out after an additional five frames have passed.
+	for (uint32 It = 0; It != 5; ++It)
+	{
+		Server->UpdateAndSend({ Client });
+	}
+	UE_NET_ASSERT_EQ(Client->GetReplicationBridge()->GetReplicatedObject(ServerDependentObject->NetRefHandle), nullptr);
+}
+
+// Test case where dependent objects are filtered out yet should be replicated due to their parent being in scope. Dependency is then removed and hysteresis should kick in.
+UE_NET_TEST_FIXTURE(FTestFilteringFixture, HysteresisKicksInForFilteredOutFormerlyDependentObject)
+{
+	constexpr uint32 DefaultHysteresisFrameCount = 2;
+	FScopedDefaultHysteresisFrameCount ScopedDefaultHysteresisFrameCount(DefaultHysteresisFrameCount);
+
+	// Add client
+	FReplicationSystemTestClient* Client = CreateClient();
+
+	// Only the dependent object will have a filter set in this test. Filter out by default.
+	SetDynamicFilterStatus(ENetFilterStatus::Disallow);
+
+	UReplicatedTestObject* ServerObject = Server->CreateObject({ .IrisComponentCount = 0 });
+	UReplicatedTestObject* ServerDependentObject = Server->CreateObject({ .IrisComponentCount = 0 });
+	Server->ReplicationSystem->SetFilter(ServerDependentObject->NetRefHandle, MockFilterHandle);
+
+	Server->UpdateAndSend({ Client });
+
+	// Add the dependency
+	Server->ReplicationBridge->AddDependentObject(ServerObject->NetRefHandle, ServerDependentObject->NetRefHandle);
+
+	Server->UpdateAndSend({ Client });
+
+	// Make sure the dependent object has been created 
+	UE_NET_ASSERT_NE(Client->GetReplicationBridge()->GetReplicatedObject(ServerDependentObject->NetRefHandle), nullptr);
+
+	// Remove dependency
+	Server->ReplicationBridge->RemoveDependentObject(ServerObject->NetRefHandle, ServerDependentObject->NetRefHandle);
+
+	// Perform send update for a few frames and make sure the dependent object stays relevant.
+	for (uint32 It = 0; It != DefaultHysteresisFrameCount; ++It)
+	{
+		Server->UpdateAndSend({ Client });
+		UE_NET_ASSERT_NE(Client->GetReplicationBridge()->GetReplicatedObject(ServerDependentObject->NetRefHandle), nullptr);
+	}
+
+	// The formerly dependent object should now be filtered out
+	Server->UpdateAndSend({ Client });
+	UE_NET_ASSERT_EQ(Client->GetReplicationBridge()->GetReplicatedObject(ServerDependentObject->NetRefHandle), nullptr);
+}
+
+// Test case where an object is replicated, then filtered out and filtered in prior to hysteresis frame timeout.
+UE_NET_TEST_FIXTURE(FTestFilteringFixture, FilteringChangesDoesNotCauseHysteresisToFilterOutObject)
+{
+	constexpr uint32 DefaultHysteresisFrameCount = 3;
+	FScopedDefaultHysteresisFrameCount ScopedDefaultHysteresisFrameCount(DefaultHysteresisFrameCount);
+
+	// Add client
+	FReplicationSystemTestClient* Client = CreateClient();
+
+	SetDynamicFilterStatus(ENetFilterStatus::Allow);
+
+	UReplicatedTestObject* ServerObject = Server->CreateObject({ .IrisComponentCount = 0 });
+	Server->ReplicationSystem->SetFilter(ServerObject->NetRefHandle, MockFilterHandle);
+
+	Server->UpdateAndSend({ Client });
+
+	// Make sure the object has been created 
+	UE_NET_ASSERT_NE(Client->GetReplicationBridge()->GetReplicatedObject(ServerObject->NetRefHandle), nullptr);
+
+	// Filter out object and perform update. Due to hysteresis the object should remain replicated.
+	SetDynamicFilterStatus(ENetFilterStatus::Disallow);
+
+	Server->UpdateAndSend({ Client });
+
+	// Hysteresis should cause the object to remain replicated.
+	UE_NET_ASSERT_NE(Client->GetReplicationBridge()->GetReplicatedObject(ServerObject->NetRefHandle), nullptr);
+
+	// Filter in the object again and make sure it stays replicated.
+	SetDynamicFilterStatus(ENetFilterStatus::Allow);
+
+	// Perform send update for a few frames and make sure the dependent object stays relevant.
+	for (uint32 It = 0; It != DefaultHysteresisFrameCount; ++It)
+	{
+		Server->UpdateAndSend({ Client });
+		UE_NET_ASSERT_NE(Client->GetReplicationBridge()->GetReplicatedObject(ServerObject->NetRefHandle), nullptr);
+	}
+}
+
+// Test case where inclusion group added objects are filtered out yet should be replicated due to the inclusion group allowing replication. Inclusion group then disallows replication causing hysteresis to kick in.
+UE_NET_TEST_FIXTURE(FTestFilteringFixture, HysteresisKicksInForFormerlyInclusionGroupAllowedObject)
+{
+	constexpr uint32 DefaultHysteresisFrameCount = 2;
+	FScopedDefaultHysteresisFrameCount ScopedDefaultHysteresisFrameCount(DefaultHysteresisFrameCount);
+
+	// Add client
+	FReplicationSystemTestClient* Client = CreateClient();
+
+	// Add object to filter which disallows replication but inclusion group that allows it.
+	SetDynamicFilterStatus(ENetFilterStatus::Disallow);
+
+	UReplicatedTestObject* ServerObject = Server->CreateObject({ .IrisComponentCount = 0 });
+	Server->ReplicationSystem->SetFilter(ServerObject->NetRefHandle, MockFilterHandle);
+
+	FNetObjectGroupHandle GroupHandle = Server->ReplicationSystem->CreateGroup();
+	Server->ReplicationSystem->AddToGroup(GroupHandle, ServerObject->NetRefHandle);
+	Server->ReplicationSystem->AddInclusionFilterGroup(GroupHandle);
+	Server->ReplicationSystem->SetGroupFilterStatus(GroupHandle, ENetFilterStatus::Allow);
+
+	Server->UpdateAndSend({ Client });
+
+	// Make sure the object has been created 
+	UE_NET_ASSERT_NE(Client->GetReplicationBridge()->GetReplicatedObject(ServerObject->NetRefHandle), nullptr);
+
+	// Disallow replication of inclusion group
+	Server->ReplicationSystem->SetGroupFilterStatus(GroupHandle, ENetFilterStatus::Disallow);
+
+	// Perform send update forthe hysteresis duration and make sure the object stays replicated.
+	for (uint32 It = 0; It != DefaultHysteresisFrameCount; ++It)
+	{
+		Server->UpdateAndSend({ Client });
+		UE_NET_ASSERT_NE(Client->GetReplicationBridge()->GetReplicatedObject(ServerObject->NetRefHandle), nullptr);
+	}
+
+	// Now the object should now be filtered out
+	Server->UpdateAndSend({ Client });
+	UE_NET_ASSERT_EQ(Client->GetReplicationBridge()->GetReplicatedObject(ServerObject->NetRefHandle), nullptr);
+}
+
+// Make sure that connection throttling does not cause objects to be filtered out too soon. Also verify throttling occurs.
+UE_NET_TEST_FIXTURE(FTestFilteringFixture, HysteresisConnectionThrottlingWorksAsExpected)
+{
+	constexpr uint8 ConnectionThrottlingFrameCount = 5;
+	FScopedHysteresisUpdateConnectionThrottling ConnectionThrottling(ConnectionThrottlingFrameCount);
+
+	// Add client
+	FReplicationSystemTestClient* Client = CreateClient();
+
+	// Figure out hysteresis update frame.
+	{
+		SetDynamicFilterStatus(ENetFilterStatus::Allow);
+
+		// Spawn object on server and set filter and filter profile for hysteresis
+		UReplicatedTestObject* ServerObject = Server->CreateObject({ .IrisComponentCount = 0 });
+		Server->ReplicationSystem->SetFilter(ServerObject->NetRefHandle, MockFilterHandle, /* .FilterProfile = */ "OneFrame");
+
+		Server->UpdateAndSend({ Client });
+
+		SetDynamicFilterStatus(ENetFilterStatus::Disallow);
+
+		// As we have an hysteresis of one frame we will detect immediately when the trottling is updated
+		for (uint32 It = 0, EndIt = ConnectionThrottlingFrameCount; It < EndIt; ++It)
+		{
+			Server->UpdateAndSend({ Client });
+			if (Client->GetReplicationBridge()->GetReplicatedObject(ServerObject->NetRefHandle) == nullptr)
+			{
+				break;
+			}
+		}
+
+		// Object must have been destroyed on the client by now.
+		UE_NET_ASSERT_EQ(Client->GetReplicationBridge()->GetReplicatedObject(ServerObject->NetRefHandle), nullptr);
+	}
+
+	// Make sure object is kept alive for at least the expected frame count
+	{
+		SetDynamicFilterStatus(ENetFilterStatus::Allow);
+
+		UReplicatedTestObject* ServerObject = Server->CreateObject({ .IrisComponentCount = 0 });
+		Server->ReplicationSystem->SetFilter(ServerObject->NetRefHandle, MockFilterHandle, /* .FilterProfile = */ "FiveFrames");
+
+		// Advance up to one frame before we expect hysteresis update
+		for (uint32 It = 0, EndIt = ConnectionThrottlingFrameCount - 1; It < EndIt; ++It)
+		{
+			Server->UpdateAndSend({ Client });
+		}
+
+		// Filter out object. Hysteresis update should be performed but the object should not be filtered out immediately as not enough frames have passed.
+		SetDynamicFilterStatus(ENetFilterStatus::Disallow);
+
+		uint32 WaitFrameCount = 0;
+		for (uint32 It = 0, EndIt = 2U * ConnectionThrottlingFrameCount; It < EndIt; ++It)
+		{
+			++WaitFrameCount;
+			Server->UpdateAndSend({ Client });
+
+			if (Client->GetReplicationBridge()->GetReplicatedObject(ServerObject->NetRefHandle) == nullptr)
+			{
+				break;
+			}
+		}
+
+		// Object must have been destroyed on the client by now.
+		UE_NET_ASSERT_EQ(Client->GetReplicationBridge()->GetReplicatedObject(ServerObject->NetRefHandle), nullptr);
+
+		// FiveFames profile means at least five frames of waiting
+		UE_NET_ASSERT_GT(WaitFrameCount, 5U);
+
+		// This assert assumes ConnectionThrottlingFrameCount is five as well. If it's four we'd expect 4+4+1 frames.
+		UE_NET_ASSERT_EQ(WaitFrameCount, ConnectionThrottlingFrameCount + 1U);
+	}
+}
+
+// Test that a lot of objects can be filtered out on the same frame.
+UE_NET_TEST_FIXTURE(FTestFilteringFixture, LotsOfObjectsCanBeFilteredOutViaHysteresisInOneFrame)
+{
+	constexpr uint32 HighObjectCount = 65;
+
+	// Add client
+	FReplicationSystemTestClient* Client = CreateClient();
+
+	SetDynamicFilterStatus(ENetFilterStatus::Allow);
+
+	// Spawn lots of objects on server
+	UReplicatedTestObject* ServerObjects[HighObjectCount];
+	for (UReplicatedTestObject*& ServerObject : ServerObjects)
+	{
+		ServerObject = Server->CreateObject({ .IrisComponentCount = 0 });
+		Server->ReplicationSystem->SetFilter(ServerObject->NetRefHandle, MockFilterHandle, /* .FilterProfile = */ "OneFrame");
+	}
+
+	// Send and deliver packets until we believe all objects have been created on the client
+	bool bAllObjectsCreated = false;
+	for (uint32 It = 0; It < HighObjectCount; ++It)
+	{
+		Server->UpdateAndSend({ Client });
+		if (Client->GetReplicationBridge()->GetReplicatedObject(ServerObjects[HighObjectCount - 1]->NetRefHandle) != nullptr)
+		{
+			bAllObjectsCreated = true;
+
+			for (const UReplicatedTestObject* ServerObject : ServerObjects)
+			{
+				if (Client->GetReplicationBridge()->GetReplicatedObject(ServerObject->NetRefHandle) == nullptr)
+				{
+					bAllObjectsCreated = false;
+					break;
+				}
+			}
+			
+			if (bAllObjectsCreated)
+			{
+				break;
+			}
+		}
+	}
+
+	UE_NET_ASSERT_TRUE(bAllObjectsCreated);
+
+	SetDynamicFilterStatus(ENetFilterStatus::Disallow);
+
+	// Send and deliver packet. Need to update twice to have the objects filtered out.
+	Server->UpdateAndSend({ Client });
+	Server->UpdateAndSend({ Client });
+
+	// All client objects should be destroyed
+	for (const UReplicatedTestObject* ServerObject : ServerObjects)
+	{
+		UE_NET_ASSERT_EQ(Client->GetReplicationBridge()->GetReplicatedObject(ServerObject->NetRefHandle), nullptr);
+	}
+}
+
 } // end namespace UE::Net::Private
diff --git a/Engine/Source/Runtime/Experimental/Iris/Core/Private/Iris/ReplicationSystem/Filtering/ObjectScopeHysteresisUpdater.cpp b/Engine/Source/Runtime/Experimental/Iris/Core/Private/Iris/ReplicationSystem/Filtering/ObjectScopeHysteresisUpdater.cpp
--- a/Engine/Source/Runtime/Experimental/Iris/Core/Private/Iris/ReplicationSystem/Filtering/ObjectScopeHysteresisUpdater.cpp
+++ b/Engine/Source/Runtime/Experimental/Iris/Core/Private/Iris/ReplicationSystem/Filtering/ObjectScopeHysteresisUpdater.cpp
@@ -2,7 +2,7 @@
 
 #include "Iris/ReplicationSystem/Filtering/ObjectScopeHysteresisUpdater.h"
 #include "Iris/Core/IrisProfiler.h"
-#include "Containers/Map.h"
+#include "Containers/ArrayView.h"
 
 namespace UE::Net::Private
 {
@@ -39,6 +39,11 @@ void FObjectScopeHysteresisUpdater::RemoveHysteresis(FInternalNetRefIndex NetRef
 void FObjectScopeHysteresisUpdater::RemoveHysteresis(const FNetBitArrayView& ObjectsToRemove)
 {
 	IRIS_PROFILER_SCOPE(FObjectScopeHysteresisUpdater_RemoveHysteresis);
+	if (NetRefIndexToLocalIndex.IsEmpty())
+	{
+		return;
+	}
+
 	FNetBitArrayView::ForAllSetBits(MakeNetBitArrayView(ObjectsToUpdate), ObjectsToRemove, FNetBitArrayView::AndOp, [this](uint32 NetRefIndex)
 		{
 			this->RemoveHysteresis(NetRefIndex);
@@ -46,6 +51,20 @@ void FObjectScopeHysteresisUpdater::RemoveHysteresis(const FNetBitArrayView& Obj
 	);
 }
 
+void FObjectScopeHysteresisUpdater::RemoveHysteresis(TArrayView<const uint32> ObjectsToRemove)
+{
+	IRIS_PROFILER_SCOPE(FObjectScopeHysteresisUpdater_RemoveHysteresis);
+	if (NetRefIndexToLocalIndex.IsEmpty())
+	{
+		return;
+	}
+
+	for (const uint32 ObjectIndex : ObjectsToRemove)
+	{
+		this->RemoveHysteresis(ObjectIndex);
+	}
+}
+
 void FObjectScopeHysteresisUpdater::Update(uint8 FramesSinceLastUpdate, TArray<FInternalNetRefIndex>& OutObjectsToFilterOut)
 {
 	IRIS_PROFILER_SCOPE(FObjectScopeHysteresisUpdater_Update);
diff --git a/Engine/Source/Runtime/Experimental/Iris/Core/Private/Iris/ReplicationSystem/Filtering/ObjectScopeHysteresisUpdater.h b/Engine/Source/Runtime/Experimental/Iris/Core/Private/Iris/ReplicationSystem/Filtering/ObjectScopeHysteresisUpdater.h
--- a/Engine/Source/Runtime/Experimental/Iris/Core/Private/Iris/ReplicationSystem/Filtering/ObjectScopeHysteresisUpdater.h
+++ b/Engine/Source/Runtime/Experimental/Iris/Core/Private/Iris/ReplicationSystem/Filtering/ObjectScopeHysteresisUpdater.h
@@ -31,16 +31,20 @@ public:
 	void RemoveHysteresis(const FNetBitArrayView& ObjectsToRemove);
 
 	/** Removes all objects in the array from hysteresis. */
-	void RemoveHysteresis(const TArrayView<const uint32>& ObjectsToRemove);
+	void RemoveHysteresis(TArrayView<const uint32> ObjectsToRemove);
 
 	/** Adjusts the connection relevant objects based on hysteresis frame counts. */
 	void Update(uint8 FramesSinceLastUpdate, TArray<FInternalNetRefIndex>& OutObjectsToFilterOut);
 
+	/** Whether any objects are updated for hysteresis. If not there's no point in calling Update(). */
 	bool HasObjectsToUpdate() const;
 
 	/** Returns the bitarray of objects affected by hysteresis */
 	FNetBitArrayView GetUpdatedObjects() const;
 
+	/** Returns true if the object is currently updated. */
+	bool IsObjectUpdated(FInternalNetRefIndex NetRefIndex) const;
+
 private:
 	enum : unsigned
 	{
@@ -76,4 +80,9 @@ inline FNetBitArrayView FObjectScopeHysteresisUpdater::GetUpdatedObjects() const
 	return MakeNetBitArrayView(ObjectsToUpdate);
 }
 
+inline bool FObjectScopeHysteresisUpdater::IsObjectUpdated(FInternalNetRefIndex ObjectIndex) const
+{
+	return ObjectsToUpdate.GetBit(ObjectIndex);
+}
+
 }
diff --git a/Engine/Source/Runtime/Experimental/Iris/Core/Private/Iris/ReplicationSystem/Filtering/ReplicationFiltering.cpp b/Engine/Source/Runtime/Experimental/Iris/Core/Private/Iris/ReplicationSystem/Filtering/ReplicationFiltering.cpp
--- a/Engine/Source/Runtime/Experimental/Iris/Core/Private/Iris/ReplicationSystem/Filtering/ReplicationFiltering.cpp
+++ b/Engine/Source/Runtime/Experimental/Iris/Core/Private/Iris/ReplicationSystem/Filtering/ReplicationFiltering.cpp
@@ -6,6 +6,7 @@
 
 #include "Iris/IrisConfigInternal.h"
 #include "Iris/IrisConstants.h"
+#include "Iris/Core/BitTwiddling.h"
 #include "Iris/Core/IrisCsv.h"
 #include "Iris/Core/IrisLog.h"
 #include "Iris/Core/IrisProfiler.h"
@@ -17,7 +18,7 @@
 #include "Iris/ReplicationSystem/DeltaCompression/DeltaCompressionBaselineInvalidationTracker.h"
 #include "Iris/ReplicationSystem/Filtering/NetObjectGroups.h"
 #include "Iris/ReplicationSystem/Filtering/NetObjectFilterDefinitions.h"
-
+#include "Iris/ReplicationSystem/Filtering/ReplicationFilteringConfig.h"
 #include <limits>
 
 namespace UE::Net::Private
@@ -187,6 +188,8 @@ void FReplicationFiltering::Init(FReplicationFilteringInitParams& Params)
 	check(Params.Connections != nullptr);
 	check(Params.Connections->GetMaxConnectionCount() <= std::numeric_limits<decltype(ObjectIndexToOwningConnection)::ElementType>::max());
 
+	Config = TStrongObjectPtr(GetDefault<UReplicationFilteringConfig>());
+
 	ReplicationSystem = Params.ReplicationSystem;
 
 	Connections = Params.Connections;
@@ -234,6 +237,9 @@ void FReplicationFiltering::Init(FReplicationFilteringInitParams& Params)
 	// Owning connections
 	ObjectIndexToOwningConnection.SetNumZeroed(MaxObjectCount);
 
+	// Object scope hysteresis
+	ObjectScopeHysteresisFrameCounts.SetNumZeroed(MaxObjectCount);
+
 	// Dynamic filters
 	{
 		NetObjectFilteringInfos.SetNumUninitialized(Params.MaxObjectCount);
@@ -245,6 +251,7 @@ void FReplicationFiltering::Init(FReplicationFilteringInitParams& Params)
 	}
 
 	InitFilters();
+	InitObjectScopeHysteresis();
 }
 
 void FReplicationFiltering::Filter()
@@ -252,6 +259,9 @@ void FReplicationFiltering::Filter()
 #if UE_NET_IRIS_CSV_STATS
 	CSV_SCOPED_TIMING_STAT(Iris, Filter_PrePoll);
 #endif
+
+	++FrameIndex;
+
 	ResetRemovedConnections();
 
 	InitNewConnections();
@@ -266,6 +276,8 @@ void FReplicationFiltering::Filter()
 
 	UpdateSubObjectFilters();
 
+	PreUpdateObjectScopeHysteresis();
+
 	if (HasDynamicFilters())
 	{
 		UpdateDynamicFilters();
@@ -426,6 +438,7 @@ bool FReplicationFiltering::SetFilter(FInternalNetRefIndex ObjectIndex, FNetObje
 			FilterInfo.Filter->GetFilteredObjects().SetBit(ObjIndex);
 			this->ObjectIndexToDynamicFilterIndex[ObjIndex] = static_cast<uint8>(FilterIndex);
 			this->DynamicFilterEnabledObjects.SetBit(ObjIndex);
+			this->ObjectScopeHysteresisFrameCounts[ObjIndex] = this->GetObjectScopeHysteresisFrameCount(FilterConfigProfile);
 			return true;
 		}
 		
@@ -605,11 +618,7 @@ void FReplicationFiltering::RemoveConnection(uint32 ConnectionId)
 	
 	// Reset connection info
 	FPerConnectionInfo& ConnectionInfo = ConnectionInfos[ConnectionId];
-	ConnectionInfo.ConnectionFilteredObjects.Empty();
-	ConnectionInfo.GroupExcludedObjects.Empty();
-	ConnectionInfo.ObjectsInScopeBeforeDynamicFiltering.Empty();
-	ConnectionInfo.GroupIncludedObjects.Empty();
-	ConnectionInfo.ObjectsInScope.Empty();
+	ConnectionInfo.Deinit();
 
 	for (FFilterInfo& Info : DynamicFilterInfos)
 	{
@@ -664,6 +673,8 @@ void FReplicationFiltering::InitNewConnections()
 		{
 			ConnectionInfo.DynamicFilteredOutObjects.Init(MaxObjectCount);
 			ConnectionInfo.InProgressDynamicFilteredOutObjects.Init(MaxObjectCount);
+			ConnectionInfo.DynamicFilteredOutObjectsHysteresisAdjusted.Init(MaxObjectCount);
+			ConnectionInfo.HysteresisUpdater.Init(MaxObjectCount);
 		}
 
 		// Update group exclusion filtering
@@ -868,7 +879,7 @@ void FReplicationFiltering::UpdateObjectsInScope()
 			uint32 DeletedObjects = (PrevExistingObjects & ~ExistingObjects);
 			for ( ; DeletedObjects; )
 			{
-				const uint32 LeastSignificantBit = DeletedObjects & uint32(-int32(DeletedObjects));
+				const uint32 LeastSignificantBit = GetLeastSignificantBit(DeletedObjects);
 				DeletedObjects ^= LeastSignificantBit;
 
 				const uint32 ObjectIndex = BitOffset + FPlatformMath::CountTrailingZeros(LeastSignificantBit);
@@ -892,7 +903,7 @@ void FReplicationFiltering::UpdateObjectsInScope()
 			{
 				for ( ; AddedSubObjects; )
 				{
-					const uint32 LeastSignificantBit = AddedSubObjects & uint32(-int32(AddedSubObjects));
+					const uint32 LeastSignificantBit = GetLeastSignificantBit(AddedSubObjects);
 					AddedSubObjects ^= LeastSignificantBit;
 
 					const uint32 ObjectIndex = BitOffset + FPlatformMath::CountTrailingZeros(LeastSignificantBit);
@@ -961,6 +972,7 @@ void FReplicationFiltering::UpdateObjectsInScope()
 		uint32* GroupIncludedObjectsStorage = ConnectionInfo.GroupIncludedObjects.GetData();
 		uint32* ObjectsInScopeBeforeDynamicFiltering = ConnectionInfo.ObjectsInScopeBeforeDynamicFiltering.GetData();
 		uint32* DynamicFilteredOutObjects = ConnectionInfo.DynamicFilteredOutObjects.GetData();
+		uint32* DynamicFilteredOutObjectsHysteresisAdjusted = ConnectionInfo.DynamicFilteredOutObjectsHysteresisAdjusted.GetData();
 
 		for (uint32 WordIt = 0, WordEndIt = ModifiedWordIndex; WordIt != WordEndIt; ++WordIt)
 		{
@@ -982,10 +994,11 @@ void FReplicationFiltering::UpdateObjectsInScope()
 			// Note that we only filter out objects from exclusion groups here. Inclusion groups only overrides dynamic filtering.
 			ObjectsInScopeBeforeDynamicFiltering[WordIndex] = FilteredObjectsWord & ~GroupExcludedObjectsWord;
 
-			// Make sure the we restore dynamic filtering for new objects
+			// Make sure to reset dynamic filtering for new objects. These can have been subobjects to other dynamically filtered objects.
 			if (HasDynamicFilters())
 			{
 				DynamicFilteredOutObjects[WordIndex] &= ~NewObjects;
+				DynamicFilteredOutObjectsHysteresisAdjusted[WordIndex] &= ~NewObjects;
 			}
 		}
 	}
@@ -1245,8 +1258,16 @@ void FReplicationFiltering::UpdateDynamicFiltering()
 {
 	IRIS_PROFILER_SCOPE(FReplicationFiltering_UpdateDynamicFiltering);
 
-	// Working set
-	TArray<uint32, TInlineAllocator<256>> DisabledDependentObjects;
+	constexpr uint32 TotalObjectCountOnStack = 1024;
+	constexpr uint32 FilteredOutDependentObjectCountOnStack = 192;
+	constexpr uint32 FilteredInObjectCountOnStack = 192;
+	constexpr uint32 FilteredOutObjectCountOnStack = TotalObjectCountOnStack - FilteredInObjectCountOnStack - FilteredOutDependentObjectCountOnStack;
+
+	TArray<uint32, TInlineAllocator<FilteredOutObjectCountOnStack>> FilteredOutObjects;
+	TArray<uint32, TInlineAllocator<FilteredInObjectCountOnStack>> FilteredInObjects;
+	TArray<uint32, TInlineAllocator<FilteredOutDependentObjectCountOnStack>> FilteredOutDependentObjects;
+	TArray<uint32> FilteredOutByHysteresisObjects;
+	FilteredOutByHysteresisObjects.Reserve(256);
 
 	uint32* AllowedObjectsData = static_cast<uint32*>(FMemory_Alloca(WordCountForObjectBitArrays * sizeof(uint32)));
 	FNetBitArrayView AllowedObjects(AllowedObjectsData, MaxObjectCount, FNetBitArrayView::NoResetNoValidate);
@@ -1267,11 +1288,10 @@ void FReplicationFiltering::UpdateDynamicFiltering()
 
 		ConnectionInfo.InProgressDynamicFilteredOutObjects.Reset();
 		uint32* InProgressDynamicFilteredOutObjectsData = ConnectionInfo.InProgressDynamicFilteredOutObjects.GetData();
-		FNetBitArrayView InProgressDynamicFilteredOutObjects = MakeNetBitArrayView(ConnectionInfo.InProgressDynamicFilteredOutObjects);
 
 		/*
 		 * Apply dynamic filters.
-		 * 
+		 *
 		 * The algorithm will loop over all connections, call each filter and update
 		 * the DynamicFilteredOutObjects bit array. Once all filters have been applied
 		 * the result is compared with that of the previous frame.
@@ -1311,25 +1331,31 @@ void FReplicationFiltering::UpdateDynamicFiltering()
 		 * and try only to do the more expensive operations, such as subobject management,
 		 * for objects that have changed filter status since the previous frame.
 		 */
+		FNetBitArrayView DynamicFilteredOutObjects = MakeNetBitArrayView(ConnectionInfo.DynamicFilteredOutObjects);
+		FNetBitArrayView DynamicFilteredOutObjectsHysteresisAdjusted = MakeNetBitArrayView(ConnectionInfo.DynamicFilteredOutObjectsHysteresisAdjusted);
 		{
 			IRIS_PROFILER_SCOPE(FReplicationFiltering_OnFilterStatusChanged);
-			DisabledDependentObjects.Reset();
+			FilteredOutObjects.Reset();
+			FilteredOutDependentObjects.Reset();
+			FilteredInObjects.Reset();
 
 			const uint32* DynamicFilterEnabledObjectsData = DynamicFilterEnabledObjects.GetData();
-			FNetBitArrayView DynamicFilteredOutObjects = MakeNetBitArrayView(ConnectionInfo.DynamicFilteredOutObjects);
 			uint32* DynamicFilteredOutObjectsData = ConnectionInfo.DynamicFilteredOutObjects.GetData();
+			const uint32* GroupIncludedObjectsData = ConnectionInfo.GroupIncludedObjects.GetData();
 			for (uint32 WordIt = 0, WordEndIt = WordCountForObjectBitArrays; WordIt != WordEndIt; ++WordIt)
 			{
 				const uint32 SubObjects = SubObjectsData[WordIt];
 				const uint32 DependentObjects = DependentObjectsData[WordIt];
 				const uint32 ObjectsRequiringUpdate = ObjectsRequiringDynamicFilterUpdateData[WordIt];
 				const uint32 PrevFilteredOutObjects = DynamicFilteredOutObjectsData[WordIt];
-				const uint32 NewFilteredOutObjects = InProgressDynamicFilteredOutObjectsData[WordIt];
+				// Mask off group included objects in this loops so they're accounted for when updating hysteresis. A no longer group included object should also be subject to hysteresis.
+				const uint32 CurrentFilteredOutObjects = InProgressDynamicFilteredOutObjectsData[WordIt] & ~GroupIncludedObjectsData[WordIt];
 				const uint32 FilterEnabledObjects = DynamicFilterEnabledObjectsData[WordIt];
 
-				const uint32 ModifiedScopeObjects = (PrevFilteredOutObjects ^ NewFilteredOutObjects);
+				const uint32 ModifiedScopeObjects = (PrevFilteredOutObjects ^ CurrentFilteredOutObjects);
 				// ObjectsRequiringUpdate may contain objects that had a dynamic filter set last frame. We need to update
-				// the filtered out objects.
+				// the filtered out objects. We process dependencies every frame to deal with cases where they're filtered out and they should no
+				// longer be allowed to replicate due to the object with the dependency being filtered out to.
 				const uint32 ObjectsToProcess = ((ModifiedScopeObjects | DependentObjects) & FilterEnabledObjects) | ObjectsRequiringUpdate;
 
 				if (!ObjectsToProcess)
@@ -1337,92 +1363,248 @@ void FReplicationFiltering::UpdateDynamicFiltering()
 					continue;
 				}
 
-				// Update dynamic filter, but preserve subobject information.
-				DynamicFilteredOutObjectsData[WordIt] = (PrevFilteredOutObjects & SubObjects) | NewFilteredOutObjects;
+				// Update dynamic filter
+				DynamicFilteredOutObjectsData[WordIt] = CurrentFilteredOutObjects;
 
-				const uint32 BitOffset = WordIt*32U;
+				const uint32 BitOffset = WordIt * 32U;
 
 				// Calculate which objects need to be updated due to being filtered out.
 				// We want to process as few as possible since modifying subobject status is expensive.
-				for (uint32 DisabledObjects = NewFilteredOutObjects & ObjectsToProcess; DisabledObjects; )
+				for (uint32 DisabledObjects = CurrentFilteredOutObjects & ObjectsToProcess; DisabledObjects; )
 				{
 					// Extract one set bit from DisabledObjects. The code below will extract the least significant bit set and then clear it.
-					const uint32 LeastSignificantBit = DisabledObjects & uint32(-int32(DisabledObjects));
+					const uint32 LeastSignificantBit = GetLeastSignificantBit(DisabledObjects);
 					DisabledObjects ^= LeastSignificantBit;
 
+					// Store disabled objects for later processing to minimize performance impact of hysteresis
 					const uint32 ObjectIndex = BitOffset + FPlatformMath::CountTrailingZeros(LeastSignificantBit);
-					for (const FInternalNetRefIndex SubObjectIndex : NetRefHandleManager->GetSubObjects(ObjectIndex))
-					{						
-						DynamicFilteredOutObjects.SetBit(SubObjectIndex);
-					}
 
-					// Store dependent objects for later processing.
 					if (DependentObjects & LeastSignificantBit)
 					{
-						DisabledDependentObjects.Add(ObjectIndex);
+						// Store dependent objects for later processing.
+						FilteredOutDependentObjects.Add(ObjectIndex);
+					}
+					else
+					{
+						// Store objects for later processing
+						FilteredOutObjects.Add(ObjectIndex);
 					}
 				}
 
-				for (uint32 EnabledObjects = ~NewFilteredOutObjects & ObjectsToProcess; EnabledObjects; )
+				for (uint32 EnabledObjects = ~CurrentFilteredOutObjects & ObjectsToProcess; EnabledObjects; )
 				{
-					const uint32 LeastSignificantBit = EnabledObjects & uint32(-int32(EnabledObjects));
+					const uint32 LeastSignificantBit = GetLeastSignificantBit(EnabledObjects);
 					EnabledObjects ^= LeastSignificantBit;
 
 					const uint32 ObjectIndex = BitOffset + FPlatformMath::CountTrailingZeros(LeastSignificantBit);
+					FilteredInObjects.Add(ObjectIndex);
+				}
+			}
+		}
+
+		// Hysteresis path need to add explicitly filtered out objects to hysteresis. Objects not supported by hysteresis as well as objects' whose timeout has passed needs to be filtered out immediately.
+		const uint32 ConnIdMod = ConnId % HysteresisState.ConnectionIdStride;
+		if (HysteresisState.Mode == EHysteresisProcessingMode::Enabled)
+		{
+			IRIS_CSV_PROFILER_SCOPE(Iris, FReplicationFiltering_OnFilterStatusChanged_Hysteresis);
+
+			// With connection throttling we need to adjust the hysteresis with the difference between the update rate and the number of frames til next update.
+			// If we're updating this frame we need to add HysteresisState.ConnectionIdStride to compensate for an extra update of that amount, if we're updating the next frame we need to add HysteresisState.ConnectionIdStride - 1 and so on.
+			const uint32 AdjustHysteresisForUpdateThrottling = HysteresisState.ConnectionIdStride - ((ConnIdMod + HysteresisState.ConnectionIdStride - HysteresisState.ConnectionStartId) % HysteresisState.ConnectionIdStride);
+
+			// Remove filtered in objects from hysteresis immediately. 
+			ConnectionInfo.HysteresisUpdater.RemoveHysteresis(MakeArrayView(FilteredInObjects));
+
+			// Make sure filtered in objects aren't filtered out.
+			for (const uint32 ObjectIndex : FilteredInObjects)
+			{
+				DynamicFilteredOutObjectsHysteresisAdjusted.ClearBit(ObjectIndex);
+				for (const FInternalNetRefIndex SubObjectIndex : NetRefHandleManager->GetSubObjects(ObjectIndex))
+				{
+					DynamicFilteredOutObjectsHysteresisAdjusted.ClearBit(SubObjectIndex);
+				}
+			}
+
+			// Add filtered out objects to hysteresis if eligible, meaning they support it and are still dynamically filtered. Dependent objects are not processed here.
+			for (const uint32 ObjectIndex : FilteredOutObjects)
+			{
+				const uint32 HysteresisFrameCount = ObjectScopeHysteresisFrameCounts[ObjectIndex];
+				if (HysteresisFrameCount && DynamicFilterEnabledObjects.GetBit(ObjectIndex))
+				{
+					// We need to adjust the hysteresis frame count to account for when it will be updated. The -1 stems from the fact that updating won't happen until next frame at the earliest.
+					const uint16 TotalHysteresisFrameCount = static_cast<uint16>(HysteresisFrameCount - 1 + AdjustHysteresisForUpdateThrottling);
+					ConnectionInfo.HysteresisUpdater.SetHysteresisFrameCount(ObjectIndex, TotalHysteresisFrameCount);
+				}
+				else
+				{
+					// Filter out
+					DynamicFilteredOutObjectsHysteresisAdjusted.SetBit(ObjectIndex);
+					for (const FInternalNetRefIndex SubObjectIndex : NetRefHandleManager->GetSubObjects(ObjectIndex))
+					{
+						DynamicFilteredOutObjectsHysteresisAdjusted.SetBit(SubObjectIndex);
+					}
+				}
+			}
+
+			// Update hysteresis
+			if (ConnIdMod == HysteresisState.ConnectionStartId)
+			{
+				FilteredOutByHysteresisObjects.Reset();
+				ConnectionInfo.HysteresisUpdater.Update(static_cast<uint8>(HysteresisState.ConnectionIdStride), FilteredOutByHysteresisObjects);
+
+				// Immediately filter out objects whose hysteresis timed out.
+				for (const uint32 ObjectIndex : FilteredOutByHysteresisObjects)
+				{
+					DynamicFilteredOutObjectsHysteresisAdjusted.SetBit(ObjectIndex);
 					for (const FInternalNetRefIndex SubObjectIndex : NetRefHandleManager->GetSubObjects(ObjectIndex))
 					{
-						DynamicFilteredOutObjects.ClearBit(SubObjectIndex);
+						DynamicFilteredOutObjectsHysteresisAdjusted.SetBit(SubObjectIndex);
 					}
 				}
 			}
+
+		}
+		else
+		{
+			IRIS_CSV_PROFILER_SCOPE(Iris, FReplicationFiltering_OnFilterStatusChanged_NonHysteresis);
+
+			// Make sure filtered in objects aren't filtered out.
+			for (const uint32 ObjectIndex : FilteredInObjects)
+			{
+				DynamicFilteredOutObjectsHysteresisAdjusted.ClearBit(ObjectIndex);
+				for (const FInternalNetRefIndex SubObjectIndex : NetRefHandleManager->GetSubObjects(ObjectIndex))
+				{
+					DynamicFilteredOutObjectsHysteresisAdjusted.ClearBit(SubObjectIndex);
+				}
+			}
+
+			// Filter out objects except dependent ones.
+			for (const uint32 ObjectIndex : FilteredOutObjects)
+			{
+				DynamicFilteredOutObjectsHysteresisAdjusted.SetBit(ObjectIndex);
+				for (const FInternalNetRefIndex SubObjectIndex : NetRefHandleManager->GetSubObjects(ObjectIndex))
+				{
+					DynamicFilteredOutObjectsHysteresisAdjusted.SetBit(SubObjectIndex);
+				}
+			}
 		}
 
-		// Update the entire scope for the connection
+		// Update the entire scope for the connection.
 		{
 			IRIS_PROFILER_SCOPE(FReplicationFiltering_UpdateConnectionScope);
 
 			uint32* ObjectsInScopeData = ConnectionInfo.ObjectsInScope.GetData();
 			const uint32* ObjectsInScopeBeforeDynamicFilteringData = ConnectionInfo.ObjectsInScopeBeforeDynamicFiltering.GetData();
-			const uint32* DynamicFilteredOutObjectsData = ConnectionInfo.DynamicFilteredOutObjects.GetData();
-			const uint32* GroupIncludedObjectsData = ConnectionInfo.GroupIncludedObjects.GetData();
-			// $IRIS TODO Vectorization opportunity
+			const uint32* DynamicFilteredOutObjectsHysteresisAdjustedData = DynamicFilteredOutObjectsHysteresisAdjusted.GetData();
 			for (SIZE_T WordIt = 0, WordEndIt = WordCountForObjectBitArrays; WordIt != WordEndIt; ++WordIt)
 			{
 				const uint32 ObjectsInScopeBeforeWord = ObjectsInScopeBeforeDynamicFilteringData[WordIt];
-				const uint32 DynamicFilteredOutWord = DynamicFilteredOutObjectsData[WordIt];
-				const uint32 GroupIncludedWord = GroupIncludedObjectsData[WordIt];
-				ObjectsInScopeData[WordIt] = ObjectsInScopeBeforeWord & (GroupIncludedWord | ~DynamicFilteredOutWord);
+				const uint32 DynamicFilteredOutWord = DynamicFilteredOutObjectsHysteresisAdjustedData[WordIt];
+				ObjectsInScopeData[WordIt] = ObjectsInScopeBeforeWord & ~DynamicFilteredOutWord;
+			}
+
+			// Unconditionally filter out filtered out dependent objects from ObjectsInScope as GetDependentObjectFilterStatus, called later, requires it. Subobjects are dealt with later.
+			for (const uint32 DependentObjectIndex : FilteredOutDependentObjects)
+			{
+				ConnectionInfo.ObjectsInScope.ClearBit(DependentObjectIndex);
 			}
 		}
 
 		// The scope for the connection is now fully updated, apart from disabled dependent objects.
-		// If any object that has a dependency that isn't filtered out we must re-enable the dependent object.
+		// If any object that has a dependency that isn't filtered out we must enable the dependent object.
+		if (HysteresisState.Mode == EHysteresisProcessingMode::Enabled)
 		{
-			IRIS_PROFILER_SCOPE(FReplicationFiltering_UpdateDependentObjects);
-			for (const uint32 DependentObjectIndex : DisabledDependentObjects)
+			IRIS_CSV_PROFILER_SCOPE(Iris, FReplicationFiltering_DependentObjects_Hysteresis);
+
+			// We didn't filter out any dependent objects yet so we enable/disable as needed.
+			// Different algorithm for figuring out update throttling adjustment since we already updated this frame.
+			const uint32 AdjustHysteresisForUpdateThrottling = (HysteresisState.ConnectionStartId + HysteresisState.ConnectionIdStride - ConnIdMod) % HysteresisState.ConnectionIdStride;
+			for (const uint32 DependentObjectIndex : FilteredOutDependentObjects)
 			{
-				if (GetDependentObjectFilterStatus(NetRefHandleManager, ConnectionInfo.ObjectsInScope, DependentObjectIndex) == ENetFilterStatus::Allow)
+				const bool bAllowReplication = GetDependentObjectFilterStatus(NetRefHandleManager, ConnectionInfo.ObjectsInScope, DependentObjectIndex) == ENetFilterStatus::Allow;
+				if (bAllowReplication)
 				{
+					ConnectionInfo.HysteresisUpdater.RemoveHysteresis(DependentObjectIndex);
+					// We use the status of the hysteresis adjusted dynamic filtering to figure out whether hysteresis should be enabled or not. ObjectsInScope have already been updated and can't be used for this purpose.
+					DynamicFilteredOutObjectsHysteresisAdjusted.ClearBit(DependentObjectIndex);
+					
 					ConnectionInfo.ObjectsInScope.SetBit(DependentObjectIndex);
 					for (const FInternalNetRefIndex SubObjectIndex : NetRefHandleManager->GetSubObjects(DependentObjectIndex))
 					{
+						DynamicFilteredOutObjectsHysteresisAdjusted.ClearBit(SubObjectIndex);
 						const bool bIsInScope = ConnectionInfo.ObjectsInScopeBeforeDynamicFiltering.GetBit(SubObjectIndex);
 						ConnectionInfo.ObjectsInScope.SetBitValue(SubObjectIndex, bIsInScope);
 					}
 				}
+				else
+				{
+					// Start/continue hysteresis or time to prevent replication?
+					bool bIsFilteredOut = DynamicFilteredOutObjectsHysteresisAdjusted.GetBit(DependentObjectIndex);
+					if (!bIsFilteredOut)
+					{
+						// If we're updated by hysteresis at this point we should not start again but otherwise let's start!
+						if (!ConnectionInfo.HysteresisUpdater.IsObjectUpdated(DependentObjectIndex))
+						{
+							if (const uint32 HysteresisFrameCount = ObjectScopeHysteresisFrameCounts[DependentObjectIndex])
+							{
+								const uint16 TotalHysteresisFrameCount = static_cast<uint16>(HysteresisFrameCount - 1 + AdjustHysteresisForUpdateThrottling);
+								ConnectionInfo.HysteresisUpdater.SetHysteresisFrameCount(DependentObjectIndex, TotalHysteresisFrameCount);
+							}
+							else
+							{
+								// Object doesn't support hysteresis so we need to filter it out immediately.
+								bIsFilteredOut = true;
+							}
+						}
+						// Else let hysteresis updating continue. Once fully processed the DynamicFilteredOutObjectsHysteresisAdjusted bitarray will be adjsuted accordingly and filter out the object.
+
+						// Explicitly enable dependent object and its subobjects.
+						ConnectionInfo.ObjectsInScope.SetBitValue(DependentObjectIndex, !bIsFilteredOut);
+						for (const FInternalNetRefIndex SubObjectIndex : NetRefHandleManager->GetSubObjects(DependentObjectIndex))
+						{
+							DynamicFilteredOutObjectsHysteresisAdjusted.SetBitValue(SubObjectIndex, bIsFilteredOut);
+							const bool bIsInScope = ConnectionInfo.ObjectsInScopeBeforeDynamicFiltering.GetBit(SubObjectIndex);
+							ConnectionInfo.ObjectsInScope.SetBitValue(SubObjectIndex, !bIsFilteredOut && bIsInScope);
+						}
+					}
+					else
+					{
+						// Explicitly filter out dependent object and its subobjects
+						DynamicFilteredOutObjectsHysteresisAdjusted.SetBit(DependentObjectIndex);
+						ConnectionInfo.ObjectsInScope.ClearBit(DependentObjectIndex);
+						for (const FInternalNetRefIndex SubObjectIndex : NetRefHandleManager->GetSubObjects(DependentObjectIndex))
+						{
+							DynamicFilteredOutObjectsHysteresisAdjusted.SetBit(SubObjectIndex);
+							ConnectionInfo.ObjectsInScope.ClearBit(SubObjectIndex);
+						}
+					}
+				}
+			}
+		}
+		else
+		{
+			IRIS_CSV_PROFILER_SCOPE(Iris, FReplicationFiltering_DependentObjects_Standard);
+			// We didn't filter out any dependent objects yet so we enable/disable as needed.
+			for (const uint32 DependentObjectIndex : FilteredOutDependentObjects)
+			{
+				const bool bAllowReplication = GetDependentObjectFilterStatus(NetRefHandleManager, ConnectionInfo.ObjectsInScope, DependentObjectIndex) == ENetFilterStatus::Allow;
+				ConnectionInfo.ObjectsInScope.SetBitValue(DependentObjectIndex, bAllowReplication);
+				for (const FInternalNetRefIndex SubObjectIndex : NetRefHandleManager->GetSubObjects(DependentObjectIndex))
+				{
+					const bool bIsInScope = ConnectionInfo.ObjectsInScopeBeforeDynamicFiltering.GetBit(SubObjectIndex);
+					ConnectionInfo.ObjectsInScope.SetBitValue(SubObjectIndex, bIsInScope && bAllowReplication);
+				}
 			}
 		}
 	}
-
-
-	// It's ok to reset this in the first filter pass since objects in it would be fully treated.
-	ObjectsRequiringDynamicFilterUpdate.Reset();
 }
 
 void FReplicationFiltering::PostUpdateDynamicFiltering()
 {
 	IRIS_PROFILER_SCOPE(FReplicationFiltering_PostUpdateDynamicFiltering);
 
+	ObjectsRequiringDynamicFilterUpdate.Reset();
+
 	// Tell filters to clean up after filtering. It's only called if any object has the filter set.
 	{
 		FNetObjectPostFilteringParams PostFilteringParams;
@@ -1436,6 +1618,8 @@ void FReplicationFiltering::PostUpdateDynamicFiltering()
 			Info.Filter->PostFilter(PostFilteringParams);
 		}
 	}
+
+	PostUpdateObjectScopeHysteresis();
 }
 
 void FReplicationFiltering::NotifyFiltersOfDirtyObjects()
@@ -2063,6 +2247,18 @@ bool FReplicationFiltering::ClearGroupInclusionFilterEffectsForObject(uint32 Obj
 		
 		GroupIncludedObjects.ClearBit(ObjectIndex);
 
+		// Dynamically filtered objects are subject to hysteresis.
+		if (HysteresisState.Mode == EHysteresisProcessingMode::Enabled)
+		{
+			const uint8 HysteresisFrameCount = DynamicFilterEnabledObjects.GetBit(ObjectIndex) ? ObjectScopeHysteresisFrameCounts[ObjectIndex] : uint8(0);
+			if (HysteresisFrameCount)
+			{
+				const uint32 ConnIdMod = ConnectionId % HysteresisState.ConnectionIdStride;
+				const uint32 AdjustHysteresisForUpdateThrottling = (HysteresisState.ConnectionStartId + HysteresisState.ConnectionIdStride - ConnIdMod) % HysteresisState.ConnectionIdStride;
+				ConnectionInfo.HysteresisUpdater.SetHysteresisFrameCount(ObjectIndex, static_cast<uint16>(HysteresisFrameCount + AdjustHysteresisForUpdateThrottling));
+			}
+		}
+
 		for (const FInternalNetRefIndex SubObjectIndex : NetRefHandleManager->GetSubObjects(ObjectIndex))
 		{
 			GroupIncludedObjects.ClearBit(SubObjectIndex);
@@ -2412,6 +2608,13 @@ void FReplicationFiltering::InitFilters()
 	}
 }
 
+void FReplicationFiltering::InitObjectScopeHysteresis()
+{
+	HysteresisState.Mode = Config->IsObjectScopeHysteresisEnabled() ? EHysteresisProcessingMode::Enabled : EHysteresisProcessingMode::Disabled;
+	// Always allocated and maintained regardless of whether the feature is enabled or not.
+	HysteresisState.ObjectsToClear.SetNumBits(MaxObjectCount);
+}
+
 void FReplicationFiltering::RemoveFromDynamicFilter(uint32 ObjectIndex, uint32 FilterIndex)
 {
 	UE_LOG(LogIrisFiltering, Verbose, TEXT("RemoveFromDynamicFilter removing %s from Dynamic Filter %s"), 
@@ -2426,6 +2629,9 @@ void FReplicationFiltering::RemoveFromDynamicFilter(uint32 ObjectIndex, uint32 F
 
 	DynamicFilterEnabledObjects.ClearBit(ObjectIndex);
 	ObjectsRequiringDynamicFilterUpdate.SetBit(ObjectIndex);
+
+	// Remove from hysteresis
+	HysteresisState.ClearFromHysteresis(ObjectIndex);
 }
 
 void FReplicationFiltering::InvalidateBaselinesForObject(uint32 ObjectIndex, uint32 NewOwningConnectionId, uint32 PrevOwningConnectionId)
@@ -2509,6 +2715,73 @@ void FReplicationFiltering::BuildObjectsInFilterList(FNetBitArrayView OutObjects
 	}
 }
 
+
+void FReplicationFiltering::PreUpdateObjectScopeHysteresis()
+{
+	if (HysteresisState.Mode == EHysteresisProcessingMode::Enabled)
+	{
+		HysteresisState.ConnectionStartId = (FrameIndex % Config->GetHysteresisUpdateConnectionThrottling());
+		HysteresisState.ConnectionIdStride = Config->GetHysteresisUpdateConnectionThrottling();
+	}
+
+	ClearObjectsFromHysteresis();
+}
+
+void FReplicationFiltering::PostUpdateObjectScopeHysteresis()
+{
+	HysteresisState.ObjectsToClearCount = 0;
+	HysteresisState.ObjectsToClear.Reset();
+}
+
+void FReplicationFiltering::ClearObjectsFromHysteresis()
+{
+	if (!HysteresisState.ObjectsToClearCount)
+	{
+		return;
+	}
+
+	IRIS_PROFILER_SCOPE(FObjectScopeHysteresisUpdater_ClearObjectsFromHysteresis);
+
+	// $IRIS TODO: Use more optimal path for low counts of objects to clear, for example passing an ArrayView instead.
+	ValidConnections.ForAllSetBits([this](uint32 ConnectionId)
+		{
+			FPerConnectionInfo& ConnectionInfo = this->ConnectionInfos[ConnectionId];
+			ConnectionInfo.HysteresisUpdater.RemoveHysteresis(MakeNetBitArrayView(this->HysteresisState.ObjectsToClear));
+		});
+}
+
+uint8 FReplicationFiltering::GetObjectScopeHysteresisFrameCount(FName ProfileName) const
+{
+	if (const FObjectScopeHysteresisProfile* Profile = Config->GetHysteresisProfiles().FindByKey(ProfileName))
+	{
+		return Profile->HysteresisFrameCount;
+	}
+
+	return Config->GetDefaultHysteresisFrameCount();
+}
+
+void FReplicationFiltering::FPerConnectionInfo::Deinit()
+{
+	ConnectionFilteredObjects.Empty();
+	GroupExcludedObjects.Empty();
+	ObjectsInScopeBeforeDynamicFiltering.Empty();
+	GroupIncludedObjects.Empty();
+	ObjectsInScope.Empty();
+	DynamicFilteredOutObjects.Empty();
+	InProgressDynamicFilteredOutObjects.Empty();
+	DynamicFilteredOutObjectsHysteresisAdjusted.Empty();
+	HysteresisUpdater.Deinit();
+}
+
+//*************************************************************************************************
+// FObjectScopeHysteresisState
+//*************************************************************************************************
+void FReplicationFiltering::FObjectScopeHysteresisState::ClearFromHysteresis(FInternalNetRefIndex NetRefIndex)
+{
+	ObjectsToClear.SetBit(NetRefIndex);
+	++ObjectsToClearCount;
+}
+
 //*************************************************************************************************
 // FNetObjectFilteringInfoAccessor
 //*************************************************************************************************
diff --git a/Engine/Source/Runtime/Experimental/Iris/Core/Private/Iris/ReplicationSystem/Filtering/ReplicationFiltering.h b/Engine/Source/Runtime/Experimental/Iris/Core/Private/Iris/ReplicationSystem/Filtering/ReplicationFiltering.h
--- a/Engine/Source/Runtime/Experimental/Iris/Core/Private/Iris/ReplicationSystem/Filtering/ReplicationFiltering.h
+++ b/Engine/Source/Runtime/Experimental/Iris/Core/Private/Iris/ReplicationSystem/Filtering/ReplicationFiltering.h
@@ -5,10 +5,13 @@
 #include "CoreTypes.h"
 #include "Net/Core/NetBitArray.h"
 #include "Iris/ReplicationSystem/Filtering/NetObjectFilter.h"
+#include "Iris/ReplicationSystem/Filtering/ObjectScopeHysteresisUpdater.h"
+#include "Iris/ReplicationSystem/Filtering/ReplicationFilteringConfig.h"
 #include "Iris/ReplicationSystem/NetObjectGroupHandle.h"
 #include "Containers/Array.h"
 #include "UObject/StrongObjectPtr.h"
 
+class UReplicationFilteringConfig;
 class UReplicationSystem;
 namespace UE::Net
 {
@@ -146,6 +149,8 @@ public:
 private:
 	struct FPerConnectionInfo
 	{
+		void Deinit();
+
 		// Objects filtered depending on owning connection or user set connection filtering
 		FNetBitArray ConnectionFilteredObjects;
 		// Objects filtered out due to one or more exclusion groups it belongs to is filtered out
@@ -156,12 +161,14 @@ private:
 		FNetBitArray GroupIncludedObjects;
 		// Objects in scope after all kinds of filtering, including dynamic filtering, has been applied
 		FNetBitArray ObjectsInScope;
-
-		// Which objects are filtered out after dynamic filters and inclusion groups have been processed.
+		// Which objects are filtered out after dynamic filters have been processed.
 		FNetBitArray DynamicFilteredOutObjects;
-
-		// List of objects currently filtered out after processing dynamic filter passes and inclusion groups
+		// List of objects currently filtered out after processing dynamic filter passes. This could be temporary allocations in UpdateDynamicFiltering() but does require one bitarray per connection.
 		FNetBitArray InProgressDynamicFilteredOutObjects;
+		// Which objects are filtered out after dynamic filters, inclusion groups and hysteresis have been processed.
+		FNetBitArray DynamicFilteredOutObjectsHysteresisAdjusted;
+		// Updater of hysteresis for objects being dynamically filtered out
+		FObjectScopeHysteresisUpdater HysteresisUpdater;
 	};
 
 	struct FPerObjectInfo
@@ -192,6 +199,7 @@ private:
 	static void StaticChecks();
 
 	void InitFilters();
+	void InitObjectScopeHysteresis();
 
 	void InitNewConnections();
 	void ResetRemovedConnections();
@@ -206,6 +214,10 @@ private:
 	void UpdateDynamicFiltering();
 	void PostUpdateDynamicFiltering();
 
+	void PreUpdateObjectScopeHysteresis();
+	void PostUpdateObjectScopeHysteresis();
+	void ClearObjectsFromHysteresis();
+
 	/** Build the list of always relevant objects + objects that are currently relevant to at least one connection. */
 	void FilterNonRelevantObjects();
 
@@ -253,13 +265,45 @@ private:
 	/** Returns all the filtering infos. */
 	TArrayView<FNetObjectFilteringInfo> GetNetObjectFilteringInfos();
 
+	uint8 GetObjectScopeHysteresisFrameCount(FName Profile) const;
+
 private:
+	enum EHysteresisProcessingMode : uint32
+	{
+		Disabled,
+		Enabled,
+	};
+
+	// Scope hysteresis state. Hysteresis is applied to objects going out of scope for objects that so desire.
+	struct FObjectScopeHysteresisState
+	{
+	public:
+		void ClearFromHysteresis(FInternalNetRefIndex NetRefIndex);
+
+		// Processing mode
+		EHysteresisProcessingMode Mode = EHysteresisProcessingMode::Disabled;
+		// Which connection ID to start with for updating.
+		uint32 ConnectionStartId = 0;
+		// Stride for connection update throttling.
+		uint32 ConnectionIdStride = 1;
+
+		// Approximate number of objects that should be cleared from hysteresis.
+		uint32 ObjectsToClearCount = 0;
+
+		// Objects to clear from hysteresis due to being destroyed or removed from dynamic filtering.
+		FNetBitArray ObjectsToClear;
+	};
+
 	// Used for ObjectIndexToDynamicFilterIndex lookup
 	static constexpr uint8 InvalidDynamicFilterIndex = 255U;
 
+	// Config
+	TStrongObjectPtr<const UReplicationFilteringConfig> Config;
+
 	// General
 	TObjectPtr<UReplicationSystem> ReplicationSystem = nullptr;
 	const FNetRefHandleManager* NetRefHandleManager = nullptr;
+	uint32 FrameIndex = 0;
 
 	// Groups
 	FNetObjectGroups* Groups = nullptr;
@@ -296,6 +340,9 @@ private:
 	TArray<FPerGroupInfo> GroupInfos;
 	uint32 MaxGroupCount = 0;
 
+	// Hysteresis frame counts for dynamically filtered objects
+	TArray<uint8> ObjectScopeHysteresisFrameCounts;
+
 	/** NetObjectGroups used for filtering out objects. */
 	FNetBitArray ExclusionFilterGroups;
 
@@ -326,6 +373,9 @@ private:
 	FNetBitArray DynamicFilterEnabledObjects;
 	FNetBitArray ObjectsRequiringDynamicFilterUpdate;
 
+	// Object scope hystereris
+	FObjectScopeHysteresisState HysteresisState;
+
 	uint32 bHasNewConnection : 1;
 	uint32 bHasRemovedConnection : 1;
 	uint32 bHasDirtyConnectionFilter: 1;
