Iris - Fixed GC error during test.

diff --git a/Engine/Plugins/Runtime/ReplicationSystemTestPlugin/Source/Private/Tests/ReplicationSystem/TestReplicationOperations.cpp b/Engine/Plugins/Runtime/ReplicationSystemTestPlugin/Source/Private/Tests/ReplicationSystem/TestReplicationOperations.cpp
--- a/Engine/Plugins/Runtime/ReplicationSystemTestPlugin/Source/Private/Tests/ReplicationSystem/TestReplicationOperations.cpp
+++ b/Engine/Plugins/Runtime/ReplicationSystemTestPlugin/Source/Private/Tests/ReplicationSystem/TestReplicationOperations.cpp
@@ -594,7 +594,8 @@ UE_NET_TEST_FIXTURE(FTestReplicationOperationsForObjectsFixture, StaleObjectPoin
 	// Destroy first object and invalidate references to it.
 	ReplicationBridge->EndReplication(HandleToObject0);
 	DestroyObject(Object0);
-	CollectGarbage(RF_NoFlags);
+	constexpr bool bPerformFullPurge = false;
+	CollectGarbage(GARBAGE_COLLECTION_KEEPFLAGS, bPerformFullPurge);
 
 	ReplicationSystem->PreSendUpdate(0.033f);
 	FReplicationInstanceOperations::Quantize(SerializationContext, StateBuffers1[1], GetChangeMaskWriter(Protocol1->ChangeMaskBitCount), InstanceProtocol1, Protocol1);
