Iris
* Added Unit test for Growing netbitarray operation

#rnx

diff --git a/Engine/Plugins/Runtime/ReplicationSystemTestPlugin/Source/Private/Tests/ReplicationSystem/ReplicationSystemServerClientTestFixture.cpp b/Engine/Plugins/Runtime/ReplicationSystemTestPlugin/Source/Private/Tests/ReplicationSystem/ReplicationSystemServerClientTestFixture.cpp
--- a/Engine/Plugins/Runtime/ReplicationSystemTestPlugin/Source/Private/Tests/ReplicationSystem/ReplicationSystemServerClientTestFixture.cpp
+++ b/Engine/Plugins/Runtime/ReplicationSystemTestPlugin/Source/Private/Tests/ReplicationSystem/ReplicationSystemServerClientTestFixture.cpp
@@ -64,6 +64,16 @@ void FDataStreamTestUtil::AddDataStreamDefinition(const TCHAR* StreamName, const
 
 // FReplicationSystemTestNode implementation
 FReplicationSystemTestNode::FReplicationSystemTestNode(bool bIsServer, const TCHAR* Name)
+{
+	Setup(bIsServer, Name);
+}
+
+FReplicationSystemTestNode::FReplicationSystemTestNode(FReplicationSystemTestNode::EDelaySetup)
+{
+	// Do nothing until asked to start up
+}
+
+void FReplicationSystemTestNode::Setup(bool bIsServer, const TCHAR* Name, FReplicationSystemTestNode::FReplicationSystemParamsOverride* ParamsOverride)
 {
 	ReplicationBridge = NewObject<UReplicatedTestObjectBridge>();
 	check(ReplicationBridge != nullptr);
@@ -75,6 +85,13 @@ FReplicationSystemTestNode::FReplicationSystemTestNode(bool bIsServer, const TCH
 	Params.bIsServer = bIsServer;
 	Params.bAllowObjectReplication = bIsServer;
 
+	if (ParamsOverride)
+	{
+		Params.MaxReplicatedObjectCount = ParamsOverride->MaxReplicatedObjectCount > 0 ? ParamsOverride->MaxReplicatedObjectCount : Params.MaxReplicatedObjectCount;
+		Params.InitialNetObjectListCount = ParamsOverride->InitialNetObjectListCount > 0 ? ParamsOverride->InitialNetObjectListCount : Params.InitialNetObjectListCount;
+		Params.NetObjectListGrowCount = ParamsOverride->NetObjectListGrowCount > 0 ? ParamsOverride->NetObjectListGrowCount : Params.NetObjectListGrowCount;
+	}
+
 	LOG_SCOPE_VERBOSITY_OVERRIDE(LogIris, ELogVerbosity::Error);
 	ReplicationSystem = FReplicationSystemFactory::CreateReplicationSystem(Params);	
 	if (!bIsServer)
diff --git a/Engine/Plugins/Runtime/ReplicationSystemTestPlugin/Source/Private/Tests/ReplicationSystem/ReplicationSystemServerClientTestFixture.h b/Engine/Plugins/Runtime/ReplicationSystemTestPlugin/Source/Private/Tests/ReplicationSystem/ReplicationSystemServerClientTestFixture.h
--- a/Engine/Plugins/Runtime/ReplicationSystemTestPlugin/Source/Private/Tests/ReplicationSystem/ReplicationSystemServerClientTestFixture.h
+++ b/Engine/Plugins/Runtime/ReplicationSystemTestPlugin/Source/Private/Tests/ReplicationSystem/ReplicationSystemServerClientTestFixture.h
@@ -74,10 +74,22 @@ public:
 		uint32 ConnectionId;
 	};
 
+	struct FReplicationSystemParamsOverride
+	{
+		uint32 MaxReplicatedObjectCount = 0;
+		uint32 InitialNetObjectListCount = 0;
+		uint32 NetObjectListGrowCount = 0;
+	};
+
+	enum EDelaySetup { DelaySetup=0 };
+
 public:
 	FReplicationSystemTestNode(bool bIsServer, const TCHAR* Name);
+	explicit FReplicationSystemTestNode(FReplicationSystemTestNode::EDelaySetup);
 	~FReplicationSystemTestNode();
 
+	void Setup(bool bIsServer, const TCHAR* Name, FReplicationSystemTestNode::FReplicationSystemParamsOverride* ParamsOverride=nullptr);
+
 	template<typename T>
 	T* CreateObject()
 	{
@@ -168,6 +180,7 @@ class FReplicationSystemTestClient : public FReplicationSystemTestNode
 {
 public:
 	FReplicationSystemTestClient(const TCHAR* Name);
+	explicit FReplicationSystemTestClient(FReplicationSystemTestNode::EDelaySetup Delay) : FReplicationSystemTestNode(Delay) {}
 
 	// Tick and send packets to the server
 	bool UpdateAndSend(class FReplicationSystemTestServer* Server, bool bDeliver = true);
@@ -180,6 +193,7 @@ class FReplicationSystemTestServer : public FReplicationSystemTestNode
 {
 public:
 	explicit FReplicationSystemTestServer(const TCHAR* Name);
+	explicit FReplicationSystemTestServer(FReplicationSystemTestNode::EDelaySetup Delay) : FReplicationSystemTestNode(Delay) {}
 
 	// Send data and deliver to the client if bDeliver is true
 	bool SendAndDeliverTo(FReplicationSystemTestClient* Client, bool bDeliver, const TCHAR* Desc = nullptr);
diff --git a/Engine/Plugins/Runtime/ReplicationSystemTestPlugin/Source/Private/Tests/ReplicationSystem/TestReplicationConfig.cpp b/Engine/Plugins/Runtime/ReplicationSystemTestPlugin/Source/Private/Tests/ReplicationSystem/TestReplicationConfig.cpp
new file mode 100644
--- /dev/null
+++ b/Engine/Plugins/Runtime/ReplicationSystemTestPlugin/Source/Private/Tests/ReplicationSystem/TestReplicationConfig.cpp
@@ -0,0 +1,194 @@
+// Copyright Epic Games, Inc. All Rights Reserved.
+
+#include "ReplicationSystemServerClientTestFixture.h"
+#include "Iris/ReplicationSystem/ReplicationSystem.h"
+#include "Iris/ReplicationSystem/ReplicationSystemInternal.h"
+#include "Iris/Core/IrisLog.h"
+
+namespace UE::Net
+{
+
+class FReplicationConfigTestFixture : public FNetworkAutomationTestSuiteFixture
+{
+protected:
+	virtual void SetUp() override
+	{
+		// Delay the setup until the test is ready
+	}
+
+	virtual void TearDown() override
+	{
+		delete Server;
+		delete Client;
+		DataStreamUtil.TearDown();
+		FNetworkAutomationTestSuiteFixture::TearDown();
+	}
+
+	void StartReplicationSystem()
+	{
+		FNetworkAutomationTestSuiteFixture::SetUp();
+
+		// Fake what we normally get from config
+		DataStreamUtil.SetUp();
+		DataStreamUtil.AddDataStreamDefinition(TEXT("NetToken"), TEXT("/Script/IrisCore.NetTokenDataStream"));
+		DataStreamUtil.AddDataStreamDefinition(TEXT("Replication"), TEXT("/Script/IrisCore.ReplicationDataStream"));
+		DataStreamUtil.FixupDefinitions();
+
+		Server = new FReplicationSystemTestServer(FReplicationSystemTestNode::DelaySetup);
+		Server->Setup(true, GetName(), &OverrideServerConfig);
+
+		Client = new FReplicationSystemTestClient(FReplicationSystemTestNode::DelaySetup);
+		Client->Setup(false, GetName(), &OverrideClientConfig);
+
+		// The client needs a connection
+		Client->LocalConnectionId = Client->AddConnection();
+
+		// Auto connect to server
+		Client->ConnectionIdOnServer = Server->AddConnection();
+	}
+
+	void CreateReplicatedObjects(int32 NumObjects)
+	{
+		const int32 StartingIndex = ServerObjects.Num();
+
+		ServerObjects.AddZeroed(NumObjects);
+		ClientObjects.AddZeroed(NumObjects);
+
+		for (int32 i = 0; i < NumObjects; ++i)
+		{
+			const int32 Index = StartingIndex + i;
+
+			UTestReplicatedIrisObject* ServerObject = Server->CreateObject<UTestReplicatedIrisObject>();
+			const FNetRefHandle ServerHandle = ServerObject->NetRefHandle;
+			UE_NET_ASSERT_TRUE(ServerHandle.IsValid());
+
+			ServerObjects[Index] = ServerObject;
+		}
+
+		bool bAllObjectsCreated = false;
+
+		constexpr uint32 MaxTickCount = 100;
+
+		for (uint32 SendCount = 0; SendCount < MaxTickCount; ++SendCount)
+		{
+			// Send and deliver packet
+			Server->UpdateAndSend({ Client });
+
+			bAllObjectsCreated = true;
+			for (int32 i = 0; i < NumObjects; ++i)
+			{
+				const int32 Index = StartingIndex + i;
+				if (ClientObjects[Index] == nullptr)
+				{
+					const FNetRefHandle ServerHandle = ServerObjects[Index]->NetRefHandle;
+					UTestReplicatedIrisObject* ClientObject = Cast<UTestReplicatedIrisObject>(Client->GetReplicationBridge()->GetReplicatedObject(ServerHandle));
+					ClientObjects[Index] = ClientObject;
+
+					bAllObjectsCreated &= (ClientObject != nullptr);
+				}
+			}
+
+			if (bAllObjectsCreated)
+			{
+				break;
+			}
+		}
+
+		UE_NET_ASSERT_TRUE(bAllObjectsCreated);
+	}
+
+	TArray<UTestReplicatedIrisObject*> ServerObjects;
+	TArray<UTestReplicatedIrisObject*> ClientObjects;
+
+	FDataStreamTestUtil DataStreamUtil;
+	FReplicationSystemTestServer* Server = nullptr;
+	FReplicationSystemTestClient* Client = nullptr;
+
+	FReplicationSystemTestNode::FReplicationSystemParamsOverride OverrideServerConfig;
+	FReplicationSystemTestNode::FReplicationSystemParamsOverride OverrideClientConfig;
+};
+
+} // end namespace UE::Net
+
+namespace UE::Net::Private
+{
+UE_NET_TEST_FIXTURE(FReplicationConfigTestFixture, TestNetObjectListGrowEvent)
+{
+	const uint32 MaxNumObjects = 96;
+	const uint32 InitNumObjects = 32;
+	const uint32 GrowCount = 32;
+
+	OverrideServerConfig.MaxReplicatedObjectCount = MaxNumObjects;
+	OverrideServerConfig.InitialNetObjectListCount = InitNumObjects;
+	OverrideServerConfig.NetObjectListGrowCount = GrowCount;
+
+	OverrideClientConfig.MaxReplicatedObjectCount = MaxNumObjects;
+	OverrideClientConfig.InitialNetObjectListCount = InitNumObjects;
+	OverrideClientConfig.NetObjectListGrowCount = GrowCount;
+
+	StartReplicationSystem();
+
+	FDelegateHandle ServerDelegate;
+	FDelegateHandle ClientDelegate;
+
+	FNetRefHandleManager* ServerNetRefHandleManager = &Server->ReplicationSystem->GetReplicationSystemInternal()->GetNetRefHandleManager();
+	FNetRefHandleManager* ClientNetRefHandleManager = &Client->ReplicationSystem->GetReplicationSystemInternal()->GetNetRefHandleManager();
+
+	ON_SCOPE_EXIT
+	{
+		ServerNetRefHandleManager->GetOnMaxInternalNetRefIndexIncreasedDelegate().Remove(ServerDelegate);
+		ClientNetRefHandleManager->GetOnMaxInternalNetRefIndexIncreasedDelegate().Remove(ClientDelegate);
+	};
+
+
+	bool bHasServerRealloc = false;
+	ServerDelegate = ServerNetRefHandleManager->GetOnMaxInternalNetRefIndexIncreasedDelegate().AddLambda([&bHasServerRealloc](uint32 NewMaxIndex)
+	{
+		bHasServerRealloc = true;
+	});
+
+	bool bHasClientRealloc = false;
+	ClientDelegate = ClientNetRefHandleManager->GetOnMaxInternalNetRefIndexIncreasedDelegate().AddLambda([&bHasClientRealloc](uint32 NewMaxIndex)
+	{
+		bHasClientRealloc = true;
+	});
+
+	// Create enough objects to fill the init list (entry[0] is already reserved for invalid object)
+	CreateReplicatedObjects(InitNumObjects - 1);
+	UE_NET_ASSERT_EQ(bHasServerRealloc, false);
+	UE_NET_ASSERT_EQ(bHasClientRealloc, false);
+
+	// Create one more to cause a realloc to trigger
+	CreateReplicatedObjects(1);
+	UE_NET_ASSERT_EQ(bHasServerRealloc, true);
+	UE_NET_ASSERT_EQ(bHasClientRealloc, true);
+
+	// Reset the test conditions
+	bHasServerRealloc = false;
+	bHasClientRealloc = false;
+
+	// Maximimize the objects allowed to exist without reallocating
+	CreateReplicatedObjects(GrowCount-1);
+	UE_NET_ASSERT_EQ(bHasServerRealloc, false);
+	UE_NET_ASSERT_EQ(bHasClientRealloc, false);
+
+	// Create one more to cause a realloc to trigger
+	CreateReplicatedObjects(1);
+	UE_NET_ASSERT_EQ(bHasServerRealloc, true);
+	UE_NET_ASSERT_EQ(bHasClientRealloc, true);
+
+	// Reset the test conditions
+	bHasServerRealloc = false;
+	bHasClientRealloc = false;
+	
+	// Create the maximum amount of held objects
+	CreateReplicatedObjects(MaxNumObjects - ServerObjects.Num() - 1); //Remove 1 to account for invalid entry[0]
+	UE_NET_ASSERT_EQ(bHasServerRealloc, false);
+	UE_NET_ASSERT_EQ(bHasClientRealloc, false);
+
+	// Adding 1 more should cause a Fatal error
+	//CreateReplicatedObjects(1);
+}
+
+} // end namespace UE::Net::Private
+
diff --git a/Engine/Plugins/Runtime/ReplicationSystemTestPlugin/Source/Private/Tests/ReplicationSystem/TestReplicationSystem.cpp b/Engine/Plugins/Runtime/ReplicationSystemTestPlugin/Source/Private/Tests/ReplicationSystem/TestReplicationSystem.cpp
--- a/Engine/Plugins/Runtime/ReplicationSystemTestPlugin/Source/Private/Tests/ReplicationSystem/TestReplicationSystem.cpp
+++ b/Engine/Plugins/Runtime/ReplicationSystemTestPlugin/Source/Private/Tests/ReplicationSystem/TestReplicationSystem.cpp
@@ -1940,6 +1940,6 @@ UE_NET_TEST_FIXTURE(FReplicationSystemServerClientTestFixture, TestPartialDequan
 	CVarForceFullDequantizeAndApply->Set(bForceFullDequantizeAndApply, ECVF_SetByCode);
 }
 
+} // end namespace UE::Net::Private
 
 
-}
