Iris
- Fixed case where flush of reliable rpcs to objects not yet acknowledged by clients would be ignored.

diff --git a/Engine/Plugins/Runtime/ReplicationSystemTestPlugin/Source/Private/Tests/ReplicationSystem/TestFlushBeforeDestroy.cpp b/Engine/Plugins/Runtime/ReplicationSystemTestPlugin/Source/Private/Tests/ReplicationSystem/TestFlushBeforeDestroy.cpp
--- a/Engine/Plugins/Runtime/ReplicationSystemTestPlugin/Source/Private/Tests/ReplicationSystem/TestFlushBeforeDestroy.cpp
+++ b/Engine/Plugins/Runtime/ReplicationSystemTestPlugin/Source/Private/Tests/ReplicationSystem/TestFlushBeforeDestroy.cpp
@@ -5,6 +5,7 @@
 #include "Iris/ReplicationSystem/Filtering/NetObjectFilter.h"
 #include "NetBlob/PartialNetBlobTestFixture.h"
 #include "NetBlob/MockNetBlob.h"
+#include "Misc/ScopeExit.h"
 
 namespace UE::Net::Private
 {
@@ -299,6 +300,18 @@ UE_NET_TEST_FIXTURE(FTestFlushBeforeDestroyFixture, TestReliableAttachmentFlushe
 // destroyed objects.
 UE_NET_TEST_FIXTURE(FTestFlushBeforeDestroyFixture, TestReliableAttachmentFlushedWithPendingCreationLostBeforeDestroy)
 {
+	// Disable flushing / caching for this test as we want to keep exercising the bad path regardless of if we force flushing or not.
+	IConsoleVariable* CVarEnableFlushReliableRPCOnDestroy = IConsoleManager::Get().FindConsoleVariable(TEXT("net.Iris.EnableFlushReliableRPCOnDestroy"));
+	check(CVarEnableFlushReliableRPCOnDestroy != nullptr && CVarEnableFlushReliableRPCOnDestroy->IsVariableBool());
+	const bool bPrevEnableFlushReliableRPCOnDestroy = CVarEnableFlushReliableRPCOnDestroy->GetBool();
+	CVarEnableFlushReliableRPCOnDestroy->Set(false, ECVF_SetByCode);
+
+	ON_SCOPE_EXIT
+	{
+		// Restore cvars
+		CVarEnableFlushReliableRPCOnDestroy->Set(bPrevEnableFlushReliableRPCOnDestroy, ECVF_SetByCode);
+	};
+
 	FReplicationSystemTestClient* Client = CreateClient();
 	RegisterNetBlobHandlers(Client);
 
@@ -400,7 +413,7 @@ UE_NET_TEST_FIXTURE(FTestFlushBeforeDestroyFixture, TestReliableAttachmentFlushe
 	Server->SendAndDeliverTo(Client, DeliverPacket, TEXT("WaitOnFlush"));
 	Server->PostSendUpdate();
 
-	// Verify that the attachment has not received
+	// Verify that the attachment has been received
 	UE_NET_ASSERT_EQ(ClientMockNetObjectAttachmentHandler->GetFunctionCallCounts().OnNetBlobReceived, 1U);
 
 	// Expected to destroy the object
@@ -412,6 +425,76 @@ UE_NET_TEST_FIXTURE(FTestFlushBeforeDestroyFixture, TestReliableAttachmentFlushe
 	UE_NET_ASSERT_EQ(Client->GetReplicationBridge()->GetReplicatedObject(ObjectHandle), nullptr);
 }
 
+UE_NET_TEST_FIXTURE(FTestFlushBeforeDestroyFixture, TestReliableAttachmentFlushedWithLostPendingCreationInflightBeforeDestroy)
+{
+	FReplicationSystemTestClient* Client = CreateClient();
+	RegisterNetBlobHandlers(Client);
+
+	UReplicatedTestObject* ServerObject = Server->CreateObject(0, 0);
+	FNetRefHandle ObjectHandle = ServerObject->NetRefHandle;
+
+	// Setup a situation where we have creation info in flight when the object is destroyed
+
+	// Send creation info
+	Server->PreSendUpdate();
+	Server->SendTo(Client, TEXT("WaitOnCreateConfirmation"));
+	Server->PostSendUpdate();
+
+	// Create attachment
+	{
+		constexpr uint32 PayloadBitCount = 24;
+		const TRefCountPtr<FNetObjectAttachment>& Attachment = MockNetObjectAttachmentHandler->CreateReliableNetObjectAttachment(PayloadBitCount);
+		FNetObjectReference AttachmentTarget = FObjectReferenceCache::MakeNetObjectReference(ServerObject->NetRefHandle);
+		Server->GetReplicationSystem()->QueueNetObjectAttachment(Client->ConnectionIdOnServer, AttachmentTarget, Attachment);
+	}
+
+	// Destroy object, on server
+	Server->DestroyObject(ServerObject, EEndReplicationFlags::Destroy | EEndReplicationFlags::Flush);
+
+	// Previously this would issue a flush and send data before creation is confirmed.
+	Server->PreSendUpdate();
+	const bool bDataWasSentInError = Server->SendTo(Client, TEXT("State should still be WaitOnCreateConfirmation"));
+	Server->PostSendUpdate();
+	
+	// We do not expect any data to be in this packet.
+	UE_NET_ASSERT_FALSE(bDataWasSentInError);
+
+	// Drop initial creation info.
+	Server->DeliverTo(Client, false);
+
+	// Deliver data if we sent data.
+	if (bDataWasSentInError)
+	{
+		// Caused bitstream error on client.
+		Server->DeliverTo(Client, true);
+	}
+
+	// Expected to create object and send attachment
+	Server->PreSendUpdate();
+	Server->SendAndDeliverTo(Client, DeliverPacket, TEXT("CreateResend"));
+	Server->PostSendUpdate();
+
+	// Verify that object does not exist
+	UE_NET_ASSERT_NE(Client->GetReplicationBridge()->GetReplicatedObject(ObjectHandle), nullptr);
+
+	// Verify that the attachment has been received
+	UE_NET_ASSERT_EQ(ClientMockNetObjectAttachmentHandler->GetFunctionCallCounts().OnNetBlobReceived, 1U);
+
+	// Expected to write the attachment
+	Server->PreSendUpdate();
+	Server->SendAndDeliverTo(Client, DeliverPacket, TEXT("WaitOnFlush"));
+	Server->PostSendUpdate();
+
+	// Expected to destroy the object
+	Server->PreSendUpdate();
+	Server->SendAndDeliverTo(Client, DeliverPacket, TEXT("Destroy"));
+	Server->PostSendUpdate();
+
+	// Verify that object has been destroyed
+	UE_NET_ASSERT_EQ(Client->GetReplicationBridge()->GetReplicatedObject(ObjectHandle), nullptr);
+}
+
+
 UE_NET_TEST_FIXTURE(FTestFlushBeforeDestroyFixture, TestReliableAttachmentSubObjectFlushedBeforeDestroy)
 {
 	FReplicationSystemTestClient* Client = CreateClient();
diff --git a/Engine/Plugins/Runtime/ReplicationSystemTestPlugin/Source/Private/Tests/ReplicationSystem/TestReplicationSystem.cpp b/Engine/Plugins/Runtime/ReplicationSystemTestPlugin/Source/Private/Tests/ReplicationSystem/TestReplicationSystem.cpp
--- a/Engine/Plugins/Runtime/ReplicationSystemTestPlugin/Source/Private/Tests/ReplicationSystem/TestReplicationSystem.cpp
+++ b/Engine/Plugins/Runtime/ReplicationSystemTestPlugin/Source/Private/Tests/ReplicationSystem/TestReplicationSystem.cpp
@@ -7,6 +7,7 @@
 #include "Iris/ReplicationSystem/ReplicationSystemInternal.h"
 #include "Iris/ReplicationSystem/NetTokenStore.h"
 #include "Iris/Core/IrisLog.h"
+#include "Misc/ScopeExit.h"
 
 namespace UE::Net::Private
 {
@@ -1854,6 +1855,13 @@ UE_NET_TEST_FIXTURE(FReplicationSystemServerClientTestFixture, TestPartialDequan
 	const bool bForceFullDequantizeAndApply = CVarForceFullDequantizeAndApply->GetBool();
 	CVarForceFullDequantizeAndApply->Set(false, ECVF_SetByCode);
 
+	ON_SCOPE_EXIT
+	{
+		// Restore cvars
+		CVarUsePrevReceivedStateForOnReps->Set(bUsePrevReceivedStateForOnReps, ECVF_SetByCode);
+		CVarForceFullDequantizeAndApply->Set(bForceFullDequantizeAndApply, ECVF_SetByCode);
+	};
+
 	// Add a client
 	FReplicationSystemTestClient* Client = CreateClient();
 
@@ -1935,9 +1943,6 @@ UE_NET_TEST_FIXTURE(FReplicationSystemServerClientTestFixture, TestPartialDequan
 	UE_NET_ASSERT_EQ(ClientObjectA->PrevIntBStoredInOnRep, 1);
 	UE_NET_ASSERT_EQ(ServerObjectA->IntC, ClientObjectA->IntC);
 
-	// Restore cvars
-	CVarUsePrevReceivedStateForOnReps->Set(bUsePrevReceivedStateForOnReps, ECVF_SetByCode);
-	CVarForceFullDequantizeAndApply->Set(bForceFullDequantizeAndApply, ECVF_SetByCode);
 }
 
 } // end namespace UE::Net::Private
diff --git a/Engine/Source/Runtime/Experimental/Iris/Core/Private/Iris/ReplicationSystem/ReplicationWriter.cpp b/Engine/Source/Runtime/Experimental/Iris/Core/Private/Iris/ReplicationSystem/ReplicationWriter.cpp
--- a/Engine/Source/Runtime/Experimental/Iris/Core/Private/Iris/ReplicationSystem/ReplicationWriter.cpp
+++ b/Engine/Source/Runtime/Experimental/Iris/Core/Private/Iris/ReplicationSystem/ReplicationWriter.cpp
@@ -789,10 +789,19 @@ void FReplicationWriter::UpdateScope(const FNetBitArrayView& UpdatedScope)
 		const EReplicatedObjectState State = Info.GetState();
 		if (State < EReplicatedObjectState::PendingDestroy)
 		{
-			// We have not sent the object yet so we can just stop replication
 			if (State == EReplicatedObjectState::PendingCreate)
 			{
-				StopReplication(Index);
+				// If we have no data to flush, we can stop replication now.
+				const uint32 FlushFlags = GetFlushStatus(Index, Info, Info.FlushFlags);
+				if (GetFlushStatus(Index, Info, Info.FlushFlags) == FlushFlags_None)
+				{
+					StopReplication(Index);
+				}
+				else
+				{
+					// Mark for destroy.
+					ObjectsPendingDestroy.SetBit(Index);
+				}
 			}
 			else if (State == EReplicatedObjectState::CancelPendingDestroy)
 			{
@@ -1248,9 +1257,10 @@ void FReplicationWriter::HandleDroppedRecord<FReplicationWriter::EReplicatedObje
 
 	if (CurrentState < EReplicatedObjectState::Created)
 	{
-		// Until we have implemented cached creation info we cannot send creation info for destroyed objects
-		// So we just have to StopReplication
-		const bool bCanSendCreationInfo = !ObjectsPendingDestroy.GetBit(InternalIndex);
+		const FNetRefHandleManager::FReplicatedObjectData& ObjectData = NetRefHandleManager->GetReplicatedObjectData(InternalIndex);
+
+		// We can resend creation info even if we are marked for destroy if we have cached creation info.
+		const bool bCanSendCreationInfo = !ObjectsPendingDestroy.GetBit(InternalIndex) || ObjectData.bHasCachedCreationInfo;
 		if (bCanSendCreationInfo)
 		{
 			// Mark object as having dirty changes
@@ -1276,7 +1286,6 @@ void FReplicationWriter::HandleDroppedRecord<FReplicationWriter::EReplicatedObje
 			if (Info.IsSubObject)
 			{
 				// Mark owner dirty as well as subobjects only are scheduled together with owner
-				const FNetRefHandleManager::FReplicatedObjectData& ObjectData = NetRefHandleManager->GetReplicatedObjectData(InternalIndex);
 				uint32 SubObjectOwnerInternalIndex = ObjectData.SubObjectRootIndex;
 
 				FReplicationInfo& SubObjectOwnerReplicationInfo = GetReplicationInfo(SubObjectOwnerInternalIndex);
