iris
- Fixed issue with Torned-off objects being destroyed while waiting for initial acknowledgement.

diff --git a/Engine/Plugins/Runtime/ReplicationSystemTestPlugin/Source/Private/Tests/ReplicationSystem/TestFlushBeforeDestroy.cpp b/Engine/Plugins/Runtime/ReplicationSystemTestPlugin/Source/Private/Tests/ReplicationSystem/TestFlushBeforeDestroy.cpp
--- a/Engine/Plugins/Runtime/ReplicationSystemTestPlugin/Source/Private/Tests/ReplicationSystem/TestFlushBeforeDestroy.cpp
+++ b/Engine/Plugins/Runtime/ReplicationSystemTestPlugin/Source/Private/Tests/ReplicationSystem/TestFlushBeforeDestroy.cpp
@@ -681,6 +681,209 @@ UE_NET_TEST_FIXTURE(FTestFlushBeforeDestroyFixture, TestStateInFlightFlushedBefo
 	UE_NET_ASSERT_EQ(Client->GetReplicationBridge()->GetReplicatedObject(ObjectHandle), nullptr);
 }
 
+UE_NET_TEST_FIXTURE(FTestFlushBeforeDestroyFixture, TestDroppedPendingTearOffIsCancelledByEndReplication)
+{
+	// As we are testing old behavior, we need to make sure to allow double endreplication so we hit the path we want to test.
+	IConsoleVariable* CVarAllowDestroyToCancelFlushAndTearOff = IConsoleManager::Get().FindConsoleVariable(TEXT("net.Iris.AllowDestroyToCancelFlushAndTearOff"), false);
+	UE_NET_ASSERT_NE(CVarAllowDestroyToCancelFlushAndTearOff, nullptr);
+	UE_NET_ASSERT_TRUE(CVarAllowDestroyToCancelFlushAndTearOff->IsVariableBool());
+
+	const bool bOldAllowDestroyToCancelFlushAndTearOff = CVarAllowDestroyToCancelFlushAndTearOff->GetBool();
+	ON_SCOPE_EXIT { CVarAllowDestroyToCancelFlushAndTearOff->Set(bOldAllowDestroyToCancelFlushAndTearOff, ECVF_SetByCode); };
+
+	CVarAllowDestroyToCancelFlushAndTearOff->Set(true, ECVF_SetByCode);
+
+	FReplicationSystemTestClient* Client = CreateClient();
+	RegisterNetBlobHandlers(Client);
+
+	Server->PreSendUpdate();
+	Server->SendAndDeliverTo(Client, DeliverPacket);
+	Server->PostSendUpdate();
+
+	// Setup case where we have a new object for which we have an attachment which should execute a tearoff after we have confirmed creation
+	UReplicatedTestObject* ServerObject = Server->CreateObject(0, 0);
+	FNetRefHandle ObjectHandle = ServerObject->NetRefHandle;
+
+	// Create attachment
+	{
+		constexpr uint32 PayloadBitCount = 24;
+		const TRefCountPtr<FNetObjectAttachment>& Attachment = MockNetObjectAttachmentHandler->CreateReliableNetObjectAttachment(PayloadBitCount);
+		FNetObjectReference AttachmentTarget = FObjectReferenceCache::MakeNetObjectReference(ServerObject->NetRefHandle);
+		Server->GetReplicationSystem()->QueueNetObjectAttachment(Client->ConnectionIdOnServer, AttachmentTarget, Attachment);
+	}
+
+	// Request tearoff
+	Server->ReplicationBridge->EndReplication(ServerObject->NetRefHandle, EEndReplicationFlags::TearOff);
+
+	// Send packet so that we have creationdata in flight
+	Server->PreSendUpdate();
+	Server->SendTo(Client);
+	Server->PostSendUpdate();
+
+	// Force destroy object already pending tearoff/flush. DestroyLocalNetHandle will invalidate cached creationinfo.
+	Server->ReplicationBridge->EndReplication(ServerObject->NetRefHandle, EEndReplicationFlags::Destroy);
+
+	// Drop and notify that the packet while object still is the state waitoncreateconfirmation as we have not yet updated scope.
+	// When this failed it did put the state of the object back in PendingCreate even though we no longer had any cached creationinfo.
+	Server->DeliverTo(Client, false);
+
+	// Deliver a packet, this should flush the object and deliver the attachment
+	Server->PreSendUpdate();
+	Server->SendAndDeliverTo(Client, DeliverPacket);
+	Server->PostSendUpdate();
+
+	// Verify that the attachment has not been received
+	UE_NET_ASSERT_NE(ClientMockNetObjectAttachmentHandler->GetFunctionCallCounts().OnNetBlobReceived, 1U);
+
+	// Deliver a packet. Should destroy the object on the client unless that was already done
+	Server->PreSendUpdate();
+	Server->SendAndDeliverTo(Client, DeliverPacket);
+	Server->PostSendUpdate();
+
+	// Verify that object is destroyed
+	UE_NET_ASSERT_EQ(Client->GetReplicationBridge()->GetReplicatedObject(ObjectHandle), nullptr);
+}
+
+UE_NET_TEST_FIXTURE(FTestFlushBeforeDestroyFixture, TestPendingCreateTearOffIsCancelledByEndReplication)
+{
+	// As we are testing old behavior, we need to make sure to allow double endreplication so we hit the path we want to test.
+	IConsoleVariable* CVarAllowDestroyToCancelFlushAndTearOff = IConsoleManager::Get().FindConsoleVariable(TEXT("net.Iris.AllowDestroyToCancelFlushAndTearOff"), false);
+	UE_NET_ASSERT_NE(CVarAllowDestroyToCancelFlushAndTearOff, nullptr);
+	UE_NET_ASSERT_TRUE(CVarAllowDestroyToCancelFlushAndTearOff->IsVariableBool());
+
+	const bool bOldAllowDestroyToCancelFlushAndTearOff = CVarAllowDestroyToCancelFlushAndTearOff->GetBool();
+	ON_SCOPE_EXIT { CVarAllowDestroyToCancelFlushAndTearOff->Set(bOldAllowDestroyToCancelFlushAndTearOff, ECVF_SetByCode); };
+
+	CVarAllowDestroyToCancelFlushAndTearOff->Set(true, ECVF_SetByCode);
+
+	FReplicationSystemTestClient* Client = CreateClient();
+	RegisterNetBlobHandlers(Client);
+
+	// Setup case where we have a new object for which we have an attachment which should execute a tearoff after we have confirmed creation
+	UReplicatedTestObject* ServerObject = Server->CreateObject(0, 0);
+	FNetRefHandle ObjectHandle = ServerObject->NetRefHandle;
+
+	// Create attachment
+	{
+		constexpr uint32 PayloadBitCount = 24;
+		const TRefCountPtr<FNetObjectAttachment>& Attachment = MockNetObjectAttachmentHandler->CreateReliableNetObjectAttachment(PayloadBitCount);
+		FNetObjectReference AttachmentTarget = FObjectReferenceCache::MakeNetObjectReference(ServerObject->NetRefHandle);
+		Server->GetReplicationSystem()->QueueNetObjectAttachment(Client->ConnectionIdOnServer, AttachmentTarget, Attachment);
+	}
+
+	// Request tearoff
+	Server->ReplicationBridge->EndReplication(ServerObject->NetRefHandle, EEndReplicationFlags::TearOff);
+
+	// PreUpdate to update scoping to get the object into the PendingCreate state.
+	Server->PreSendUpdate();
+	Server->PostSendUpdate();
+
+	// Force destroy object already pending tearoff/flush. DestroyLocalNetHandle will invalidate cached creationinfo.
+	Server->ReplicationBridge->EndReplication(ServerObject->NetRefHandle, EEndReplicationFlags::Destroy);
+
+	// Send a packet.
+	Server->PreSendUpdate();
+	Server->SendAndDeliverTo(Client, DeliverPacket);
+	Server->PostSendUpdate();
+
+	// Verify that the attachment has not been received.
+	UE_NET_ASSERT_NE(ClientMockNetObjectAttachmentHandler->GetFunctionCallCounts().OnNetBlobReceived, 1U);
+
+	// Verify that the object is not created.
+	UE_NET_ASSERT_EQ(Client->GetReplicationBridge()->GetReplicatedObject(ObjectHandle), nullptr);
+}
+
+UE_NET_TEST_FIXTURE(FTestFlushBeforeDestroyFixture, TestPendingCreateTearOffIsNotCancelledByEndReplication)
+{
+	FReplicationSystemTestClient* Client = CreateClient();
+	RegisterNetBlobHandlers(Client);
+
+	// Setup case where we have a new object for which we have an attachment which should execute a tearoff after we have confirmed creation
+	UReplicatedTestObject* ServerObject = Server->CreateObject(0, 0);
+	FNetRefHandle ObjectHandle = ServerObject->NetRefHandle;
+
+	// Create attachment
+	{
+		constexpr uint32 PayloadBitCount = 24;
+		const TRefCountPtr<FNetObjectAttachment>& Attachment = MockNetObjectAttachmentHandler->CreateReliableNetObjectAttachment(PayloadBitCount);
+		FNetObjectReference AttachmentTarget = FObjectReferenceCache::MakeNetObjectReference(ServerObject->NetRefHandle);
+		Server->GetReplicationSystem()->QueueNetObjectAttachment(Client->ConnectionIdOnServer, AttachmentTarget, Attachment);
+	}
+
+	// Request tearoff
+	Server->ReplicationBridge->EndReplication(ServerObject->NetRefHandle, EEndReplicationFlags::TearOff);
+
+	// PreUpdate to update scoping to get the object into the PendingCreate state.
+	Server->PreSendUpdate();
+	Server->PostSendUpdate();
+
+	// This should be ignored as we are already pending tear off.
+	Server->ReplicationBridge->EndReplication(ServerObject->NetRefHandle, EEndReplicationFlags::Destroy);
+
+	// Deliver a packet, this should flush the object and deliver the attachment
+	Server->PreSendUpdate();
+	Server->SendAndDeliverTo(Client, DeliverPacket);
+	Server->PostSendUpdate();
+
+	// Verify that the attachment has been received
+	UE_NET_ASSERT_EQ(ClientMockNetObjectAttachmentHandler->GetFunctionCallCounts().OnNetBlobReceived, 1U);
+
+	// Verify that object is created
+	UE_NET_ASSERT_NE(Client->GetReplicationBridge()->GetReplicatedObject(ObjectHandle), nullptr);
+}
+
+UE_NET_TEST_FIXTURE(FTestFlushBeforeDestroyFixture, TestDroppedTearOffIsNotCancelledByEndReplication)
+{
+	FReplicationSystemTestClient* Client = CreateClient();
+	RegisterNetBlobHandlers(Client);
+
+	Server->PreSendUpdate();
+	Server->SendAndDeliverTo(Client, DeliverPacket);
+	Server->PostSendUpdate();
+
+	// Setup case where we have a new object for which we have an attachment which should execute a tearoff after we have confirmed creation
+	UReplicatedTestObject* ServerObject = Server->CreateObject(0, 0);
+	FNetRefHandle ObjectHandle = ServerObject->NetRefHandle;
+
+	// Create attachment
+	{
+		constexpr uint32 PayloadBitCount = 24;
+		const TRefCountPtr<FNetObjectAttachment>& Attachment = MockNetObjectAttachmentHandler->CreateReliableNetObjectAttachment(PayloadBitCount);
+		FNetObjectReference AttachmentTarget = FObjectReferenceCache::MakeNetObjectReference(ServerObject->NetRefHandle);
+		Server->GetReplicationSystem()->QueueNetObjectAttachment(Client->ConnectionIdOnServer, AttachmentTarget, Attachment);
+	}
+
+	// Request tearoff
+	Server->ReplicationBridge->EndReplication(ServerObject->NetRefHandle, EEndReplicationFlags::TearOff);
+
+	// Send packet so that we have creationdata in flight
+	Server->PreSendUpdate();
+	Server->SendTo(Client);
+	Server->PostSendUpdate();
+
+	// Force destroy object already pending tearoff/flush. This should be ignored
+	Server->ReplicationBridge->EndReplication(ServerObject->NetRefHandle, EEndReplicationFlags::Destroy);
+
+	// Drop and notify that the packet while object still is the state waitoncreateconfirmation as we have not yet updated scope.
+	Server->DeliverTo(Client, false);
+
+	// Deliver a packet, this should flush the object and deliver the attachment
+	Server->PreSendUpdate();
+	Server->SendAndDeliverTo(Client, DeliverPacket);
+	Server->PostSendUpdate();
+
+	// Verify that the attachment has been received
+	UE_NET_ASSERT_EQ(ClientMockNetObjectAttachmentHandler->GetFunctionCallCounts().OnNetBlobReceived, 1U);
+
+	// Deliver a packet. Should tear off the object on the client
+	Server->PreSendUpdate();
+	Server->SendAndDeliverTo(Client, DeliverPacket);
+	Server->PostSendUpdate();
+
+	// Verify that object is not findable
+	UE_NET_ASSERT_EQ(Client->GetReplicationBridge()->GetReplicatedObject(ObjectHandle), nullptr);
+}
+
 UE_NET_TEST_FIXTURE(FTestFlushBeforeDestroyFixture, TestSubObjectStateFlushedBeforeOwnerDestroy)
 {
 	FReplicationSystemTestClient* Client = CreateClient();
diff --git a/Engine/Source/Runtime/Experimental/Iris/Core/Private/Iris/ReplicationSystem/NetRefHandleManager.cpp b/Engine/Source/Runtime/Experimental/Iris/Core/Private/Iris/ReplicationSystem/NetRefHandleManager.cpp
--- a/Engine/Source/Runtime/Experimental/Iris/Core/Private/Iris/ReplicationSystem/NetRefHandleManager.cpp
+++ b/Engine/Source/Runtime/Experimental/Iris/Core/Private/Iris/ReplicationSystem/NetRefHandleManager.cpp
@@ -556,6 +556,7 @@ void FNetRefHandleManager::DestroyObjectsPendingDestroy()
 			PendingDestroyInternalIndices.RemoveAtSwap(It);
 		}
 	}
+	CSV_CUSTOM_STAT(IrisCommon, PendingDestroyInternalIndicesCount, (float)PendingDestroyInternalIndices.Num(), ECsvCustomStatOp::Set);
 }
 
 bool FNetRefHandleManager::AddSubObject(FNetRefHandle OwnerHandle, FNetRefHandle SubObjectHandle, FNetRefHandle RelativeOtherSubObjectHandle, EAddSubObjectFlags Flags)
diff --git a/Engine/Source/Runtime/Experimental/Iris/Core/Private/Iris/ReplicationSystem/ReplicationBridge.cpp b/Engine/Source/Runtime/Experimental/Iris/Core/Private/Iris/ReplicationSystem/ReplicationBridge.cpp
--- a/Engine/Source/Runtime/Experimental/Iris/Core/Private/Iris/ReplicationSystem/ReplicationBridge.cpp
+++ b/Engine/Source/Runtime/Experimental/Iris/Core/Private/Iris/ReplicationSystem/ReplicationBridge.cpp
@@ -32,7 +32,6 @@ DEFINE_LOG_CATEGORY(LogIrisBridge)
 
 #define UE_LOG_REPLICATIONBRIDGE(Category, Format, ...)  UE_LOG(LogIrisBridge, Category, TEXT("ReplicationBridge(%u)::") Format, ReplicationSystem->GetId(), ##__VA_ARGS__)
 
-
 static bool bEnableFlushReliableRPCOnDestroy = true;
 static FAutoConsoleVariableRef CVarEnableFlushReliableRPCOnDestroy(
 	TEXT("net.Iris.EnableFlushReliableRPCOnDestroy"),
@@ -40,6 +39,13 @@ static FAutoConsoleVariableRef CVarEnableFlushReliableRPCOnDestroy(
 	TEXT("When true EEndReplicationFlags::Flush flag will be appended in EndReplication if we have pending unprocessed attachments/RPC:s when destroying a replicated object.")
 );
 
+static bool bAllowDestroyToCancelFlushAndTearOff = false;
+static FAutoConsoleVariableRef CVarAllowDestroyToCancelFlushAndTearOff(
+	TEXT("net.Iris.AllowDestroyToCancelFlushAndTearOff"),
+	bAllowDestroyToCancelFlushAndTearOff,
+	TEXT("When true issuing a EndReplication on an object that is already Tear-off or pending endreplication will cancel destroy/flush and destroy the replicated objects.")
+);
+
 /**
  * ReplicationBridge Implementation
  */
@@ -451,6 +457,8 @@ void UReplicationBridge::InternalDetachInstanceFromNetRefHandle(FNetRefHandle Re
 
 	if (FReplicationInstanceProtocol* InstanceProtocol = const_cast<FReplicationInstanceProtocol*>(NetRefHandleManager->DetachInstanceProtocol(InternalReplicationIndex)))
 	{
+		UE_LOG_REPLICATIONBRIDGE(Verbose, TEXT("InternalDetachInstanceToNetHandle Detached: %s from ( InternalIndex: %u )"), *RefHandle.ToString(), InternalReplicationIndex);
+
 		if (EnumHasAnyFlags(InstanceProtocol->InstanceTraits, EReplicationInstanceProtocolTraits::IsBound))
 		{
 			FReplicationInstanceOperationsInternal::UnbindInstanceProtocol(InstanceProtocol, NetRefHandleManager->GetReplicatedObjectDataNoCheck(InternalReplicationIndex).Protocol);
@@ -474,23 +482,26 @@ void UReplicationBridge::InternalDestroyNetObject(FNetRefHandle Handle)
 		FWorldLocations& WorldLocations = ReplicationSystemInternal->GetWorldLocations();
 		WorldLocations.RemoveObjectInfoCache(ObjectInternalIndex);
 
-		const FNetRefHandleManager::FReplicatedObjectData& ObjectData = NetRefHandleManager->GetReplicatedObjectDataNoCheck(ObjectInternalIndex);
+		FNetRefHandleManager::FReplicatedObjectData& ObjectData = NetRefHandleManager->GetReplicatedObjectDataNoCheck(ObjectInternalIndex);
 
-		if (ObjectData.bTearOff && NetRefHandleManager->GetNetObjectRefCount(ObjectInternalIndex) > 0U)
+		// At this point we can no longer instantiate objects and must explicitly clear out objects pending create from ReplicationWriter
+		if (ObjectData.bHasCachedCreationInfo && NetRefHandleManager->GetNetObjectRefCount(ObjectInternalIndex) > 0U)
 		{
 			// We need to explicitly notify all ReplicationWriters that we are destroying objects pending tearoff
-			// The handle will automatically be removed from HandlesPendingTearOff after the next update
-
+			// The handle will automatically be removed from HandlesPendingEndReplication after the next update
 			FReplicationConnections& Connections = ReplicationSystemInternal->GetConnections();
 
-			auto NotifyDestroyedObjectPendingTearOff = [&Connections, &ObjectInternalIndex](uint32 ConnectionId)
+			auto NotifyDestroyedObjectPendingEndReplication = [&Connections, &ObjectInternalIndex](uint32 ConnectionId)
 			{
 				FReplicationConnection* Conn = Connections.GetConnection(ConnectionId);			
-				Conn->ReplicationWriter->NotifyDestroyedObjectPendingTearOff(ObjectInternalIndex);
+				Conn->ReplicationWriter->NotifyDestroyedObjectPendingEndReplication(ObjectInternalIndex);
 			};
 
 			const FNetBitArray& ValidConnections = Connections.GetValidConnections();
-			ValidConnections.ForAllSetBits(NotifyDestroyedObjectPendingTearOff);					
+			ValidConnections.ForAllSetBits(NotifyDestroyedObjectPendingEndReplication);					
+
+			// As we have release cached creation headers before calling InternalDestroyNetObject it is critical that we also clear out the flag.
+			ObjectData.bHasCachedCreationInfo = 0U;
 		}
 	}
 
@@ -588,6 +599,13 @@ void UReplicationBridge::EndReplication(FNetRefHandle Handle, EEndReplicationFla
 	if (NetRefHandleManager->IsLocal(InternalReplicationIndex))
 	{
 		FNetRefHandleManager::FReplicatedObjectData& ObjectData = NetRefHandleManager->GetReplicatedObjectDataNoCheck(InternalReplicationIndex);
+
+		if (ObjectData.bPendingEndReplication && !bAllowDestroyToCancelFlushAndTearOff)
+		{
+			UE_LOG_REPLICATIONBRIDGE(Verbose, TEXT("Ignoring EndReplication called on object already PendingEndReplication %s."), *NetRefHandleManager->PrintObjectFromNetRefHandle(Handle));
+			return;
+		}
+
 		if (EnumHasAnyFlags(EndReplicationFlags, EEndReplicationFlags::TearOff))
 		{
 			// Add handle to list of objects pending EndReplication indicate that it should be destroyed during next update
@@ -595,8 +613,6 @@ void UReplicationBridge::EndReplication(FNetRefHandle Handle, EEndReplicationFla
 			AddPendingEndReplication(Handle, EndReplicationFlags);
 
 			// We do however copy the final state data and mark object to stop propagating state changes
-			// Note: Current implementation keeps the mapping between instance and handle alive so that we can handle the case where we both start replication 
-			// and tear the object off during the same frame
 			InternalTearOff(Handle, ECacheCreationInfo::Yes);
 
 			// Detach instance as we must assume that we should not access the object after this call.
@@ -791,6 +807,8 @@ void UReplicationBridge::InternalFlushStateData(UE::Net::FNetSerializationContex
 		return;
 	}
 
+	UE_LOG_REPLICATIONBRIDGE(Verbose, TEXT("InternalFlushStateData Initiating flush for %s ( InternalIndex: %u )"), *ObjectData.RefHandle.ToString(), InternalObjectIndex);
+
 	if (ObjectData.InstanceProtocol)
 	{
 		if (EnumHasAnyFlags(ObjectData.InstanceProtocol->InstanceTraits, EReplicationInstanceProtocolTraits::NeedsPoll | EReplicationInstanceProtocolTraits::NeedsPreSendUpdate))
@@ -880,7 +898,7 @@ void UReplicationBridge::InternalTearOff(FNetRefHandle Handle, ECacheCreationInf
 	// Copy state data and tear off now
 	if (!ObjectData.bTearOff)
 	{
-		UE_LOG_REPLICATIONBRIDGE(Verbose, TEXT("TearOff %s"), *Handle.ToString());
+		UE_LOG_REPLICATIONBRIDGE(Verbose, TEXT("TearOff %s ( InteralIndex: %u )"), *Handle.ToString(), InternalObjectIndex);
 
 		// Force copy of final state data as we will detach the object after scope update
 		FChangeMaskCache ChangeMaskCache;
diff --git a/Engine/Source/Runtime/Experimental/Iris/Core/Private/Iris/ReplicationSystem/ReplicationWriter.cpp b/Engine/Source/Runtime/Experimental/Iris/Core/Private/Iris/ReplicationSystem/ReplicationWriter.cpp
--- a/Engine/Source/Runtime/Experimental/Iris/Core/Private/Iris/ReplicationSystem/ReplicationWriter.cpp
+++ b/Engine/Source/Runtime/Experimental/Iris/Core/Private/Iris/ReplicationSystem/ReplicationWriter.cpp
@@ -901,12 +901,11 @@ void FReplicationWriter::InternalUpdateDirtyChangeMasks(const FChangeMaskCache&
 	//UE_LOG_REPLICATIONWRITER(TEXT("FReplicationWriter::UpdateDirtyChangeMasks() Updated %u Objects for ConnectionId:%u, ReplicationSystemId: %u."), CachedChangeMasks.Indices.Num(), Parameters.ConnectionId, Parameters.ReplicationSystem->GetId());	
 }
 
-void FReplicationWriter::NotifyDestroyedObjectPendingTearOff(FInternalNetRefIndex ObjectInternalIndex)
+void FReplicationWriter::NotifyDestroyedObjectPendingEndReplication(FInternalNetRefIndex ObjectInternalIndex)
 {
 	const FReplicationInfo& ReplicationInfo = GetReplicationInfo(ObjectInternalIndex);
 	if (ReplicationInfo.GetState() == EReplicatedObjectState::PendingCreate)
 	{
-		check(ReplicationInfo.TearOff == 1U);
 		StopReplication(ObjectInternalIndex);
 	}
 }
@@ -1073,6 +1072,21 @@ void FReplicationWriter::HandleDeliveredRecord(const FReplicationRecord::FRecord
 				// If this object was teared off, it can now be considered as destroyed
 				else if (RecordInfo.WroteTearOff)
 				{
+					// Must also mark owner dirty as it might have been waiting for a subobject flush
+					if (Info.IsSubObject)
+					{
+						FNetRefHandleManager::FReplicatedObjectData& ObjectData = NetRefHandleManager->GetReplicatedObjectDataNoCheck(InternalIndex);
+						if (ObjectData.IsSubObject())
+						{
+							FReplicationInfo& OwnerInfo = ReplicatedObjects[ObjectData.SubObjectRootIndex];
+							if (OwnerInfo.GetState() != EReplicatedObjectState::Invalid && OwnerInfo.GetState() < EReplicatedObjectState::PendingDestroy)
+							{
+								MarkObjectDirty(ObjectData.SubObjectRootIndex, "HandleDeliveredRecordSubObjectTearOff");
+								OwnerInfo.HasDirtySubObjects = 1U;
+							}
+						}
+					}
+
 					SetState(InternalIndex, EReplicatedObjectState::PendingTearOff);
 					SetState(InternalIndex, EReplicatedObjectState::Destroyed);
 					StopReplication(InternalIndex);
@@ -1261,8 +1275,8 @@ void FReplicationWriter::HandleDroppedRecord<FReplicationWriter::EReplicatedObje
 	{
 		const FNetRefHandleManager::FReplicatedObjectData& ObjectData = NetRefHandleManager->GetReplicatedObjectData(InternalIndex);
 
-		// We can resend creation info even if we are marked for destroy if we have cached creation info.
-		const bool bCanSendCreationInfo = !ObjectsPendingDestroy.GetBit(InternalIndex) || ObjectData.bHasCachedCreationInfo;
+		// We can resend creation info even if we are marked for destroy/endrepliation as long as we have cached creation info.
+		const bool bCanSendCreationInfo = ObjectData.bHasCachedCreationInfo || (!ObjectsPendingDestroy.GetBit(InternalIndex) && !ObjectData.bPendingEndReplication);
 		if (bCanSendCreationInfo)
 		{
 			// Mark object as having dirty changes
@@ -2182,6 +2196,9 @@ FReplicationWriter::EWriteObjectStatus FReplicationWriter::WriteObjectAndSubObje
 					{
 						if (!Context.HasErrorOrOverflow())
 						{
+							UE_LOG_REPLICATIONWRITER_WARNING(TEXT("Failed to replicate ( InternalIndex: %u ) %s, ProtocolName: %s, InstanceProtocol pointer: %p, HasCachedCreationInfo: %u"), InternalIndex, *NetRefHandle.ToString(), (ObjectData.Protocol ? ToCStr(ObjectData.Protocol->DebugName) : TEXT("nullptr")), ObjectData.InstanceProtocol, ObjectData.bHasCachedCreationInfo);
+							ensureMsgf(ObjectData.Protocol, TEXT("Failed to replicate ( InternalIndex: %u ) %s, Protocol: nullptr, InstanceProtocol pointer: %p, HasCachedCreationInfo: %u"), InternalIndex, *NetRefHandle.ToString(), ObjectData.InstanceProtocol, ObjectData.bHasCachedCreationInfo);
+
 							// Unforced error, treat it as we have no instance and cannot create this object but we can continue with other objects
 							return EWriteObjectStatus::NoInstanceProtocol;
 						}
diff --git a/Engine/Source/Runtime/Experimental/Iris/Core/Private/Iris/ReplicationSystem/ReplicationWriter.h b/Engine/Source/Runtime/Experimental/Iris/Core/Private/Iris/ReplicationSystem/ReplicationWriter.h
--- a/Engine/Source/Runtime/Experimental/Iris/Core/Private/Iris/ReplicationSystem/ReplicationWriter.h
+++ b/Engine/Source/Runtime/Experimental/Iris/Core/Private/Iris/ReplicationSystem/ReplicationWriter.h
@@ -158,8 +158,8 @@ public:
 	// Force update DirtyChangeMasks and mark objects for flush and/or tearoff depending on flags
 	void ForceUpdateDirtyChangeMasks(const FChangeMaskCache& CachedChangeMasks, EFlushFlags ExtraFlushFlags, bool bMarkForTearOff) { InternalUpdateDirtyChangeMasks(CachedChangeMasks, ExtraFlushFlags, bMarkForTearOff); }
 
-	// Called if an object first being teared off and then explicitly destroyed before it has been removed from scope
-	void NotifyDestroyedObjectPendingTearOff(FInternalNetRefIndex ObjectInternalIndex);
+	// Called if an object first being teared-off/flushed and then explicitly destroyed before it has been removed from scope
+	void NotifyDestroyedObjectPendingEndReplication(FInternalNetRefIndex ObjectInternalIndex);
 
 	// Propagate dirty changemasks
 	void UpdateDirtyChangeMasks(const FChangeMaskCache& CachedChangeMasks) { InternalUpdateDirtyChangeMasks(CachedChangeMasks, EFlushFlags::FlushFlags_None, false); }
