Iris - Minor test framework improvements.

#rnx

diff --git a/Engine/Plugins/Runtime/ReplicationSystemTestPlugin/Source/Private/Tests/ReplicationSystem/ReplicatedTestObject.cpp b/Engine/Plugins/Runtime/ReplicationSystemTestPlugin/Source/Private/Tests/ReplicationSystem/ReplicatedTestObject.cpp
--- a/Engine/Plugins/Runtime/ReplicationSystemTestPlugin/Source/Private/Tests/ReplicationSystem/ReplicatedTestObject.cpp
+++ b/Engine/Plugins/Runtime/ReplicationSystemTestPlugin/Source/Private/Tests/ReplicationSystem/ReplicatedTestObject.cpp
@@ -279,6 +279,11 @@ void UReplicatedTestObjectBridge::SetExternalWorldLocationUpdateFunctor(TFunctio
 	SetInstanceGetWorldObjectInfoFunction(LocUpdateFunctor);
 }
 
+TFunction<void(UE::Net::FNetRefHandle, const UObject*, FVector&, float&)> UReplicatedTestObjectBridge::GetExternalWorldLocationUpdateFunctor() const
+{
+	return GetInstanceGetWorldObjectInfoFunction();
+}
+
 void UReplicatedTestObjectBridge::SetExternalPreUpdateFunctor(TFunction<void(TArrayView<UObject*>, const UReplicationBridge*)> PreUpdateFunctor)
 {
 	SetInstancePreUpdateFunction(PreUpdateFunctor);
diff --git a/Engine/Plugins/Runtime/ReplicationSystemTestPlugin/Source/Private/Tests/ReplicationSystem/ReplicatedTestObject.h b/Engine/Plugins/Runtime/ReplicationSystemTestPlugin/Source/Private/Tests/ReplicationSystem/ReplicatedTestObject.h
--- a/Engine/Plugins/Runtime/ReplicationSystemTestPlugin/Source/Private/Tests/ReplicationSystem/ReplicatedTestObject.h
+++ b/Engine/Plugins/Runtime/ReplicationSystemTestPlugin/Source/Private/Tests/ReplicationSystem/ReplicatedTestObject.h
@@ -505,6 +505,7 @@ public:
 	const UE::Net::FReplicationInstanceProtocol* GetReplicationInstanceProtocol(FNetRefHandle Handle) const;
 
 	void SetExternalWorldLocationUpdateFunctor(TFunction<void(FNetRefHandle NetHandle, const UObject* ReplicatedObject, FVector& OutLocation, float& OutCullDistance)> LocUpdateFunctor);
+	TFunction<void(FNetRefHandle, const UObject*, FVector&, float&)> GetExternalWorldLocationUpdateFunctor() const;
 
 	void SetExternalPreUpdateFunctor(TFunction<void(TArrayView<UObject*>, const UReplicationBridge*)> PreUpdateFunctor);
 
diff --git a/Engine/Plugins/Runtime/ReplicationSystemTestPlugin/Source/Private/Tests/ReplicationSystem/ReplicationSystemServerClientTestFixture.cpp b/Engine/Plugins/Runtime/ReplicationSystemTestPlugin/Source/Private/Tests/ReplicationSystem/ReplicationSystemServerClientTestFixture.cpp
--- a/Engine/Plugins/Runtime/ReplicationSystemTestPlugin/Source/Private/Tests/ReplicationSystem/ReplicationSystemServerClientTestFixture.cpp
+++ b/Engine/Plugins/Runtime/ReplicationSystemTestPlugin/Source/Private/Tests/ReplicationSystem/ReplicationSystemServerClientTestFixture.cpp
@@ -256,6 +256,31 @@ uint32 FReplicationSystemTestNode::AddConnection()
 	return Connection.ConnectionId;
 }
 
+void FReplicationSystemTestNode::RemoveConnection(uint32 ConnectionId)
+{
+	for (TArray<FConnectionInfo>::TIterator It = Connections.CreateIterator(); It; ++It)
+	{
+		FConnectionInfo& ConnectionInfo = *It;
+		if (ConnectionInfo.ConnectionId != ConnectionId)
+		{
+			continue;
+		}
+
+		ReplicationSystem->RemoveConnection(ConnectionInfo.ConnectionId);
+
+		if (IsValid(ConnectionInfo.DataStreamManager))
+		{
+			ConnectionInfo.DataStreamManager->Deinit();
+			ConnectionInfo.DataStreamManager->MarkAsGarbage();
+		}
+
+		UE_NET_TRACE_CONNECTION_CLOSED(GetNetTraceId(), ConnectionId);
+
+		It.RemoveCurrent();
+		break;
+	}
+}
+
 void FReplicationSystemTestNode::PreSendUpdate(const UReplicationSystem::FSendUpdateParams& Params)
 {
 	ReplicationSystem->PreSendUpdate(Params);
@@ -505,4 +530,17 @@ FReplicationSystemTestClient* FReplicationSystemServerClientTestFixture::CreateC
 	return Client;
 }
 
+void FReplicationSystemServerClientTestFixture::DestroyClient(FReplicationSystemTestClient* Client)
+{
+	if (!Clients.Remove(Client))
+	{
+		UE_LOG(LogIris, Warning, TEXT("Unable to find FReplicationSystemTestClient %p for destroy. NOT destroying."), Client);
+		return;
+	}
+
+	Server->RemoveConnection(Client->ConnectionIdOnServer);
+
+	delete Client;
+}
+
 }
diff --git a/Engine/Plugins/Runtime/ReplicationSystemTestPlugin/Source/Private/Tests/ReplicationSystem/ReplicationSystemServerClientTestFixture.h b/Engine/Plugins/Runtime/ReplicationSystemTestPlugin/Source/Private/Tests/ReplicationSystem/ReplicationSystemServerClientTestFixture.h
--- a/Engine/Plugins/Runtime/ReplicationSystemTestPlugin/Source/Private/Tests/ReplicationSystem/ReplicationSystemServerClientTestFixture.h
+++ b/Engine/Plugins/Runtime/ReplicationSystemTestPlugin/Source/Private/Tests/ReplicationSystem/ReplicationSystemServerClientTestFixture.h
@@ -65,13 +65,13 @@ public:
 
 	struct FConnectionInfo
 	{
-		UDataStreamManager* DataStreamManager;
-		UReplicationDataStream* ReplicationDataStream;
-		UNetTokenDataStream* NetTokenDataStream;
-		FNetTokenStoreState* RemoteNetTokenStoreState;
+		UDataStreamManager* DataStreamManager = nullptr;
+		UReplicationDataStream* ReplicationDataStream = nullptr;
+		UNetTokenDataStream* NetTokenDataStream = nullptr;
+		FNetTokenStoreState* RemoteNetTokenStoreState = nullptr;
 		TResizableCircularQueue<const FDataStreamRecord*> WriteRecords;
 		TResizableCircularQueue<FPacketData> WrittenPackets;
-		uint32 ConnectionId;
+		uint32 ConnectionId = 0;
 	};
 
 	struct FReplicationSystemParamsOverride
@@ -86,7 +86,7 @@ public:
 public:
 	FReplicationSystemTestNode(bool bIsServer, const TCHAR* Name);
 	explicit FReplicationSystemTestNode(FReplicationSystemTestNode::EDelaySetup);
-	~FReplicationSystemTestNode();
+	virtual ~FReplicationSystemTestNode();
 
 	void Setup(bool bIsServer, const TCHAR* Name, FReplicationSystemTestNode::FReplicationSystemParamsOverride* ParamsOverride=nullptr);
 
@@ -139,6 +139,7 @@ public:
 
 	// Connection
 	uint32 AddConnection();
+	void RemoveConnection(uint32 ConnectionId);
 
 	// System Update
 	void PreSendUpdate(const UReplicationSystem::FSendUpdateParams& Params);
@@ -221,9 +222,11 @@ protected:
 	};
 
 	virtual void SetUp() override;
-	FReplicationSystemTestClient* CreateClient();
 	virtual void TearDown() override;
 
+	FReplicationSystemTestClient* CreateClient();
+	void DestroyClient(FReplicationSystemTestClient* Client);
+
 	FDataStreamTestUtil DataStreamUtil;
 	FReplicationSystemTestServer* Server;
 	TArray<FReplicationSystemTestClient*> Clients;
diff --git a/Engine/Plugins/Runtime/ReplicationSystemTestPlugin/Source/Private/Tests/ReplicationSystem/ReplicationSystemTestFixture.h b/Engine/Plugins/Runtime/ReplicationSystemTestPlugin/Source/Private/Tests/ReplicationSystem/ReplicationSystemTestFixture.h
--- a/Engine/Plugins/Runtime/ReplicationSystemTestPlugin/Source/Private/Tests/ReplicationSystem/ReplicationSystemTestFixture.h
+++ b/Engine/Plugins/Runtime/ReplicationSystemTestPlugin/Source/Private/Tests/ReplicationSystem/ReplicationSystemTestFixture.h
@@ -104,7 +104,6 @@ protected:
 
 	UReplicationSystem* ReplicationSystem;
 	UReplicatedTestObjectBridge* ReplicationBridge;
-	Private::FReplicationProtocolManager ReplicationProtocolManager;
 	
 	TArray<TStrongObjectPtr<UObject>> CreatedObjects;
 };
diff --git a/Engine/Source/Runtime/Experimental/Iris/Core/Private/Iris/ReplicationSystem/ObjectReplicationBridge.cpp b/Engine/Source/Runtime/Experimental/Iris/Core/Private/Iris/ReplicationSystem/ObjectReplicationBridge.cpp
--- a/Engine/Source/Runtime/Experimental/Iris/Core/Private/Iris/ReplicationSystem/ObjectReplicationBridge.cpp
+++ b/Engine/Source/Runtime/Experimental/Iris/Core/Private/Iris/ReplicationSystem/ObjectReplicationBridge.cpp
@@ -936,11 +936,21 @@ void UObjectReplicationBridge::SetInstancePreUpdateFunction(FInstancePreUpdateFu
 	PreUpdateInstanceFunction = InPreUpdateFunction;
 }
 
+UObjectReplicationBridge::FInstancePreUpdateFunction UObjectReplicationBridge::GetInstancePreUpdateFunction() const
+{
+	return PreUpdateInstanceFunction;
+}
+
 void UObjectReplicationBridge::SetInstanceGetWorldObjectInfoFunction(FInstanceGetWorldObjectInfoFunction InGetWorldObjectInfoFunction)
 {
 	GetInstanceWorldObjectInfoFunction = InGetWorldObjectInfoFunction;
 }
 
+UObjectReplicationBridge::FInstanceGetWorldObjectInfoFunction UObjectReplicationBridge::GetInstanceGetWorldObjectInfoFunction() const
+{
+	return GetInstanceWorldObjectInfoFunction;
+}
+
 void UObjectReplicationBridge::ForcePollObject(FNetRefHandle Handle)
 {
 	using namespace UE::Net::Private;
diff --git a/Engine/Source/Runtime/Experimental/Iris/Core/Public/Iris/ReplicationSystem/ObjectReplicationBridge.h b/Engine/Source/Runtime/Experimental/Iris/Core/Public/Iris/ReplicationSystem/ObjectReplicationBridge.h
--- a/Engine/Source/Runtime/Experimental/Iris/Core/Public/Iris/ReplicationSystem/ObjectReplicationBridge.h
+++ b/Engine/Source/Runtime/Experimental/Iris/Core/Public/Iris/ReplicationSystem/ObjectReplicationBridge.h
@@ -279,12 +279,18 @@ protected:
 	/** Set the function that we should call before copying state data. */
 	IRISCORE_API void SetInstancePreUpdateFunction(FInstancePreUpdateFunction InPreUpdateFunction);
 	
+	/** Get the function that we should call before copying state data. */
+	IRISCORE_API FInstancePreUpdateFunction GetInstancePreUpdateFunction() const;
+
 	/** Helper method to get the world location & cull distance for replicated instances with the HasWorldLocation trait. */
 	using FInstanceGetWorldObjectInfoFunction = TFunction<void(UE::Net::FNetRefHandle, const UObject*, FVector&, float&)>;
 
 	/** Set the function that we should call to get the world location of an object. */
 	IRISCORE_API void SetInstanceGetWorldObjectInfoFunction(FInstanceGetWorldObjectInfoFunction InGetWorldObjectInfoFunction);
 	
+	/** Get the function that we should call to get the world location of an object. */
+	IRISCORE_API FInstanceGetWorldObjectInfoFunction GetInstanceGetWorldObjectInfoFunction() const;
+
 	// Poll frequency support
 
 	/** Force polling of Object when ObjectToPollWith is polled. */
