Iris
- Fixed failed test due to changes to validate function

diff --git a/Engine/Plugins/Runtime/ReplicationSystemTestPlugin/Source/Private/Tests/Serialization/TestRotatorNetSerializer.cpp b/Engine/Plugins/Runtime/ReplicationSystemTestPlugin/Source/Private/Tests/Serialization/TestRotatorNetSerializer.cpp
--- a/Engine/Plugins/Runtime/ReplicationSystemTestPlugin/Source/Private/Tests/Serialization/TestRotatorNetSerializer.cpp
+++ b/Engine/Plugins/Runtime/ReplicationSystemTestPlugin/Source/Private/Tests/Serialization/TestRotatorNetSerializer.cpp
@@ -67,12 +67,6 @@ const typename FTestRotatorNetSerializerBase<RotatorType>::FFloatTriplet FTestRo
 	{-INFINITY, INFINITY, -INFINITY},
 	{NAN, NAN, NAN},
 	{-NAN, -NAN, -NAN},
-	{-1.0f, 0.0f, 0.0f},
-	{0.0f, -1.0f, 0.0f},
-	{0.0f, 0.0f, -1.0f},
-	{5.0f, 6.0f, 360.0f},
-	{std::numeric_limits<ScalarType>::max(), std::numeric_limits<ScalarType>::max(), std::numeric_limits<ScalarType>::max()},
-	{std::numeric_limits<ScalarType>::lowest(), std::numeric_limits<ScalarType>::lowest(), std::numeric_limits<ScalarType>::lowest()},
 };
 
 template<typename RotatorType>
