Iris
* CreateGroup public API now requires an FName
* Passing NAME_None will let Iris autogenerate an unique name to the group (not recommended)
* Ensure and don't create a group if a valid name was passed but it's not unique
* Improve LogIrisFiltering verbose logs to output group names

diff --git a/Engine/Plugins/Runtime/ReplicationSystemTestPlugin/Source/Private/Tests/ReplicationSystem/Filtering/TestFiltering.cpp b/Engine/Plugins/Runtime/ReplicationSystemTestPlugin/Source/Private/Tests/ReplicationSystem/Filtering/TestFiltering.cpp
--- a/Engine/Plugins/Runtime/ReplicationSystemTestPlugin/Source/Private/Tests/ReplicationSystem/Filtering/TestFiltering.cpp
+++ b/Engine/Plugins/Runtime/ReplicationSystemTestPlugin/Source/Private/Tests/ReplicationSystem/Filtering/TestFiltering.cpp
@@ -552,7 +552,7 @@ UE_NET_TEST_FIXTURE(FTestFilteringFixture, ConnectionFilterAllowsObjectToReplica
 	UReplicatedTestObject* ServerObject0 = Server->CreateObject(0, 0);
 	
 	// Setup group filter
-	FNetObjectGroupHandle GroupHandle = Server->ReplicationSystem->CreateGroup();
+	FNetObjectGroupHandle GroupHandle = Server->ReplicationSystem->CreateGroup(NAME_None);
 	Server->ReplicationSystem->AddToGroup(GroupHandle, ServerObject0->NetRefHandle);
 	Server->ReplicationSystem->AddExclusionFilterGroup(GroupHandle);
 
@@ -816,7 +816,7 @@ UE_NET_TEST_FIXTURE(FTestFilteringFixture, GroupFilterPreventsObjectFromReplicat
 	// Spawn object on server
 	UReplicatedTestObject* ServerObject = Server->CreateObject(0, 0);
 
-	FNetObjectGroupHandle GroupHandle = Server->ReplicationSystem->CreateGroup();
+	FNetObjectGroupHandle GroupHandle = Server->ReplicationSystem->CreateGroup(NAME_None);
 	Server->ReplicationSystem->AddToGroup(GroupHandle, ServerObject->NetRefHandle);
 	Server->ReplicationSystem->AddExclusionFilterGroup(GroupHandle);
 
@@ -840,7 +840,7 @@ UE_NET_TEST_FIXTURE(FTestFilteringFixture, GroupFilterAllowsObjectToReplicate)
 	// Spawn object on server
 	UReplicatedTestObject* ServerObject = Server->CreateObject(0, 0);
 
-	FNetObjectGroupHandle GroupHandle = Server->ReplicationSystem->CreateGroup();
+	FNetObjectGroupHandle GroupHandle = Server->ReplicationSystem->CreateGroup(NAME_None);
 	Server->ReplicationSystem->AddToGroup(GroupHandle, ServerObject->NetRefHandle);
 	Server->ReplicationSystem->AddExclusionFilterGroup(GroupHandle);
 
@@ -901,7 +901,7 @@ UE_NET_TEST_FIXTURE(FTestFilteringFixture, GroupFilterRestoresConnectionConnecti
 	}
 
 	// Create and set group filter for last client only
-	FNetObjectGroupHandle GroupHandle = Server->ReplicationSystem->CreateGroup();
+	FNetObjectGroupHandle GroupHandle = Server->ReplicationSystem->CreateGroup(NAME_None);
 	Server->ReplicationSystem->AddToGroup(GroupHandle, ServerObject->NetRefHandle);
 	Server->ReplicationSystem->AddExclusionFilterGroup(GroupHandle);
 	Server->ReplicationSystem->SetGroupFilterStatus(GroupHandle, ConnectionIdForLastClient, ENetFilterStatus::Disallow);
@@ -1382,7 +1382,7 @@ UE_NET_TEST_FIXTURE(FTestFilteringFixture, InclusionGroupDoesNotFilterOutObject)
 	UReplicatedTestObject* ServerObject = Server->CreateObject(UTestReplicatedIrisObject::FComponents{});
 	
 	// Setup inclusion group filter
-	FNetObjectGroupHandle GroupHandle = Server->ReplicationSystem->CreateGroup();
+	FNetObjectGroupHandle GroupHandle = Server->ReplicationSystem->CreateGroup(NAME_None);
 	Server->ReplicationSystem->AddToGroup(GroupHandle, ServerObject->NetRefHandle);
 	Server->ReplicationSystem->AddInclusionFilterGroup(GroupHandle);
 	Server->ReplicationSystem->SetGroupFilterStatus(GroupHandle, ENetFilterStatus::Allow);
@@ -1421,7 +1421,7 @@ UE_NET_TEST_FIXTURE(FTestFilteringFixture, InclusionGroupDoesNotOverrideOwnerFil
 	Server->ReplicationSystem->SetOwningNetConnection(ServerObject->NetRefHandle, ClientArray[OwningClientIndex]->ConnectionIdOnServer);
 
 	// Setup inclusion group filter
-	FNetObjectGroupHandle GroupHandle = Server->ReplicationSystem->CreateGroup();
+	FNetObjectGroupHandle GroupHandle = Server->ReplicationSystem->CreateGroup(NAME_None);
 	Server->ReplicationSystem->AddToGroup(GroupHandle, ServerObject->NetRefHandle);
 	Server->ReplicationSystem->AddInclusionFilterGroup(GroupHandle);
 	Server->ReplicationSystem->SetGroupFilterStatus(GroupHandle, ENetFilterStatus::Allow);
@@ -1467,7 +1467,7 @@ UE_NET_TEST_FIXTURE(FTestFilteringFixture, InclusionGroupDoesNotOverrideConnecti
 	Server->ReplicationSystem->SetConnectionFilter(ServerObject->NetRefHandle, AllowedConnections, ENetFilterStatus::Allow);
 
 	// Setup inclusion group filter
-	FNetObjectGroupHandle GroupHandle = Server->ReplicationSystem->CreateGroup();
+	FNetObjectGroupHandle GroupHandle = Server->ReplicationSystem->CreateGroup(NAME_None);
 	Server->ReplicationSystem->AddToGroup(GroupHandle, ServerObject->NetRefHandle);
 	Server->ReplicationSystem->AddInclusionFilterGroup(GroupHandle);
 	Server->ReplicationSystem->SetGroupFilterStatus(GroupHandle, ENetFilterStatus::Allow);
@@ -1507,13 +1507,13 @@ UE_NET_TEST_FIXTURE(FTestFilteringFixture, InclusionGroupDoesNotOverrideExclusio
 	UReplicatedTestObject* ServerObject = Server->CreateObject(UTestReplicatedIrisObject::FComponents{});
 
 	// Setup exclusion group filter
-	FNetObjectGroupHandle ExclusionGroupHandle = Server->ReplicationSystem->CreateGroup();
+	FNetObjectGroupHandle ExclusionGroupHandle = Server->ReplicationSystem->CreateGroup(NAME_None);
 	Server->ReplicationSystem->AddToGroup(ExclusionGroupHandle, ServerObject->NetRefHandle);
 	Server->ReplicationSystem->AddExclusionFilterGroup(ExclusionGroupHandle);
 	Server->ReplicationSystem->SetGroupFilterStatus(ExclusionGroupHandle, ClientArray[AllowedClientIndex]->ConnectionIdOnServer, ENetFilterStatus::Allow);
 
 	// Setup inclusion group filter
-	FNetObjectGroupHandle InclusionGroupHandle = Server->ReplicationSystem->CreateGroup();
+	FNetObjectGroupHandle InclusionGroupHandle = Server->ReplicationSystem->CreateGroup(NAME_None);
 	Server->ReplicationSystem->AddToGroup(InclusionGroupHandle, ServerObject->NetRefHandle);
 	Server->ReplicationSystem->AddInclusionFilterGroup(InclusionGroupHandle);
 	Server->ReplicationSystem->SetGroupFilterStatus(InclusionGroupHandle, ENetFilterStatus::Allow);
@@ -1561,7 +1561,7 @@ UE_NET_TEST_FIXTURE(FTestFilteringFixture, DisabledInclusionGroupDoesNotOverride
 	Server->ReplicationSystem->SetFilter(ServerObject->NetRefHandle, MockFilterHandle);
 
 	// Setup inclusion group filter
-	FNetObjectGroupHandle InclusionGroupHandle = Server->ReplicationSystem->CreateGroup();
+	FNetObjectGroupHandle InclusionGroupHandle = Server->ReplicationSystem->CreateGroup(NAME_None);
 	Server->ReplicationSystem->AddToGroup(InclusionGroupHandle, ServerObject->NetRefHandle);
 	Server->ReplicationSystem->AddInclusionFilterGroup(InclusionGroupHandle);
 	Server->ReplicationSystem->SetGroupFilterStatus(InclusionGroupHandle, ENetFilterStatus::Disallow);
@@ -1598,7 +1598,7 @@ UE_NET_TEST_FIXTURE(FTestFilteringFixture, EnabledInclusionGroupDoesOverrideDyna
 	Server->ReplicationSystem->SetFilter(ServerObject->NetRefHandle, MockFilterHandle);
 
 	// Setup inclusion group filter
-	FNetObjectGroupHandle InclusionGroupHandle = Server->ReplicationSystem->CreateGroup();
+	FNetObjectGroupHandle InclusionGroupHandle = Server->ReplicationSystem->CreateGroup(NAME_None);
 	Server->ReplicationSystem->AddToGroup(InclusionGroupHandle, ServerObject->NetRefHandle);
 	Server->ReplicationSystem->AddInclusionFilterGroup(InclusionGroupHandle);
 	// Disallow by default
@@ -1637,7 +1637,7 @@ UE_NET_TEST_FIXTURE(FTestFilteringFixture, LateAddingToEnabledInclusionGroupDoes
 	Server->ReplicationSystem->SetFilter(ServerObject->NetRefHandle, MockFilterHandle);
 
 	// Setup inclusion group filter
-	FNetObjectGroupHandle InclusionGroupHandle = Server->ReplicationSystem->CreateGroup();
+	FNetObjectGroupHandle InclusionGroupHandle = Server->ReplicationSystem->CreateGroup(NAME_None);
 	Server->ReplicationSystem->AddInclusionFilterGroup(InclusionGroupHandle);
 	// Disallow by default
 	Server->ReplicationSystem->SetGroupFilterStatus(InclusionGroupHandle, ENetFilterStatus::Disallow);
@@ -1687,7 +1687,7 @@ UE_NET_TEST_FIXTURE(FTestFilteringFixture, LateEnablingInclusionGroupDoesOverrid
 	Server->ReplicationSystem->SetFilter(ServerObject->NetRefHandle, MockFilterHandle);
 
 	// Setup inclusion group filter
-	FNetObjectGroupHandle InclusionGroupHandle = Server->ReplicationSystem->CreateGroup();
+	FNetObjectGroupHandle InclusionGroupHandle = Server->ReplicationSystem->CreateGroup(NAME_None);
 	Server->ReplicationSystem->AddToGroup(InclusionGroupHandle, ServerObject->NetRefHandle);
 	Server->ReplicationSystem->AddInclusionFilterGroup(InclusionGroupHandle);
 	// Disallow by default
@@ -1729,7 +1729,7 @@ UE_NET_TEST_FIXTURE(FTestFilteringFixture, RemovingFromInclusionGroupRemovesDyna
 	Server->ReplicationSystem->SetFilter(ServerObject->NetRefHandle, MockFilterHandle);
 
 	// Setup inclusion group filter
-	FNetObjectGroupHandle InclusionGroupHandle = Server->ReplicationSystem->CreateGroup();
+	FNetObjectGroupHandle InclusionGroupHandle = Server->ReplicationSystem->CreateGroup(NAME_None);
 	Server->ReplicationSystem->AddToGroup(InclusionGroupHandle, ServerObject->NetRefHandle);
 	Server->ReplicationSystem->AddInclusionFilterGroup(InclusionGroupHandle);
 	Server->ReplicationSystem->SetGroupFilterStatus(InclusionGroupHandle, ENetFilterStatus::Allow);
@@ -1766,7 +1766,7 @@ UE_NET_TEST_FIXTURE(FTestFilteringFixture, SubObjectAddedToAllowedInclusionGroup
 	Server->ReplicationSystem->SetFilter(ServerObject->NetRefHandle, MockFilterHandle);
 
 	// Setup inclusion group filter
-	FNetObjectGroupHandle InclusionGroupHandle = Server->ReplicationSystem->CreateGroup();
+	FNetObjectGroupHandle InclusionGroupHandle = Server->ReplicationSystem->CreateGroup(NAME_None);
 	Server->ReplicationSystem->AddToGroup(InclusionGroupHandle, ServerSubObject->NetRefHandle);
 	Server->ReplicationSystem->AddInclusionFilterGroup(InclusionGroupHandle);
 	Server->ReplicationSystem->SetGroupFilterStatus(InclusionGroupHandle, Client->ConnectionIdOnServer, ENetFilterStatus::Allow);
@@ -1795,7 +1795,7 @@ UE_NET_TEST_FIXTURE(FTestFilteringFixture, SubObjectLateAddedToAllowedInclusionG
 	Server->ReplicationSystem->SetFilter(ServerObject->NetRefHandle, MockFilterHandle);
 
 	// Setup inclusion group filter
-	FNetObjectGroupHandle InclusionGroupHandle = Server->ReplicationSystem->CreateGroup();
+	FNetObjectGroupHandle InclusionGroupHandle = Server->ReplicationSystem->CreateGroup(NAME_None);
 	Server->ReplicationSystem->AddInclusionFilterGroup(InclusionGroupHandle);
 	Server->ReplicationSystem->SetGroupFilterStatus(InclusionGroupHandle, Client->ConnectionIdOnServer, ENetFilterStatus::Allow);
 
@@ -1828,7 +1828,7 @@ UE_NET_TEST_FIXTURE(FTestFilteringFixture, LateAddedSubObjectFollowsOwnerInAllow
 	Server->ReplicationSystem->SetFilter(ServerObject->NetRefHandle, MockFilterHandle);
 
 	// Setup inclusion group filter and add object
-	FNetObjectGroupHandle InclusionGroupHandle = Server->ReplicationSystem->CreateGroup();
+	FNetObjectGroupHandle InclusionGroupHandle = Server->ReplicationSystem->CreateGroup(NAME_None);
 	Server->ReplicationSystem->AddInclusionFilterGroup(InclusionGroupHandle);
 	Server->ReplicationSystem->AddToGroup(InclusionGroupHandle, ServerObject->NetRefHandle);
 	Server->ReplicationSystem->SetGroupFilterStatus(InclusionGroupHandle, Client->ConnectionIdOnServer, ENetFilterStatus::Allow);
@@ -1864,7 +1864,7 @@ UE_NET_TEST_FIXTURE(FTestFilteringFixture, LateAddedSubObjectFollowsOwnerInDisal
 	Server->ReplicationSystem->SetFilter(ServerObject->NetRefHandle, MockFilterHandle);
 
 	// Setup inclusion group filter and add object
-	FNetObjectGroupHandle InclusionGroupHandle = Server->ReplicationSystem->CreateGroup();
+	FNetObjectGroupHandle InclusionGroupHandle = Server->ReplicationSystem->CreateGroup(NAME_None);
 	Server->ReplicationSystem->AddInclusionFilterGroup(InclusionGroupHandle);
 	Server->ReplicationSystem->AddToGroup(InclusionGroupHandle, ServerObject->NetRefHandle);
 	Server->ReplicationSystem->SetGroupFilterStatus(InclusionGroupHandle, Client->ConnectionIdOnServer, ENetFilterStatus::Disallow);
@@ -1901,7 +1901,7 @@ UE_NET_TEST_FIXTURE(FTestFilteringFixture, SubObjectAddedToDisallowedInclusionGr
 	Server->ReplicationSystem->SetFilter(ServerObject->NetRefHandle, MockFilterHandle);
 
 	// Setup inclusion group filter
-	FNetObjectGroupHandle InclusionGroupHandle = Server->ReplicationSystem->CreateGroup();
+	FNetObjectGroupHandle InclusionGroupHandle = Server->ReplicationSystem->CreateGroup(NAME_None);
 	Server->ReplicationSystem->AddToGroup(InclusionGroupHandle, ServerSubObject->NetRefHandle);
 	Server->ReplicationSystem->AddInclusionFilterGroup(InclusionGroupHandle);
 	Server->ReplicationSystem->SetGroupFilterStatus(InclusionGroupHandle, Client->ConnectionIdOnServer, ENetFilterStatus::Disallow);
@@ -1930,7 +1930,7 @@ UE_NET_TEST_FIXTURE(FTestFilteringFixture, SubObjectRemovedFromAllowedInclusionG
 	Server->ReplicationSystem->SetFilter(ServerObject->NetRefHandle, MockFilterHandle);
 
 	// Setup inclusion group filter
-	FNetObjectGroupHandle InclusionGroupHandle = Server->ReplicationSystem->CreateGroup();
+	FNetObjectGroupHandle InclusionGroupHandle = Server->ReplicationSystem->CreateGroup(NAME_None);
 	Server->ReplicationSystem->AddToGroup(InclusionGroupHandle, ServerSubObject->NetRefHandle);
 	Server->ReplicationSystem->AddInclusionFilterGroup(InclusionGroupHandle);
 	Server->ReplicationSystem->SetGroupFilterStatus(InclusionGroupHandle, Client->ConnectionIdOnServer, ENetFilterStatus::Allow);
@@ -1965,7 +1965,7 @@ UE_NET_TEST_FIXTURE(FTestFilteringFixture, SubObjectRemovedFromDisallowedInclusi
 	Server->ReplicationSystem->SetFilter(ServerObject->NetRefHandle, MockFilterHandle);
 
 	// Setup inclusion group filter
-	FNetObjectGroupHandle InclusionGroupHandle = Server->ReplicationSystem->CreateGroup();
+	FNetObjectGroupHandle InclusionGroupHandle = Server->ReplicationSystem->CreateGroup(NAME_None);
 	Server->ReplicationSystem->AddToGroup(InclusionGroupHandle, ServerSubObject->NetRefHandle);
 	Server->ReplicationSystem->AddInclusionFilterGroup(InclusionGroupHandle);
 	Server->ReplicationSystem->SetGroupFilterStatus(InclusionGroupHandle, Client->ConnectionIdOnServer, ENetFilterStatus::Disallow);
@@ -2000,12 +2000,12 @@ UE_NET_TEST_FIXTURE(FTestFilteringFixture, SubObjectAddedToInclusionGroupFollows
 	Server->ReplicationSystem->SetFilter(ServerObject->NetRefHandle, MockFilterHandle);
 
 	// Setup separate inclusion group filters for object and subobject
-	FNetObjectGroupHandle ObjectInclusionGroupHandle = Server->ReplicationSystem->CreateGroup();
+	FNetObjectGroupHandle ObjectInclusionGroupHandle = Server->ReplicationSystem->CreateGroup(NAME_None);
 	Server->ReplicationSystem->AddToGroup(ObjectInclusionGroupHandle, ServerObject->NetRefHandle);
 	Server->ReplicationSystem->AddInclusionFilterGroup(ObjectInclusionGroupHandle);
 	Server->ReplicationSystem->SetGroupFilterStatus(ObjectInclusionGroupHandle, Client->ConnectionIdOnServer, ENetFilterStatus::Disallow);
 	
-	FNetObjectGroupHandle SubObjectInclusionGroupHandle = Server->ReplicationSystem->CreateGroup();
+	FNetObjectGroupHandle SubObjectInclusionGroupHandle = Server->ReplicationSystem->CreateGroup(NAME_None);
 	Server->ReplicationSystem->AddToGroup(SubObjectInclusionGroupHandle, ServerSubObject->NetRefHandle);
 	Server->ReplicationSystem->AddInclusionFilterGroup(SubObjectInclusionGroupHandle);
 	Server->ReplicationSystem->SetGroupFilterStatus(SubObjectInclusionGroupHandle, Client->ConnectionIdOnServer, ENetFilterStatus::Disallow);
@@ -2055,12 +2055,12 @@ UE_NET_TEST_FIXTURE(FTestFilteringFixture, SubObjectLateAddedToInclusionGroupFol
 	Server->ReplicationSystem->SetFilter(ServerObject->NetRefHandle, MockFilterHandle);
 
 	// Setup separate inclusion group filters for object and subobject
-	FNetObjectGroupHandle ObjectInclusionGroupHandle = Server->ReplicationSystem->CreateGroup();
+	FNetObjectGroupHandle ObjectInclusionGroupHandle = Server->ReplicationSystem->CreateGroup(NAME_None);
 	Server->ReplicationSystem->AddToGroup(ObjectInclusionGroupHandle, ServerObject->NetRefHandle);
 	Server->ReplicationSystem->AddInclusionFilterGroup(ObjectInclusionGroupHandle);
 	Server->ReplicationSystem->SetGroupFilterStatus(ObjectInclusionGroupHandle, Client->ConnectionIdOnServer, ENetFilterStatus::Disallow);
 	
-	FNetObjectGroupHandle SubObjectInclusionGroupHandle = Server->ReplicationSystem->CreateGroup();
+	FNetObjectGroupHandle SubObjectInclusionGroupHandle = Server->ReplicationSystem->CreateGroup(NAME_None);
 	Server->ReplicationSystem->AddInclusionFilterGroup(SubObjectInclusionGroupHandle);
 	Server->ReplicationSystem->SetGroupFilterStatus(SubObjectInclusionGroupHandle, Client->ConnectionIdOnServer, ENetFilterStatus::Allow);
 
@@ -2094,12 +2094,12 @@ UE_NET_TEST_FIXTURE(FTestFilteringFixture, SubObjectRemovedFromInclusionGroupFol
 	Server->ReplicationSystem->SetFilter(ServerObject->NetRefHandle, MockFilterHandle);
 
 	// Setup separate inclusion group filters for object and subobject
-	FNetObjectGroupHandle ObjectInclusionGroupHandle = Server->ReplicationSystem->CreateGroup();
+	FNetObjectGroupHandle ObjectInclusionGroupHandle = Server->ReplicationSystem->CreateGroup(NAME_None);
 	Server->ReplicationSystem->AddToGroup(ObjectInclusionGroupHandle, ServerObject->NetRefHandle);
 	Server->ReplicationSystem->AddInclusionFilterGroup(ObjectInclusionGroupHandle);
 	Server->ReplicationSystem->SetGroupFilterStatus(ObjectInclusionGroupHandle, Client->ConnectionIdOnServer, ENetFilterStatus::Disallow);
 	
-	FNetObjectGroupHandle SubObjectInclusionGroupHandle = Server->ReplicationSystem->CreateGroup();
+	FNetObjectGroupHandle SubObjectInclusionGroupHandle = Server->ReplicationSystem->CreateGroup(NAME_None);
 	Server->ReplicationSystem->AddToGroup(SubObjectInclusionGroupHandle, ServerSubObject->NetRefHandle);
 	Server->ReplicationSystem->AddInclusionFilterGroup(SubObjectInclusionGroupHandle);
 	Server->ReplicationSystem->SetGroupFilterStatus(SubObjectInclusionGroupHandle, Client->ConnectionIdOnServer, ENetFilterStatus::Disallow);
@@ -2182,7 +2182,7 @@ UE_NET_TEST_FIXTURE(FTestFilteringFixture, InclusionGroupsWorksWithMultipleObjec
 	for (FNetObjectGroupHandle& InclusionGroupHandle : InclusionGroupHandles)
 	{
 		const SIZE_T Index = &InclusionGroupHandle - &InclusionGroupHandles[0];
-		InclusionGroupHandle = Server->ReplicationSystem->CreateGroup();
+		InclusionGroupHandle = Server->ReplicationSystem->CreateGroup(NAME_None);
 		Server->ReplicationSystem->AddInclusionFilterGroup(InclusionGroupHandle);
 		Server->ReplicationSystem->AddToGroup(InclusionGroupHandle, ServerObjects[Index]->NetRefHandle);
 		// Disallow by default
@@ -2230,7 +2230,7 @@ UE_NET_TEST_FIXTURE(FTestFilteringFixture, InclusionGroupsAreCumulative)
 	for (FNetObjectGroupHandle& InclusionGroupHandle : InclusionGroupHandles)
 	{
 		const SIZE_T Index = &InclusionGroupHandle - &InclusionGroupHandles[0];
-		InclusionGroupHandle = Server->ReplicationSystem->CreateGroup();
+		InclusionGroupHandle = Server->ReplicationSystem->CreateGroup(NAME_None);
 		Server->ReplicationSystem->AddInclusionFilterGroup(InclusionGroupHandle);
 		Server->ReplicationSystem->AddToGroup(InclusionGroupHandle, ServerObjects[Index]->NetRefHandle);
 		// Disallow by default
@@ -2274,7 +2274,7 @@ UE_NET_TEST_FIXTURE(FTestFilteringFixture, LateAddedConnectionWorksWithSimpleGro
 	for (FNetObjectGroupHandle& InclusionGroupHandle : InclusionGroupHandles)
 	{
 		const SIZE_T Index = &InclusionGroupHandle - &InclusionGroupHandles[0];
-		InclusionGroupHandle = Server->ReplicationSystem->CreateGroup();
+		InclusionGroupHandle = Server->ReplicationSystem->CreateGroup(NAME_None);
 		Server->ReplicationSystem->AddInclusionFilterGroup(InclusionGroupHandle);
 		Server->ReplicationSystem->AddToGroup(InclusionGroupHandle, ServerObjects[Index]->NetRefHandle);
 
@@ -2347,11 +2347,11 @@ UE_NET_TEST_FIXTURE(FTestFilteringFixture, LateAddedConnectionWorksWithComplexGr
 	{
 		const SIZE_T Index = &ObjectInclusionGroupHandle - &ObjectInclusionGroupHandles[0];
 
-		ObjectInclusionGroupHandle = Server->ReplicationSystem->CreateGroup();
+		ObjectInclusionGroupHandle = Server->ReplicationSystem->CreateGroup(NAME_None);
 		Server->ReplicationSystem->AddInclusionFilterGroup(ObjectInclusionGroupHandle);
 		Server->ReplicationSystem->AddToGroup(ObjectInclusionGroupHandle, ServerObjects[Index]->NetRefHandle);
 
-		FNetObjectGroupHandle SubObjectInclusionGroupHandle = Server->ReplicationSystem->CreateGroup();
+		FNetObjectGroupHandle SubObjectInclusionGroupHandle = Server->ReplicationSystem->CreateGroup(NAME_None);
 		Server->ReplicationSystem->AddInclusionFilterGroup(SubObjectInclusionGroupHandle);
 		Server->ReplicationSystem->AddToGroup(SubObjectInclusionGroupHandle, ServerSubObjects[Index]->NetRefHandle);
 
@@ -2528,7 +2528,7 @@ UE_NET_TEST_FIXTURE(FTestFilteringFixture, HysteresisDoesNotKickInForExclusionGr
 	UReplicatedTestObject* ServerObject = Server->CreateObject({ .IrisComponentCount = 0 });
 
 	// Add to exclusion group that allows replication to all connections.
-	FNetObjectGroupHandle GroupHandle = Server->ReplicationSystem->CreateGroup();
+	FNetObjectGroupHandle GroupHandle = Server->ReplicationSystem->CreateGroup(NAME_None);
 	Server->ReplicationSystem->AddToGroup(GroupHandle, ServerObject->NetRefHandle);
 	Server->ReplicationSystem->AddExclusionFilterGroup(GroupHandle);
 	Server->ReplicationSystem->SetGroupFilterStatus(GroupHandle, ENetFilterStatus::Allow);
@@ -2879,7 +2879,7 @@ UE_NET_TEST_FIXTURE(FTestFilteringFixture, HysteresisKicksInForFormerlyInclusion
 	UReplicatedTestObject* ServerObject = Server->CreateObject({ .IrisComponentCount = 0 });
 	Server->ReplicationSystem->SetFilter(ServerObject->NetRefHandle, MockFilterHandle);
 
-	FNetObjectGroupHandle GroupHandle = Server->ReplicationSystem->CreateGroup();
+	FNetObjectGroupHandle GroupHandle = Server->ReplicationSystem->CreateGroup(NAME_None);
 	Server->ReplicationSystem->AddToGroup(GroupHandle, ServerObject->NetRefHandle);
 	Server->ReplicationSystem->AddInclusionFilterGroup(GroupHandle);
 	Server->ReplicationSystem->SetGroupFilterStatus(GroupHandle, ENetFilterStatus::Allow);
diff --git a/Engine/Plugins/Runtime/ReplicationSystemTestPlugin/Source/Private/Tests/ReplicationSystem/Filtering/TestGroups.cpp b/Engine/Plugins/Runtime/ReplicationSystemTestPlugin/Source/Private/Tests/ReplicationSystem/Filtering/TestGroups.cpp
--- a/Engine/Plugins/Runtime/ReplicationSystemTestPlugin/Source/Private/Tests/ReplicationSystem/Filtering/TestGroups.cpp
+++ b/Engine/Plugins/Runtime/ReplicationSystemTestPlugin/Source/Private/Tests/ReplicationSystem/Filtering/TestGroups.cpp
@@ -80,22 +80,78 @@ UE_NET_TEST_FIXTURE(FTestGroupsFixture, PublicGroupAPI)
 
 	// Create
 	{
-		FNetObjectGroupHandle GroupHandle = ReplicationSystem->CreateGroup();
+		FNetObjectGroupHandle GroupHandle = ReplicationSystem->CreateGroup(NAME_None);
+		
+		UE_NET_ASSERT_TRUE(GroupHandle.IsValid());
+		UE_NET_ASSERT_TRUE(ReplicationSystem->IsValidGroup(GroupHandle));
 
 		UE_NET_ASSERT_FALSE(ReplicationSystem->IsValidGroup(FNetObjectGroupHandle()));
-		UE_NET_ASSERT_TRUE(ReplicationSystem->IsValidGroup(GroupHandle));
 
 		// Destroy
 		ReplicationSystem->DestroyGroup(GroupHandle);
+
+		// Group handle is still initialized
+		UE_NET_ASSERT_TRUE(GroupHandle.IsValid());
+
+		// But it's not considered a valid group.
 		UE_NET_ASSERT_FALSE(ReplicationSystem->IsValidGroup(GroupHandle));
 	}
 
+	// Find
+	{
+		const FName NameGroupFind = FName(TEXT("TestGroupFind"));
+		
+		// Create a group
+		FNetObjectGroupHandle OriginalGroupHandle = ReplicationSystem->CreateGroup(NameGroupFind);
+		UE_NET_ASSERT_TRUE(OriginalGroupHandle.IsValid());
+		UE_NET_ASSERT_TRUE(ReplicationSystem->IsValidGroup(OriginalGroupHandle));
+
+		// Find the group handle via it's name
+		FNetObjectGroupHandle FoundGroupHandle = ReplicationSystem->FindGroup(NameGroupFind);
+		UE_NET_ASSERT_TRUE(FoundGroupHandle.IsValid());
+		UE_NET_ASSERT_TRUE(ReplicationSystem->IsValidGroup(FoundGroupHandle));
+
+		// Should be the same handle
+		UE_NET_ASSERT_EQ(FoundGroupHandle.GetRawValue(), OriginalGroupHandle.GetRawValue());
+
+		// Destroy the group
+		ReplicationSystem->DestroyGroup(OriginalGroupHandle);
+
+		// They both should be invalid
+		UE_NET_ASSERT_FALSE(ReplicationSystem->IsValidGroup(OriginalGroupHandle));
+		UE_NET_ASSERT_FALSE(ReplicationSystem->IsValidGroup(FoundGroupHandle));
+	}
+
+	// Destroy/Create: Does not pass at the moment
+	/*{
+		// Create a group
+		FNetObjectGroupHandle ReusedGroupHandle = ReplicationSystem->CreateGroup(NAME_None);
+		UE_NET_ASSERT_TRUE(ReusedGroupHandle.IsValid());
+		UE_NET_ASSERT_TRUE(ReplicationSystem->IsValidGroup(ReusedGroupHandle));
+
+		// Destroy the group
+		ReplicationSystem->DestroyGroup(ReusedGroupHandle);
+
+		// Should be invalid
+		UE_NET_ASSERT_FALSE(ReplicationSystem->IsValidGroup(ReusedGroupHandle));
+
+		{
+			// Create a new random group
+			FNetObjectGroupHandle RandomGroupHandle = ReplicationSystem->CreateGroup(NAME_None);
+			UE_NET_ASSERT_TRUE(ReplicationSystem->IsValidGroup(RandomGroupHandle));
+		}
+
+		// Problem: The destroyed group should still be invalid but the old handle index is now assigned to the new group.
+		// This means the old handle now points to the new group and is considered valid.
+		UE_NET_ASSERT_FALSE(ReplicationSystem->IsValidGroup(ReusedGroupHandle));
+	}*/
+
 	// Add/Remove
 	{
 		// Spawn object on server
 		UReplicatedTestObject* ServerObject = Server->CreateObject(0, 0);
 
-		FNetObjectGroupHandle GroupHandle = ReplicationSystem->CreateGroup();
+		FNetObjectGroupHandle GroupHandle = ReplicationSystem->CreateGroup(NAME_None);
 
 		UE_NET_ASSERT_FALSE(ReplicationSystem->IsInGroup(GroupHandle, ServerObject->NetRefHandle));
 
@@ -118,8 +174,8 @@ UE_NET_TEST_FIXTURE(FTestGroupsFixture, PublicGroupAPIMemberOfMultipleGroups)
 	UReplicationSystem* ReplicationSystem = Server->ReplicationSystem;
 
 	// Create
-	FNetObjectGroupHandle GroupHandle0 = ReplicationSystem->CreateGroup();
-	FNetObjectGroupHandle GroupHandle1 = ReplicationSystem->CreateGroup();
+	FNetObjectGroupHandle GroupHandle0 = ReplicationSystem->CreateGroup(NAME_None);
+	FNetObjectGroupHandle GroupHandle1 = ReplicationSystem->CreateGroup(NAME_None);
 
 	UReplicatedTestObject* ServerObject = Server->CreateObject(0, 0);
 
@@ -156,8 +212,8 @@ UE_NET_TEST_FIXTURE(FTestGroupsFixture, GroupFilterAPI)
 	UReplicatedTestObject* ServerObject2 = Server->CreateObject(0, 0);
 
 	// Create group
-	FNetObjectGroupHandle GroupHandle = Server->ReplicationSystem->CreateGroup();
-	UE_NET_ASSERT_TRUE(GroupHandle.IsValid());	
+	FNetObjectGroupHandle GroupHandle = Server->ReplicationSystem->CreateGroup(NAME_None);
+	UE_NET_ASSERT_TRUE(GroupHandle.IsValid());
 
 	// Add Objects to group
 	Server->ReplicationSystem->AddToGroup(GroupHandle, ServerObject0->NetRefHandle);
@@ -169,8 +225,8 @@ UE_NET_TEST_FIXTURE(FTestGroupsFixture, GroupFilterAPI)
 	UReplicatedTestObject* ServerObject5 = Server->CreateObject(0, 0);
 
 	// Create another group
-	FNetObjectGroupHandle GroupHandle2 = Server->ReplicationSystem->CreateGroup();
-	UE_NET_ASSERT_TRUE(GroupHandle2.IsValid());	
+	FNetObjectGroupHandle GroupHandle2 = Server->ReplicationSystem->CreateGroup(NAME_None);
+	UE_NET_ASSERT_TRUE(GroupHandle2.IsValid());
 
 	// Add Objects to group
 	Server->ReplicationSystem->AddToGroup(GroupHandle2, ServerObject3->NetRefHandle);
@@ -180,8 +236,8 @@ UE_NET_TEST_FIXTURE(FTestGroupsFixture, GroupFilterAPI)
 	UReplicatedTestObject* ServerObject6 = Server->CreateObject(0, 0);
 
 	// Create another group
-	FNetObjectGroupHandle GroupHandle3 = Server->ReplicationSystem->CreateGroup();
-	UE_NET_ASSERT_TRUE(GroupHandle3.IsValid());	
+	FNetObjectGroupHandle GroupHandle3 = Server->ReplicationSystem->CreateGroup(NAME_None);
+	UE_NET_ASSERT_TRUE(GroupHandle3.IsValid());
 
 	// Update filters
 	Filter();
@@ -274,7 +330,7 @@ UE_NET_TEST_FIXTURE(FTestGroupsFixture, GroupFilterAPIManyFilters)
 
 	for (uint32 It = 0; It < GroupCount; ++It)
 	{
-		FNetObjectGroupHandle GroupHandle = Server->ReplicationSystem->CreateGroup();
+		FNetObjectGroupHandle GroupHandle = Server->ReplicationSystem->CreateGroup(NAME_None);
 		Server->ReplicationSystem->AddExclusionFilterGroup(GroupHandle);
 	
 		GroupHandles[It] = GroupHandle;	
@@ -288,10 +344,10 @@ UE_NET_TEST_FIXTURE(FTestGroupsFixture, GroupFilterAPIObjectMemberOfMoreThanOneG
 	const FNetRefHandle ServerHandle = ServerObject0->NetRefHandle;
 
 	// Create groups
-	FNetObjectGroupHandle GroupHandle0 = Server->ReplicationSystem->CreateGroup();
-	FNetObjectGroupHandle GroupHandle1 = Server->ReplicationSystem->CreateGroup();
-	UE_NET_ASSERT_TRUE(GroupHandle0.IsValid());	
-	UE_NET_ASSERT_TRUE(GroupHandle1.IsValid());	
+	FNetObjectGroupHandle GroupHandle0 = Server->ReplicationSystem->CreateGroup(NAME_None);
+	FNetObjectGroupHandle GroupHandle1 = Server->ReplicationSystem->CreateGroup(NAME_None);
+	UE_NET_ASSERT_TRUE(GroupHandle0.IsValid());
+	UE_NET_ASSERT_TRUE(GroupHandle1.IsValid());
 
 	// Add Objects to groups
 	Server->ReplicationSystem->AddToGroup(GroupHandle0, ServerHandle);
@@ -346,10 +402,10 @@ UE_NET_TEST_FIXTURE(FTestGroupsFixture, GroupFilterAPIAddThenRemoveFilterTrait)
 	const FNetRefHandle ServerHandle = ServerObject0->NetRefHandle;
 
 	// Create groups
-	FNetObjectGroupHandle GroupHandle0 = Server->ReplicationSystem->CreateGroup();
-	FNetObjectGroupHandle GroupHandle1 = Server->ReplicationSystem->CreateGroup();
-	UE_NET_ASSERT_TRUE(GroupHandle0.IsValid());	
-	UE_NET_ASSERT_TRUE(GroupHandle1.IsValid());	
+	FNetObjectGroupHandle GroupHandle0 = Server->ReplicationSystem->CreateGroup(NAME_None);
+	FNetObjectGroupHandle GroupHandle1 = Server->ReplicationSystem->CreateGroup(NAME_None);
+	UE_NET_ASSERT_TRUE(GroupHandle0.IsValid());
+	UE_NET_ASSERT_TRUE(GroupHandle1.IsValid());
 
 	// Add Objects to groups
 	Server->ReplicationSystem->AddToGroup(GroupHandle0, ServerHandle);
@@ -407,11 +463,11 @@ UE_NET_TEST_FIXTURE(FTestGroupsFixture, GroupFilterTestObjectListViaTrait)
 	UE_NET_ASSERT_NE((uint32)FNetRefHandleManager::InvalidInternalIndex, ServerInternalIndex);
 
 	// Create groups
-	FNetObjectGroupHandle GroupHandle0 = Server->ReplicationSystem->CreateGroup();
-	FNetObjectGroupHandle GroupHandle1 = Server->ReplicationSystem->CreateGroup();
-	FNetObjectGroupHandle GroupHandle2 = Server->ReplicationSystem->CreateGroup();
-	UE_NET_ASSERT_TRUE(GroupHandle0.IsValid());	
-	UE_NET_ASSERT_TRUE(GroupHandle1.IsValid());	
+	FNetObjectGroupHandle GroupHandle0 = Server->ReplicationSystem->CreateGroup(NAME_None);
+	FNetObjectGroupHandle GroupHandle1 = Server->ReplicationSystem->CreateGroup(NAME_None);
+	FNetObjectGroupHandle GroupHandle2 = Server->ReplicationSystem->CreateGroup(NAME_None);
+	UE_NET_ASSERT_TRUE(GroupHandle0.IsValid());
+	UE_NET_ASSERT_TRUE(GroupHandle1.IsValid());
 	UE_NET_ASSERT_TRUE(GroupHandle2.IsValid());
 
 	// Add Objects to groups
@@ -477,12 +533,12 @@ UE_NET_TEST_FIXTURE(FTestGroupsFixture, GroupFilterTestObjectListViaMembership)
 	UE_NET_ASSERT_NE((uint32)FNetRefHandleManager::InvalidInternalIndex, ServerInternalIndex);
 
 	// Create groups
-	FNetObjectGroupHandle GroupHandle0 = Server->ReplicationSystem->CreateGroup();
-	FNetObjectGroupHandle GroupHandle1 = Server->ReplicationSystem->CreateGroup();
-	FNetObjectGroupHandle GroupHandle2 = Server->ReplicationSystem->CreateGroup();
-	UE_NET_ASSERT_TRUE(GroupHandle0.IsValid());	
-	UE_NET_ASSERT_TRUE(GroupHandle1.IsValid());	
-	UE_NET_ASSERT_TRUE(GroupHandle2.IsValid());	
+	FNetObjectGroupHandle GroupHandle0 = Server->ReplicationSystem->CreateGroup(NAME_None);
+	FNetObjectGroupHandle GroupHandle1 = Server->ReplicationSystem->CreateGroup(NAME_None);
+	FNetObjectGroupHandle GroupHandle2 = Server->ReplicationSystem->CreateGroup(NAME_None);
+	UE_NET_ASSERT_TRUE(GroupHandle0.IsValid());
+	UE_NET_ASSERT_TRUE(GroupHandle1.IsValid());
+	UE_NET_ASSERT_TRUE(GroupHandle2.IsValid());
 
 	// Add Objects to groups
 	Server->ReplicationSystem->AddToGroup(GroupHandle0, ServerHandle);
@@ -579,7 +635,7 @@ UE_NET_TEST_FIXTURE(FTestGroupsFixture, GroupFilterAPISubObjects)
 	const FNetRefHandle ServerSubObjectHandle = ServerSubObject->NetRefHandle;
 	
 	// Create groups
-	FNetObjectGroupHandle GroupHandle = Server->ReplicationSystem->CreateGroup();
+	FNetObjectGroupHandle GroupHandle = Server->ReplicationSystem->CreateGroup(NAME_None);
 	Server->ReplicationSystem->AddExclusionFilterGroup(GroupHandle);
 
 	// Add Object to group
@@ -600,7 +656,7 @@ UE_NET_TEST_FIXTURE(FTestGroupsFixture, GroupFilterAPISubObjects)
 UE_NET_TEST_FIXTURE(FTestGroupsFixture, GroupFilterAPIRemovingSubObjectsRestoresGroupFilteredOutObjects)
 {
 	// Create groups
-	FNetObjectGroupHandle GroupHandle = Server->ReplicationSystem->CreateGroup();
+	FNetObjectGroupHandle GroupHandle = Server->ReplicationSystem->CreateGroup(NAME_None);
 	Server->ReplicationSystem->AddExclusionFilterGroup(GroupHandle);
 
 	// Update filters
@@ -640,7 +696,7 @@ UE_NET_TEST_FIXTURE(FTestGroupsFixture, GroupFilterAPILateAddedSubObjectIsFilter
 	const FNetRefHandle ServerHandle = ServerObject->NetRefHandle;
 	
 	// Create groups
-	FNetObjectGroupHandle GroupHandle = Server->ReplicationSystem->CreateGroup();
+	FNetObjectGroupHandle GroupHandle = Server->ReplicationSystem->CreateGroup(NAME_None);
 	Server->ReplicationSystem->AddExclusionFilterGroup(GroupHandle);
 
 	// Add Object to group
@@ -675,8 +731,8 @@ UE_NET_TEST_FIXTURE(FTestGroupsFixture, GroupFilterAPISubObjectIsFilteredOutIfOw
 	const FNetRefHandle ServerSubObjectHandle = ServerSubObject->NetRefHandle;
 	
 	// Create groups
-	FNetObjectGroupHandle GroupHandle = Server->ReplicationSystem->CreateGroup();
-	FNetObjectGroupHandle SubObjectGroupHandle = Server->ReplicationSystem->CreateGroup();
+	FNetObjectGroupHandle GroupHandle = Server->ReplicationSystem->CreateGroup(NAME_None);
+	FNetObjectGroupHandle SubObjectGroupHandle = Server->ReplicationSystem->CreateGroup(NAME_None);
 
 	Server->ReplicationSystem->AddExclusionFilterGroup(GroupHandle);
 	Server->ReplicationSystem->SetGroupFilterStatus(GroupHandle, ENetFilterStatus::Allow);
@@ -719,8 +775,8 @@ UE_NET_TEST_FIXTURE(FTestGroupsFixture, GroupFilterAPISubObjectIsFilteredOutWith
 	const FNetRefHandle ServerSubObjectHandle = ServerSubObject->NetRefHandle;
 	
 	// Create groups
-	FNetObjectGroupHandle GroupHandle = Server->ReplicationSystem->CreateGroup();
-	FNetObjectGroupHandle SubObjectGroupHandle = Server->ReplicationSystem->CreateGroup();
+	FNetObjectGroupHandle GroupHandle = Server->ReplicationSystem->CreateGroup(NAME_None);
+	FNetObjectGroupHandle SubObjectGroupHandle = Server->ReplicationSystem->CreateGroup(NAME_None);
 
 	Server->ReplicationSystem->AddExclusionFilterGroup(GroupHandle);
 	Server->ReplicationSystem->AddExclusionFilterGroup(SubObjectGroupHandle);
diff --git a/Engine/Plugins/Runtime/ReplicationSystemTestPlugin/Source/Private/Tests/ReplicationSystem/TestCancelPendingDestroy.cpp b/Engine/Plugins/Runtime/ReplicationSystemTestPlugin/Source/Private/Tests/ReplicationSystem/TestCancelPendingDestroy.cpp
--- a/Engine/Plugins/Runtime/ReplicationSystemTestPlugin/Source/Private/Tests/ReplicationSystem/TestCancelPendingDestroy.cpp
+++ b/Engine/Plugins/Runtime/ReplicationSystemTestPlugin/Source/Private/Tests/ReplicationSystem/TestCancelPendingDestroy.cpp
@@ -790,7 +790,7 @@ UE_NET_TEST_FIXTURE(FTestCancelPendingDestroyFixture, TestCancelPendingDestroyFr
 	UE_NET_ASSERT_EQ(Client1Object->IntA, ServerObject->IntA);
 
 	// Filter out object to cause a PendingDestroy for Client0
-	FNetObjectGroupHandle ExclusionGroupHandle = Server->ReplicationSystem->CreateGroup();
+	FNetObjectGroupHandle ExclusionGroupHandle = Server->ReplicationSystem->CreateGroup(NAME_None);
 	Server->ReplicationSystem->AddToGroup(ExclusionGroupHandle, ServerObject->NetRefHandle);
 	Server->ReplicationSystem->AddExclusionFilterGroup(ExclusionGroupHandle);
 
diff --git a/Engine/Plugins/Runtime/ReplicationSystemTestPlugin/Source/Private/Tests/ReplicationSystem/TestDependentObjects.cpp b/Engine/Plugins/Runtime/ReplicationSystemTestPlugin/Source/Private/Tests/ReplicationSystem/TestDependentObjects.cpp
--- a/Engine/Plugins/Runtime/ReplicationSystemTestPlugin/Source/Private/Tests/ReplicationSystem/TestDependentObjects.cpp
+++ b/Engine/Plugins/Runtime/ReplicationSystemTestPlugin/Source/Private/Tests/ReplicationSystem/TestDependentObjects.cpp
@@ -163,7 +163,7 @@ UE_NET_TEST_FIXTURE(FReplicationSystemServerClientTestFixture, TestDependentObje
 	UTestReplicatedIrisObject* ServerObject = Server->CreateObject(0, 0);
 
 	// Filter out Server object to start with
-	FNetObjectGroupHandle FilterGroup = ReplicationSystem->CreateGroup();
+	FNetObjectGroupHandle FilterGroup = ReplicationSystem->CreateGroup(NAME_None);
 	ReplicationSystem->AddExclusionFilterGroup(FilterGroup);
 	ReplicationSystem->AddToGroup(FilterGroup, ServerObject->NetRefHandle);
 
@@ -211,7 +211,7 @@ UE_NET_TEST_FIXTURE(FReplicationSystemServerClientTestFixture, TestChainedDepend
 	UTestReplicatedIrisObject* ServerObject = Server->CreateObject(0, 0);
 
 	// Filter out Server object to start with
-	FNetObjectGroupHandle FilterGroup = ReplicationSystem->CreateGroup();
+	FNetObjectGroupHandle FilterGroup = ReplicationSystem->CreateGroup(NAME_None);
 	ReplicationSystem->AddExclusionFilterGroup(FilterGroup);
 	ReplicationSystem->AddToGroup(FilterGroup, ServerObject->NetRefHandle);
 
diff --git a/Engine/Plugins/Runtime/ReplicationSystemTestPlugin/Source/Private/Tests/ReplicationSystem/TestFlushBeforeDestroy.cpp b/Engine/Plugins/Runtime/ReplicationSystemTestPlugin/Source/Private/Tests/ReplicationSystem/TestFlushBeforeDestroy.cpp
--- a/Engine/Plugins/Runtime/ReplicationSystemTestPlugin/Source/Private/Tests/ReplicationSystem/TestFlushBeforeDestroy.cpp
+++ b/Engine/Plugins/Runtime/ReplicationSystemTestPlugin/Source/Private/Tests/ReplicationSystem/TestFlushBeforeDestroy.cpp
@@ -1282,7 +1282,7 @@ UE_NET_TEST_FIXTURE(FTestFlushBeforeDestroyFixture, TestCancelPendingDestroyWait
 	Server->PostSendUpdate();
 
 	// Filter out object to cause a flush for Client0
-	FNetObjectGroupHandle ExclusionGroupHandle = Server->ReplicationSystem->CreateGroup();
+	FNetObjectGroupHandle ExclusionGroupHandle = Server->ReplicationSystem->CreateGroup(NAME_None);
 	Server->ReplicationSystem->AddToGroup(ExclusionGroupHandle, ServerObject->NetRefHandle);
 	Server->ReplicationSystem->AddExclusionFilterGroup(ExclusionGroupHandle);
 
diff --git a/Engine/Plugins/Runtime/ReplicationSystemTestPlugin/Source/Private/Tests/ReplicationSystem/TestReplicationSystem.cpp b/Engine/Plugins/Runtime/ReplicationSystemTestPlugin/Source/Private/Tests/ReplicationSystem/TestReplicationSystem.cpp
--- a/Engine/Plugins/Runtime/ReplicationSystemTestPlugin/Source/Private/Tests/ReplicationSystem/TestReplicationSystem.cpp
+++ b/Engine/Plugins/Runtime/ReplicationSystemTestPlugin/Source/Private/Tests/ReplicationSystem/TestReplicationSystem.cpp
@@ -714,7 +714,7 @@ UE_NET_TEST_FIXTURE(FReplicationSystemServerClientTestFixture, AddRemoveFromConn
 	UTestReplicatedIrisObject* ServerObject = Server->CreateObject(0,0);
 
 	// Add to group
-	FNetObjectGroupHandle Group = ReplicationSystem->CreateGroup();
+	FNetObjectGroupHandle Group = ReplicationSystem->CreateGroup(NAME_None);
 	ReplicationSystem->AddToGroup(Group, ServerObject->NetRefHandle);
 
 	ReplicationSystem->AddExclusionFilterGroup(Group);
diff --git a/Engine/Source/Runtime/Engine/Private/Net/Iris/ReplicationSystem/ActorReplicationBridge.cpp b/Engine/Source/Runtime/Engine/Private/Net/Iris/ReplicationSystem/ActorReplicationBridge.cpp
--- a/Engine/Source/Runtime/Engine/Private/Net/Iris/ReplicationSystem/ActorReplicationBridge.cpp
+++ b/Engine/Source/Runtime/Engine/Private/Net/Iris/ReplicationSystem/ActorReplicationBridge.cpp
@@ -1364,7 +1364,7 @@ void UActorReplicationBridge::AddActorToLevelGroup(const AActor* Actor)
 		UE::Net::FNetObjectGroupHandle LevelGroup = GetLevelGroup(Level);
 		if (!LevelGroup.IsValid())
 		{
-			LevelGroup = CreateLevelGroup(Level);
+			LevelGroup = CreateLevelGroup(Level, PackageName);
 
 			UE_LOG_ACTORREPLICATIONBRIDGE(Log, TEXT("Created new GroupIndex: %u for Level: %s"), LevelGroup.GetGroupIndex(), ToCStr(PackageName.ToString()));
 
diff --git a/Engine/Source/Runtime/Experimental/Iris/Core/Private/Iris/ReplicationSystem/Filtering/NetObjectGroups.cpp b/Engine/Source/Runtime/Experimental/Iris/Core/Private/Iris/ReplicationSystem/Filtering/NetObjectGroups.cpp
--- a/Engine/Source/Runtime/Experimental/Iris/Core/Private/Iris/ReplicationSystem/Filtering/NetObjectGroups.cpp
+++ b/Engine/Source/Runtime/Experimental/Iris/Core/Private/Iris/ReplicationSystem/Filtering/NetObjectGroups.cpp
@@ -72,92 +72,54 @@ void FNetObjectGroups::OnMaxInternalNetRefIndexIncreased(FInternalNetRefIndex Ne
 	GroupFilteredOutObjects.SetNumBits(NewMaxInternalIndex);
 }
 
-FNetObjectGroupHandle FNetObjectGroups::CreateGroup()
+FNetObjectGroupHandle FNetObjectGroups::CreateGroup(FName InGroupName)
 {
-	if (ensure((uint32)Groups.Num() < MaxGroupCount))
-	{
-		const uint32 Index = static_cast<uint32>(Groups.Add(FNetObjectGroup()));
-		FNetObjectGroupHandle GroupHandle;
-		GroupHandle.Index = static_cast<FNetObjectGroupHandle::FGroupIndexType>(Index);
-		GroupHandle.Epoch = CurrentEpoch;
-		return GroupHandle;
-	}
-	else
+	const bool bCanCreateGroup = (uint32)Groups.Num() < MaxGroupCount;
+
+	if (!bCanCreateGroup)
 	{
+		UE_LOG(LogIrisGroup, Warning, TEXT("Maximum allowed groups allocated: %u. Group: %s cannot be created"), MaxGroupCount, *InGroupName.ToString());
+		ensureMsgf(false, TEXT("FNetObjectGroups::CreateGroup: Maximum allowed groups allocated."));
 		return FNetObjectGroupHandle();
 	}
-}
 
-void FNetObjectGroups::DestroyGroup(FNetObjectGroupHandle GroupHandle)
-{
-	if (IsValidGroup(GroupHandle))
+	// Auto-generate a name if none was passed	
+	if (InGroupName == NAME_None)
 	{
-		ClearGroup(GroupHandle);
-
-		FNetObjectGroup& Group = Groups[GroupHandle.GetGroupIndex()];
-
-		if (EnumHasAnyFlags(Group.Traits, ENetObjectGroupTraits::IsFindableByName))
-		{
-			NamedGroups.FindAndRemoveChecked(Group.GroupName);
-		}
-
-		Groups.RemoveAt(GroupHandle.GetGroupIndex());
+		InGroupName = FName(TEXT("NetObjectGroup"), AutogeneratedGroupNameId++);
 	}
-}
 
-void FNetObjectGroups::SetGroupName(FNetObjectGroupHandle GroupHandle, FName GroupName)
-{
-	if (FNetObjectGroup* Group = GetGroup(GroupHandle))
+	// Validate that the given name is unique
 	{
-		if (ensureMsgf(EnumHasAnyFlags(Group->Traits, ENetObjectGroupTraits::IsFindableByName), TEXT("FNetObjectGroups::SetGroupName Cannot SetGroupName for grouphandle %u as it is FindableByName"), GroupHandle.GetGroupIndex()))
+		const FNetObjectGroupHandle ExistingGroup = FindGroupHandle(InGroupName);
+		if (ExistingGroup.IsValid())
 		{
-			Group->GroupName = GroupName;
+			UE_LOG(LogIrisGroup, Warning, TEXT("Tried to create a group named: %s. But this name was already registered to groupHandle: %u"), *InGroupName.ToString(), ExistingGroup.GetGroupIndex());
+			ensureMsgf(false, TEXT("FNetObjectGroups::CreateGroup: Group name %s is already registered"), *InGroupName.ToString());
+			return FNetObjectGroupHandle();
 		}
+		
 	}
-}
-
-FName FNetObjectGroups::GetGroupName(FNetObjectGroupHandle GroupHandle) const
-{
-	if (const FNetObjectGroup* Group = GetGroup(GroupHandle))
-	{
-		return Group->GroupName;
-	}
-
-	return FName();
-}
+	
+	const uint32 Index = static_cast<uint32>(Groups.Add(FNetObjectGroup{.GroupName = InGroupName}));
 
-FNetObjectGroupHandle FNetObjectGroups::CreateNamedGroup(FName GroupName)
-{
-	if (NamedGroups.Contains(GroupName))
-	{
-		ensureMsgf(false, TEXT("FNetObjectGroups, trying to create named group %s that already exists"), *GroupName.ToString());
-		return FNetObjectGroupHandle();
-	}
-
-	FNetObjectGroupHandle GroupHandle = CreateGroup();
-	if (FNetObjectGroup* Group = GetGroup(GroupHandle))
-	{
-		Group->GroupName = GroupName;
-		Group->Traits |= ENetObjectGroupTraits::IsFindableByName;
-		NamedGroups.Add(GroupName, GroupHandle);
-	}
+	UE_LOG(LogIrisGroup, Log, TEXT("Created GroupHandle: %u named: %s"), Index, *InGroupName.ToString());
 
+	FNetObjectGroupHandle GroupHandle;
+	GroupHandle.Index = static_cast<FNetObjectGroupHandle::FGroupIndexType>(Index);
+	GroupHandle.Epoch = CurrentEpoch;
 	return GroupHandle;
 }
 
-FNetObjectGroupHandle FNetObjectGroups::GetNamedGroupHandle(FName GroupName)
+void FNetObjectGroups::DestroyGroup(FNetObjectGroupHandle GroupHandle)
 {
-	if (FNetObjectGroupHandle* Group = NamedGroups.Find(GroupName))
+	if (IsValidGroup(GroupHandle))
 	{
-		return *Group;
-	}
-
-	return FNetObjectGroupHandle();
-}
+		ClearGroup(GroupHandle);
 
-void FNetObjectGroups::DestroyNamedGroup(FName GroupName)
-{
-	DestroyGroup(GetNamedGroupHandle(GroupName));
+		FNetObjectGroup& Group = Groups[GroupHandle.GetGroupIndex()];
+		Groups.RemoveAt(GroupHandle.GetGroupIndex());
+	}
 }
 
 void FNetObjectGroups::ClearGroup(FNetObjectGroupHandle GroupHandle)
diff --git a/Engine/Source/Runtime/Experimental/Iris/Core/Private/Iris/ReplicationSystem/Filtering/NetObjectGroups.h b/Engine/Source/Runtime/Experimental/Iris/Core/Private/Iris/ReplicationSystem/Filtering/NetObjectGroups.h
--- a/Engine/Source/Runtime/Experimental/Iris/Core/Private/Iris/ReplicationSystem/Filtering/NetObjectGroups.h
+++ b/Engine/Source/Runtime/Experimental/Iris/Core/Private/Iris/ReplicationSystem/Filtering/NetObjectGroups.h
@@ -26,9 +26,8 @@ namespace UE::Net::Private
 enum class ENetObjectGroupTraits : uint32
 {
 	None                 = 0x0000,
-	IsFindableByName     = 0x0001,
-	IsExclusionFiltering = 0x0002,
-	IsInclusionFiltering = 0x0004,
+	IsExclusionFiltering = 0x0001,
+	IsInclusionFiltering = 0x0002,
 };
 ENUM_CLASS_FLAGS(ENetObjectGroupTraits);
 
@@ -57,10 +56,12 @@ public:
 
 	void Init(const FNetObjectGroupInitParams& Params);
 
-	//$IRIS TODO: Groups should have a mandatory unique name. Otherwise it's impossible to debug issues with groups.
-	FNetObjectGroupHandle CreateGroup();
+	FNetObjectGroupHandle CreateGroup(FName GroupName);
 	void DestroyGroup(FNetObjectGroupHandle GroupHandle);
+
 	void ClearGroup(FNetObjectGroupHandle GroupHandle);
+
+	FNetObjectGroupHandle FindGroupHandle(FName GroupName) const;
 	
 	const FNetObjectGroup* GetGroup(FNetObjectGroupHandle GroupHandle) const;
 	FNetObjectGroup* GetGroup(FNetObjectGroupHandle GroupHandle);
@@ -69,8 +70,8 @@ public:
 	const FNetObjectGroup* GetGroupByIndex(FNetObjectGroupHandle::FGroupIndexType GroupIndex) const;
 	FNetObjectGroup* GetGroupByIndex(FNetObjectGroupHandle::FGroupIndexType GroupIndex);
 
-	void SetGroupName(FNetObjectGroupHandle GroupHandle, FName GroupName);
-	FName GetGroupName(FNetObjectGroupHandle GroupHandle) const;
+	inline FName GetGroupName(FNetObjectGroupHandle GroupHandle) const;
+	inline FString GetGroupNameString(FNetObjectGroupHandle GroupHandle) const;
 
 	bool IsValidGroup(FNetObjectGroupHandle GroupHandle) const;
 
@@ -102,15 +103,6 @@ public:
 	/** Get the array of all groups that the NetObject is a member of */
 	const FNetObjectGroupHandle* GetGroupMemberships(FInternalNetRefIndex InternalIndex, uint32& GroupCount) const;
 
-	/** Create and manage named groups, only groups created as a named group will be findable by name */
-	FNetObjectGroupHandle CreateNamedGroup(FName GroupName);
-
-	/** Lookup NetObjectGroupHandle for a named group */
-	FNetObjectGroupHandle GetNamedGroupHandle(FName GroupName);
-
-	/** Destroy Named group */
-	void DestroyNamedGroup(FName GroupName);
-
 	/** Returns a list of all objects currently part of a group with the filter trait */
 	const FNetBitArrayView GetGroupFilteredOutObjects() const
 	{
@@ -147,13 +139,16 @@ private:
 
 	// Track what groups each internal handle is a member of, we can tighten this up a bit if needed
 	TArray<FNetObjectGroupMembership> GroupMemberships;
+	
+	// Maximum number of groups that can be exist at once
 	uint32 MaxGroupCount = 0U;
 
+	// Index to use for groups with autogenerated names
+	int32 AutogeneratedGroupNameId = 0;
+
 	// List of objects that are members of a group with a filter trait
 	FNetBitArray GroupFilteredOutObjects;
 
-	TMap<FName, FNetObjectGroupHandle> NamedGroups;
-
 	FNetObjectGroupHandle::FGroupIndexType CurrentEpoch = 0U;
 
 	inline static FNetObjectGroupHandle::FGroupIndexType NextEpoch = 1U;
@@ -179,4 +174,33 @@ inline bool FNetObjectGroups::IsInclusionFilterGroup(const FNetObjectGroup& Grou
 	return EnumHasAnyFlags(Group.Traits, ENetObjectGroupTraits::IsInclusionFiltering);
 }
 
+inline FNetObjectGroupHandle FNetObjectGroups::FindGroupHandle(FName InGroupName) const
+{
+	for (const FNetObjectGroup& Group : Groups)
+	{
+		if (Group.GroupName == InGroupName)
+		{
+			const int32 Index = Groups.PointerToIndex(&Group);
+			return FNetObjectGroupHandle((FNetObjectGroupHandle::FGroupIndexType)Index, CurrentEpoch);
+		}
+	}
+
+	return FNetObjectGroupHandle();
+}
+
+inline FName FNetObjectGroups::GetGroupName(FNetObjectGroupHandle GroupHandle) const
+{
+	if (const FNetObjectGroup* Group = GetGroup(GroupHandle))
+	{
+		return Group->GroupName;
+	}
+
+	return FName();
+}
+
+inline FString FNetObjectGroups::GetGroupNameString(FNetObjectGroupHandle GroupHandle) const
+{
+	return GetGroupName(GroupHandle).ToString();
+}
+
 } // end namespace UE::Net::Private
diff --git a/Engine/Source/Runtime/Experimental/Iris/Core/Private/Iris/ReplicationSystem/Filtering/ReplicationFiltering.cpp b/Engine/Source/Runtime/Experimental/Iris/Core/Private/Iris/ReplicationSystem/Filtering/ReplicationFiltering.cpp
--- a/Engine/Source/Runtime/Experimental/Iris/Core/Private/Iris/ReplicationSystem/Filtering/ReplicationFiltering.cpp
+++ b/Engine/Source/Runtime/Experimental/Iris/Core/Private/Iris/ReplicationSystem/Filtering/ReplicationFiltering.cpp
@@ -1850,7 +1850,7 @@ void FReplicationFiltering::AddSubObjectFilter(FNetObjectGroupHandle GroupHandle
 	}
 
 	const bool bIsFiltering = Groups->IsFilterGroup(GroupHandle) || SubObjectFilterGroups.GetBit(GroupIndex);
-	ensureMsgf(!bIsFiltering, TEXT("NetObjectGroup Name '%s' GroupIndex %u was asked to start subobject filtering but it was already used for filtering."), ToCStr(Groups->GetGroupName(GroupHandle).ToString()), GroupIndex);
+	ensureMsgf(!bIsFiltering, TEXT("NetObjectGroup Name '%s' GroupIndex %u was asked to start subobject filtering but it was already used for filtering."), *Groups->GetGroupNameString(GroupHandle), GroupIndex);
 	if (bIsFiltering)
 	{
 		return;
@@ -1862,7 +1862,7 @@ void FReplicationFiltering::AddSubObjectFilter(FNetObjectGroupHandle GroupHandle
 	// By default we filter out all connections
 	SetPerObjectInfoFilterStatus(*GetPerObjectInfo(GroupInfos[GroupIndex].ConnectionStateIndex), ENetFilterStatus::Disallow);
 
-	UE_LOG(LogIrisFiltering, Verbose, TEXT("ReplicationFiltering::AddSubObjectFilter (%s) GroupIndex: %u, FilterStatus: DisallowReplication"), *(Groups->GetGroup(GroupHandle)->GroupName.ToString()), GroupIndex);
+	UE_LOG(LogIrisFiltering, Verbose, TEXT("ReplicationFiltering::AddSubObjectFilter Group: %s FilterStatus: DisallowReplication"), *Groups->GetGroupNameString(GroupHandle), GroupIndex);
 }
 
 void FReplicationFiltering::RemoveSubObjectFilter(FNetObjectGroupHandle GroupHandle)
@@ -1876,7 +1876,7 @@ void FReplicationFiltering::RemoveSubObjectFilter(FNetObjectGroupHandle GroupHan
 		GroupInfos[GroupIndex].ConnectionStateIndex = 0U;
 		FreePerObjectInfo(ConnectionStateIndex);
 
-		UE_LOG(LogIrisFiltering, Verbose, TEXT("ReplicationFiltering::RemoveSubObjectFilter GroupIndex: %u"), GroupIndex);
+		UE_LOG(LogIrisFiltering, Verbose, TEXT("ReplicationFiltering::RemoveSubObjectFilter Group: %s"), *Groups->GetGroupNameString(GroupHandle));
 	}
 }
 
@@ -1908,13 +1908,13 @@ void FReplicationFiltering::SetSubObjectFilterStatus(FNetObjectGroupHandle Group
 	const FNetObjectGroupHandle::FGroupIndexType GroupIndex = GroupHandle.GetGroupIndex();
 	if (GroupHandle.IsReservedNetObjectGroup())
 	{
-		UE_LOG(LogIrisFiltering, Warning, TEXT("FReplicationFiltering::SetSubObjectFilterStatus - Trying to set filter for reserved GroupIndex: %u which is not allowed."), GroupIndex);
+		UE_LOG(LogIrisFiltering, Warning, TEXT("FReplicationFiltering::SetSubObjectFilterStatus - Trying to set filter for reserved Group: %s which is not allowed."), *Groups->GetGroupNameString(GroupHandle));
 		return;
 	}
 
 	if (!SubObjectFilterGroups.GetBit(GroupIndex))
 	{
-		UE_LOG(LogIrisFiltering, Warning, TEXT("FReplicationFiltering::SetSubObjectFilterStatus - Trying to ReplicationStatus for GroupIndex: %u that is not a SubObjectFilterGroup"), GroupIndex);
+		UE_LOG(LogIrisFiltering, Warning, TEXT("FReplicationFiltering::SetSubObjectFilterStatus - Trying to ReplicationStatus for Group: %s that is not a SubObjectFilterGroup"), *Groups->GetGroupNameString(GroupHandle));
 		return;
 	}
 
@@ -1931,19 +1931,19 @@ void FReplicationFiltering::SetSubObjectFilterStatus(FNetObjectGroupHandle Group
 	const FNetObjectGroupHandle::FGroupIndexType GroupIndex = GroupHandle.GetGroupIndex();
 	if (ConnectionsBitArray.GetNumBits() > ValidConnections.GetNumBits())
 	{
-		UE_LOG(LogIrisFiltering, Warning, TEXT("FReplicationFiltering::SetSubObjectFilterStatus - Trying to set filter for %u, with invalid Connections parameters."), GroupIndex);
+		UE_LOG(LogIrisFiltering, Warning, TEXT("FReplicationFiltering::SetSubObjectFilterStatus - Trying to set filter for %s, with invalid Connections parameters."), *Groups->GetGroupNameString(GroupHandle));
 		return;
 	}
 
 	if (GroupHandle.IsReservedNetObjectGroup())
 	{
-		UE_LOG(LogIrisFiltering, Warning, TEXT("FReplicationFiltering::SetSubObjectFilterStatus - Trying to set filter for reserved GroupIndex: %u which is not allowed."), GroupIndex);
+		UE_LOG(LogIrisFiltering, Warning, TEXT("FReplicationFiltering::SetSubObjectFilterStatus - Trying to set filter for reserved Group: %s which is not allowed."), *Groups->GetGroupNameString(GroupHandle));
 		return;
 	}
 
 	if (!SubObjectFilterGroups.GetBit(GroupIndex))
 	{
-		UE_LOG(LogIrisFiltering, Warning, TEXT("FReplicationFiltering::SetSubObjectFilterStatus - Trying to ReplicationStatus for GroupIndex: %u that is not a SubObjectFilterGroup"), GroupIndex);
+		UE_LOG(LogIrisFiltering, Warning, TEXT("FReplicationFiltering::SetSubObjectFilterStatus - Trying to ReplicationStatus for Group: %s that is not a SubObjectFilterGroup"), *Groups->GetGroupNameString(GroupHandle));
 		return;
 	}
 
@@ -1970,13 +1970,13 @@ void FReplicationFiltering::SetSubObjectFilterStatus(FNetObjectGroupHandle Group
 	const FNetObjectGroupHandle::FGroupIndexType GroupIndex = GroupHandle.GetGroupIndex();
 	if (GroupHandle.IsReservedNetObjectGroup())
 	{
-		UE_LOG(LogIrisFiltering, Warning, TEXT("FReplicationFiltering::SetSubObjectFilterStatus - Trying to set filter for reserved GroupIndex: %u which is not allowed."), GroupIndex);
+		UE_LOG(LogIrisFiltering, Warning, TEXT("FReplicationFiltering::SetSubObjectFilterStatus - Trying to set filter for reserved Group: %s which is not allowed."), *Groups->GetGroupNameString(GroupHandle));
 		return;
 	}
 
 	if (ensure(ValidConnections.GetBit(ConnectionId) && SubObjectFilterGroups.GetBit(GroupIndex)))
 	{
-		UE_LOG(LogIrisFiltering, Verbose, TEXT("ReplicationFiltering::SetSubObjectFilterStatus GroupIndex: %u, ConnectionId: %u, FilterStatus: %u"), GroupHandle.GetGroupIndex(), ConnectionId, ReplicationStatus == ENetFilterStatus::Allow ? 1U : 0U);
+		UE_LOG(LogIrisFiltering, Verbose, TEXT("ReplicationFiltering::SetSubObjectFilterStatus Group: %s, ConnectionId: %u, FilterStatus: %u"), *Groups->GetGroupNameString(GroupHandle), ConnectionId, ReplicationStatus == ENetFilterStatus::Allow ? 1U : 0U);
 		FPerObjectInfo* FilterInfo = GetPerObjectInfo(GroupInfos[GroupIndex].ConnectionStateIndex);
 		SetConnectionFilterStatus(*FilterInfo, ConnectionId, ReplicationStatus);
 		if (!IsAnyConnectionFilterStatusAllowed(*FilterInfo))
@@ -2009,7 +2009,7 @@ bool FReplicationFiltering::AddExclusionFilterGroup(FNetObjectGroupHandle GroupH
 	}
 
 	const bool bIsFiltering = Groups->IsFilterGroup(GroupHandle) || SubObjectFilterGroups.GetBit(GroupIndex);
-	ensureMsgf(!bIsFiltering, TEXT("NetObjectGroup Name '%s' GroupIndex %u was asked to start exclusion filtering but it was already used for filtering."), ToCStr(Groups->GetGroupName(GroupHandle).ToString()), GroupIndex);
+	ensureMsgf(!bIsFiltering, TEXT("NetObjectGroup Name '%s' GroupIndex %u was asked to start exclusion filtering but it was already used for filtering."), *Groups->GetGroupNameString(GroupHandle), GroupIndex);
 	if (bIsFiltering)
 	{
 		return false;
@@ -2025,7 +2025,7 @@ bool FReplicationFiltering::AddExclusionFilterGroup(FNetObjectGroupHandle GroupH
 	GroupInfos[GroupIndex].ConnectionStateIndex = AllocPerObjectInfo();
 	SetPerObjectInfoFilterStatus(*GetPerObjectInfo(GroupInfos[GroupIndex].ConnectionStateIndex), ENetFilterStatus::Disallow);
 
-	UE_LOG(LogIrisFiltering, Verbose, TEXT("ReplicationFiltering::AddExclusionGroupFilter GroupIndex: %u, FilterStatus: DisallowReplication"), GroupIndex);
+	UE_LOG(LogIrisFiltering, Verbose, TEXT("ReplicationFiltering::AddExclusionGroupFilter on %s, FilterStatus: DisallowReplication"), *Groups->GetGroupNameString(GroupHandle));
 	return true;
 }
 
@@ -2039,7 +2039,7 @@ bool FReplicationFiltering::AddInclusionFilterGroup(FNetObjectGroupHandle GroupH
 	}
 
 	const bool bIsFiltering = Groups->IsFilterGroup(GroupHandle) || SubObjectFilterGroups.GetBit(GroupIndex);
-	ensureMsgf(!bIsFiltering, TEXT("NetObjectGroup Name '%s' GroupIndex %u was asked to start exclusion filtering but it was already used for filtering."), ToCStr(Groups->GetGroupName(GroupHandle).ToString()), GroupIndex);
+	ensureMsgf(!bIsFiltering, TEXT("NetObjectGroup Name '%s' GroupIndex %u was asked to start exclusion filtering but it was already used for filtering."), *Groups->GetGroupNameString(GroupHandle), GroupIndex);
 	if (bIsFiltering)
 	{
 		return false;
@@ -2055,7 +2055,7 @@ bool FReplicationFiltering::AddInclusionFilterGroup(FNetObjectGroupHandle GroupH
 	GroupInfos[GroupIndex].ConnectionStateIndex = AllocPerObjectInfo();
 	SetPerObjectInfoFilterStatus(*GetPerObjectInfo(GroupInfos[GroupIndex].ConnectionStateIndex), ENetFilterStatus::Disallow);
 
-	UE_LOG(LogIrisFiltering, Verbose, TEXT("ReplicationFiltering::AddInclusionFilterGroup GroupIndex: %u, FilterStatus: DoNotOverride"), GroupIndex);
+	UE_LOG(LogIrisFiltering, Verbose, TEXT("ReplicationFiltering::AddInclusionFilterGroup on %s, FilterStatus: DoNotOverride"), *Groups->GetGroupNameString(GroupHandle));
 	return true;
 }
 
@@ -2085,7 +2085,7 @@ void FReplicationFiltering::RemoveGroupFilter(FNetObjectGroupHandle GroupHandle)
 
 		Groups->RemoveExclusionFilterTrait(GroupHandle);
 
-		UE_LOG(LogIrisFiltering, Verbose, TEXT("ReplicationFiltering::RemoveGroupFilter ExclusionGroup GroupIndex: %u"), GroupHandle.GetGroupIndex());
+		UE_LOG(LogIrisFiltering, Verbose, TEXT("ReplicationFiltering::RemoveGroupFilter ExclusionFilter Group: %s"), *Groups->GetGroupNameString(GroupHandle));
 	}
 	else if (InclusionFilterGroups.GetBit(GroupIndex))
 	{
@@ -2102,7 +2102,7 @@ void FReplicationFiltering::RemoveGroupFilter(FNetObjectGroupHandle GroupHandle)
 
 		Groups->RemoveInclusionFilterTrait(GroupHandle);
 
-		UE_LOG(LogIrisFiltering, Verbose, TEXT("ReplicationFiltering::RemoveGroupFilter InclusionGroup GroupIndex: %u"), GroupHandle.GetGroupIndex());
+		UE_LOG(LogIrisFiltering, Verbose, TEXT("ReplicationFiltering::RemoveGroupFilter InclusionFilter Group: %s"), *Groups->GetGroupNameString(GroupHandle));
 	}
 }
 
@@ -2113,7 +2113,7 @@ void FReplicationFiltering::SetGroupFilterStatus(FNetObjectGroupHandle GroupHand
 	const FNetObjectGroupHandle::FGroupIndexType GroupIndex = GroupHandle.GetGroupIndex();
 	if (GroupHandle.IsReservedNetObjectGroup())
 	{
-		UE_LOG(LogIrisFiltering, Warning, TEXT("FReplicationFiltering::SetGroupFilterStatus - Trying to set filter for reserved GroupIndex: %u which is not allowed."), GroupIndex);
+		UE_LOG(LogIrisFiltering, Warning, TEXT("FReplicationFiltering::SetGroupFilterStatus - Trying to set filter for reserved Group: %s which is not allowed."), *Groups->GetGroupNameString(GroupHandle));
 		return;
 	}
 
@@ -2134,7 +2134,7 @@ void FReplicationFiltering::SetGroupFilterStatus(FNetObjectGroupHandle GroupHand
 		return;
 	}
 	
-	UE_LOG(LogIrisFiltering, Warning, TEXT("FReplicationFiltering::SetGroupFilterStatus - Trying to set filter for invalid GroupIndex: %u, Make sure group is added to filtering"), GroupHandle.GetGroupIndex());
+	UE_LOG(LogIrisFiltering, Warning, TEXT("FReplicationFiltering::SetGroupFilterStatus - Trying to set filter for invalid Group: %s, Make sure group is added to filtering"), *Groups->GetGroupNameString(GroupHandle));
 }
 
 void FReplicationFiltering::SetGroupFilterStatus(FNetObjectGroupHandle GroupHandle, const FNetBitArrayView& ConnectionsBitArray, ENetFilterStatus ReplicationStatus)
@@ -2144,13 +2144,13 @@ void FReplicationFiltering::SetGroupFilterStatus(FNetObjectGroupHandle GroupHand
 	const FNetObjectGroupHandle::FGroupIndexType GroupIndex = GroupHandle.GetGroupIndex();
 	if (GroupHandle.IsReservedNetObjectGroup())
 	{
-		UE_LOG(LogIrisFiltering, Warning, TEXT("FReplicationFiltering::SetGroupFilterStatus - Trying to set filter for reserved GroupIndex: %u which is not allowed."), GroupIndex);		
+		UE_LOG(LogIrisFiltering, Warning, TEXT("FReplicationFiltering::SetGroupFilterStatus - Trying to set filter on reserved Group: %s which is not allowed."), *Groups->GetGroupNameString(GroupHandle));
 		return;
 	}
 
 	if (ConnectionsBitArray.GetNumBits() > ValidConnections.GetNumBits())
 	{
-		UE_LOG(LogIrisFiltering, Warning, TEXT("FReplicationFiltering::SetGroupFilterStatus - Trying to set filter for %u, with invalid Connections parameters."), GroupIndex);
+		UE_LOG(LogIrisFiltering, Warning, TEXT("FReplicationFiltering::SetGroupFilterStatus - Trying to set filter on Group: %s, with invalid Connections parameters."), *Groups->GetGroupNameString(GroupHandle));
 		return;
 	}
 
@@ -2186,7 +2186,7 @@ void FReplicationFiltering::SetGroupFilterStatus(FNetObjectGroupHandle GroupHand
 		return;
 	}
 
-	UE_LOG(LogIrisFiltering, Warning, TEXT("FReplicationFiltering::SetGroupFilterStatus - Trying to set filter for invalid GroupIndex: %u, Make sure group is added to filtering"), GroupIndex);
+	UE_LOG(LogIrisFiltering, Warning, TEXT("FReplicationFiltering::SetGroupFilterStatus - Trying to set filter for invalid Group: %s, Make sure group is added to filtering"), *Groups->GetGroupNameString(GroupHandle));
 }
 
 void FReplicationFiltering::SetGroupFilterStatus(FNetObjectGroupHandle GroupHandle, uint32 ConnectionId, ENetFilterStatus ReplicationStatus)
@@ -2196,7 +2196,7 @@ void FReplicationFiltering::SetGroupFilterStatus(FNetObjectGroupHandle GroupHand
 	const FNetObjectGroupHandle::FGroupIndexType GroupIndex = GroupHandle.GetGroupIndex();
 	if (GroupHandle.IsReservedNetObjectGroup())
 	{
-		UE_LOG(LogIrisFiltering, Warning, TEXT("FReplicationFiltering::SetGroupFilterStatus - Trying to set filter for reserved GroupIndex: %u which is not allowed."), GroupIndex);
+		UE_LOG(LogIrisFiltering, Warning, TEXT("FReplicationFiltering::SetGroupFilterStatus - Trying to set filter on reserved Group: %s which is not allowed."), *Groups->GetGroupNameString(GroupHandle));
 		return;
 	}
 
@@ -2212,7 +2212,7 @@ void FReplicationFiltering::SetGroupFilterStatus(FNetObjectGroupHandle GroupHand
 		return;
 	}
 
-	UE_LOG(LogIrisFiltering, Warning, TEXT("FReplicationFiltering::SetGroupFilterStatus - Trying to set invalid filter ConnectionId: %u, GroupIndex: %u"), ConnectionId, GroupIndex);
+	UE_LOG(LogIrisFiltering, Warning, TEXT("FReplicationFiltering::SetGroupFilterStatus - Trying to set invalid filter status: %u on Group: %s for ConnectionId: %u"), ReplicationStatus, *Groups->GetGroupNameString(GroupHandle), ConnectionId);
 }
 
 bool FReplicationFiltering::IsExcludedByAnyGroup(uint32 ObjectInternalIndex, uint32 ConnectionId) const
@@ -2351,7 +2351,7 @@ void FReplicationFiltering::InternalSetExclusionGroupFilterStatus(FNetObjectGrou
 			// Mark filter as active for the connection
 			SetConnectionFilterStatus(*ConnectionState, ConnectionId, ENetFilterStatus::Disallow);
 
-			UE_LOG(LogIrisFiltering, Verbose, TEXT("ReplicationFiltering::SetGroupFilterStatus ExclusionGroup GroupIndex: %u, ConnectionId: %u, FilterStatus: DisallowReplication"), GroupHandle.GetGroupIndex(), ConnectionId);
+			UE_LOG(LogIrisFiltering, Verbose, TEXT("ReplicationFiltering::SetGroupFilterStatus ExclusionGroup: %s, ConnectionId: %u, FilterStatus: DisallowReplication"), *Groups->GetGroupNameString(GroupHandle), ConnectionId);
 
 			if (ValidConnections.GetBit(ConnectionId))
 			{
@@ -2393,7 +2393,7 @@ void FReplicationFiltering::InternalSetExclusionGroupFilterStatus(FNetObjectGrou
 			// Mark filter as no longer being active for the connection
 			SetConnectionFilterStatus(*ConnectionState, ConnectionId, ENetFilterStatus::Allow);
 
-			UE_LOG(LogIrisFiltering, Verbose, TEXT("ReplicationFiltering::SetGroupFilterStatus ExclusionGroup GroupIndex: %u, ConnectionId: %u, FilterStatus: AllowReplication"), GroupHandle.GetGroupIndex(), ConnectionId);
+			UE_LOG(LogIrisFiltering, Verbose, TEXT("ReplicationFiltering::SetGroupFilterStatus ExclusionGroup: %s, ConnectionId: %u, FilterStatus: AllowReplication"), *Groups->GetGroupNameString(GroupHandle), ConnectionId);
 
 			if (ValidConnections.GetBit(ConnectionId))
 			{
@@ -2457,7 +2457,7 @@ void FReplicationFiltering::InternalSetInclusionGroupFilterStatus(FNetObjectGrou
 	const FNetObjectGroup* Group = Groups->GetGroup(GroupHandle);
 	if (ReplicationStatus == ENetFilterStatus::Disallow)
 	{
-		UE_LOG(LogIrisFiltering, Verbose, TEXT("ReplicationFiltering::SetGroupFilterStatus InclusionGroup GroupIndex: %u, ConnectionId: %u, FilterStatus: DisallowReplication"), GroupHandle.GetGroupIndex(), ConnectionId);
+		UE_LOG(LogIrisFiltering, Verbose, TEXT("ReplicationFiltering::SetGroupFilterStatus InclusionGroup: %s, ConnectionId: %u, FilterStatus: DisallowReplication"), *Groups->GetGroupNameString(GroupHandle), ConnectionId);
 		for (const FInternalNetRefIndex ObjectIndex : Group->Members)
 		{
 			ClearGroupInclusionFilterEffectsForObject(ObjectIndex, ConnectionId);
@@ -2465,7 +2465,7 @@ void FReplicationFiltering::InternalSetInclusionGroupFilterStatus(FNetObjectGrou
 	}
 	else
 	{
-		UE_LOG(LogIrisFiltering, Verbose, TEXT("ReplicationFiltering::SetGroupFilterStatus InclusionGroup GroupIndex: %u, ConnectionId: %u, FilterStatus: AllowReplication"), GroupHandle.GetGroupIndex(), ConnectionId);
+		UE_LOG(LogIrisFiltering, Verbose, TEXT("ReplicationFiltering::SetGroupFilterStatus InclusionGroup: %s, ConnectionId: %u, FilterStatus: AllowReplication"), *Groups->GetGroupNameString(GroupHandle), ConnectionId);
 
 		FPerConnectionInfo& ConnectionInfo = ConnectionInfos[ConnectionId];
 		FNetBitArray& GroupIncludedObjects = ConnectionInfo.GroupIncludedObjects;
@@ -2507,7 +2507,7 @@ void FReplicationFiltering::NotifyObjectAddedToGroup(FNetObjectGroupHandle Group
 	const FNetObjectGroupHandle::FGroupIndexType GroupIndex = GroupHandle.GetGroupIndex();
 	if (SubObjectFilterGroups.GetBit(GroupIndex))
 	{
-		UE_LOG(LogIrisFiltering, Verbose, TEXT("ReplicationFiltering::NotifyObjectAddedToGroup Added %s to SubObjectFilter with GroupIndex: %u"), *(NetRefHandleManager->PrintObjectFromIndex(ObjectIndex)), GroupIndex);
+		UE_LOG(LogIrisFiltering, Verbose, TEXT("ReplicationFiltering::NotifyObjectAddedToGroup Added %s to SubObjectFilter group: %s"), *(NetRefHandleManager->PrintObjectFromIndex(ObjectIndex)), *Groups->GetGroupNameString(GroupHandle));
 	}
 	else if (ExclusionFilterGroups.GetBit(GroupIndex))
 	{
@@ -2520,7 +2520,7 @@ void FReplicationFiltering::NotifyObjectAddedToGroup(FNetObjectGroupHandle Group
 			bHasDirtyExclusionFilterGroup = 1;
 		}
 
-		UE_LOG(LogIrisFiltering, Verbose, TEXT("ReplicationFiltering::NotifyObjectAddedToGroup Added %s to ExclusionGroup filter with GroupIndex: %u"), *(NetRefHandleManager->PrintObjectFromIndex(ObjectIndex)), GroupIndex);
+		UE_LOG(LogIrisFiltering, Verbose, TEXT("ReplicationFiltering::NotifyObjectAddedToGroup Added %s to ExclusionFilter group: %s"), *(NetRefHandleManager->PrintObjectFromIndex(ObjectIndex)), *Groups->GetGroupNameString(GroupHandle));
 	}
 	else if (InclusionFilterGroups.GetBit(GroupIndex))
 	{
@@ -2533,7 +2533,7 @@ void FReplicationFiltering::NotifyObjectAddedToGroup(FNetObjectGroupHandle Group
 			bHasDirtyInclusionFilterGroup = 1;
 		}
 
-		UE_LOG(LogIrisFiltering, Verbose, TEXT("ReplicationFiltering::NotifyObjectAddedToGroup Added %s to InclusionGroup filter with GroupIndex: %u"), *(NetRefHandleManager->PrintObjectFromIndex(ObjectIndex)), GroupIndex);
+		UE_LOG(LogIrisFiltering, Verbose, TEXT("ReplicationFiltering::NotifyObjectAddedToGroup Added %s to InclusionFilter group: %s"), *(NetRefHandleManager->PrintObjectFromIndex(ObjectIndex)), *Groups->GetGroupNameString(GroupHandle));
 	}
 }
 
@@ -2541,7 +2541,7 @@ void FReplicationFiltering::NotifyObjectRemovedFromGroup(FNetObjectGroupHandle G
 {
 	const FNetObjectGroupHandle::FGroupIndexType GroupIndex = GroupHandle.GetGroupIndex();
 
-	UE_LOG(LogIrisFiltering, Verbose, TEXT("ReplicationFiltering::NotifyObjectRemovedFromGroup Removing %s from GroupIndex: %u"), ToCStr(NetRefHandleManager->PrintObjectFromIndex(ObjectIndex)), GroupIndex);
+	UE_LOG(LogIrisFiltering, Verbose, TEXT("ReplicationFiltering::NotifyObjectRemovedFromGroup Removing %s from Group: %s"), ToCStr(NetRefHandleManager->PrintObjectFromIndex(ObjectIndex)), *Groups->GetGroupNameString(GroupHandle));
 
 	if (SubObjectFilterGroups.GetBit(GroupIndex))
 	{
diff --git a/Engine/Source/Runtime/Experimental/Iris/Core/Private/Iris/ReplicationSystem/ReplicationBridge.cpp b/Engine/Source/Runtime/Experimental/Iris/Core/Private/Iris/ReplicationSystem/ReplicationBridge.cpp
--- a/Engine/Source/Runtime/Experimental/Iris/Core/Private/Iris/ReplicationSystem/ReplicationBridge.cpp
+++ b/Engine/Source/Runtime/Experimental/Iris/Core/Private/Iris/ReplicationSystem/ReplicationBridge.cpp
@@ -1091,11 +1091,11 @@ void UReplicationBridge::NotifyStreamingLevelUnload(const UObject* Level)
 	}
 }
 
-UE::Net::FNetObjectGroupHandle UReplicationBridge::CreateLevelGroup(const UObject* Level)
+UE::Net::FNetObjectGroupHandle UReplicationBridge::CreateLevelGroup(const UObject* Level, FName PackageName)
 {
 	using namespace UE::Net;
 
-	FNetObjectGroupHandle LevelGroupHandle = ReplicationSystem->CreateGroup();
+	FNetObjectGroupHandle LevelGroupHandle = ReplicationSystem->CreateGroup(PackageName);
 	if (ensure(LevelGroupHandle.IsValid()))
 	{
 		ReplicationSystem->AddExclusionFilterGroup(LevelGroupHandle);
diff --git a/Engine/Source/Runtime/Experimental/Iris/Core/Private/Iris/ReplicationSystem/ReplicationSystem.cpp b/Engine/Source/Runtime/Experimental/Iris/Core/Private/Iris/ReplicationSystem/ReplicationSystem.cpp
--- a/Engine/Source/Runtime/Experimental/Iris/Core/Private/Iris/ReplicationSystem/ReplicationSystem.cpp
+++ b/Engine/Source/Runtime/Experimental/Iris/Core/Private/Iris/ReplicationSystem/ReplicationSystem.cpp
@@ -88,7 +88,7 @@ public:
 
 	void InitDefaultFilteringGroups()
 	{
-		NotReplicatedNetObjectGroupHandle = ReplicationSystem->CreateGroup();
+		NotReplicatedNetObjectGroupHandle = ReplicationSystem->CreateGroup(FName(TEXT("NotReplicated")));
 		check(NotReplicatedNetObjectGroupHandle.IsNotReplicatedNetObjectGroup());
 		ReplicationSystem->AddExclusionFilterGroup(NotReplicatedNetObjectGroupHandle);
 		
@@ -1048,7 +1048,7 @@ bool UReplicationSystem::SetRPCSendPolicyFlags(const UFunction* Function, UE::Ne
 
 	if (EnumHasAnyFlags(SendFlags, UE::Net::ENetObjectAttachmentSendPolicyFlags::SendImmediate) && (Function->FunctionFlags & FUNC_NetReliable))
 	{
-		ensureAlwaysMsgf(false, TEXT("ENetObjectAttachmentSendPolicyFlags::SendImmediate is not allowed to use on Reliable RPC: %s"), *GetNameSafe(Function));
+		ensureMsgf(false, TEXT("ENetObjectAttachmentSendPolicyFlags::SendImmediate is not allowed to use on Reliable RPC: %s"), *GetNameSafe(Function));
 		return false;
 	}
 
@@ -1240,14 +1240,14 @@ UE::Net::FNetObjectGroupHandle UReplicationSystem::GetOrCreateSubObjectFilter(FN
 	FNetObjectGroups& Groups = Impl->ReplicationSystemInternal.GetGroups();
 	FReplicationFiltering& Filtering = Impl->ReplicationSystemInternal.GetFiltering();
 
-	FNetObjectGroupHandle GroupHandle = Groups.GetNamedGroupHandle(GroupName);
+	FNetObjectGroupHandle GroupHandle = Groups.FindGroupHandle(GroupName);
 	if (GroupHandle.IsValid())
 	{
 		check(Filtering.IsSubObjectFilterGroup(GroupHandle));
 		return GroupHandle;
 	}
 
-	GroupHandle = Groups.CreateNamedGroup(GroupName);
+	GroupHandle = Groups.CreateGroup(GroupName);
 	if (GroupHandle.IsValid())
 	{
 		Filtering.AddSubObjectFilter(GroupHandle);
@@ -1263,10 +1263,10 @@ UE::Net::FNetObjectGroupHandle UReplicationSystem::GetSubObjectFilterGroupHandle
 	FNetObjectGroups& Groups = Impl->ReplicationSystemInternal.GetGroups();
 	FReplicationFiltering& Filtering = Impl->ReplicationSystemInternal.GetFiltering();
 
-	FNetObjectGroupHandle GroupHandle = Groups.GetNamedGroupHandle(GroupName);
+	FNetObjectGroupHandle GroupHandle = Groups.FindGroupHandle(GroupName);
 	if (GroupHandle.IsValid())
 	{
-		if (ensureAlwaysMsgf(Filtering.IsSubObjectFilterGroup(GroupHandle), TEXT("UReplicationSystem::GetSubObjectFilterGroupHandle Trying to lookup NetObjectGroupHandle for NetGroup %s that is not a subobject filter"), *GroupName.ToString()))
+		if (ensureMsgf(Filtering.IsSubObjectFilterGroup(GroupHandle), TEXT("UReplicationSystem::GetSubObjectFilterGroupHandle Trying to lookup NetObjectGroupHandle for NetGroup %s that is not a subobject filter"), *GroupName.ToString()))
 		{
 			return GroupHandle;
 		}
@@ -1281,7 +1281,7 @@ void UReplicationSystem::SetSubObjectFilterStatus(FName GroupName, uint32 Connec
 
 	if (UE::Net::IsSpecialNetConditionGroup(GroupName))
 	{
-		ensureAlwaysMsgf(false, TEXT("UReplicationSystem::SetSubObjectFilterStatus Cannot SetSubObjectFilterStatus for special NetGroup %s"), *GroupName.ToString());
+		ensureMsgf(false, TEXT("UReplicationSystem::SetSubObjectFilterStatus Cannot SetSubObjectFilterStatus for special NetGroup %s"), *GroupName.ToString());
 		return;
 	}
 
@@ -1324,11 +1324,11 @@ void UReplicationSystem::RemoveSubObjectFilter(FName GroupName)
 	}
 }
 
-UE::Net::FNetObjectGroupHandle UReplicationSystem::CreateGroup()
+UE::Net::FNetObjectGroupHandle UReplicationSystem::CreateGroup(FName GroupName)
 {
 	LLM_SCOPE_BYTAG(Iris);
 
-	return Impl->ReplicationSystemInternal.GetGroups().CreateGroup();
+	return Impl->ReplicationSystemInternal.GetGroups().CreateGroup(GroupName);
 }
 
 void UReplicationSystem::AddToGroup(FNetObjectGroupHandle GroupHandle, FNetRefHandle Handle)
@@ -1468,6 +1468,13 @@ void UReplicationSystem::DestroyGroup(FNetObjectGroupHandle GroupHandle)
 	Groups.DestroyGroup(GroupHandle);
 }
 
+UE::Net::FNetObjectGroupHandle UReplicationSystem::FindGroup(FName GroupName) const
+{
+	const UE::Net::Private::FNetObjectGroups& Groups = Impl->ReplicationSystemInternal.GetGroups();
+
+	return Groups.FindGroupHandle(GroupName);
+}
+
 UE::Net::FNetObjectGroupHandle UReplicationSystem::GetNotReplicatedNetObjectGroup() const
 {
 	return Impl->NotReplicatedNetObjectGroupHandle;
diff --git a/Engine/Source/Runtime/Experimental/Iris/Core/Public/Iris/ReplicationSystem/ReplicationBridge.h b/Engine/Source/Runtime/Experimental/Iris/Core/Public/Iris/ReplicationSystem/ReplicationBridge.h
--- a/Engine/Source/Runtime/Experimental/Iris/Core/Public/Iris/ReplicationSystem/ReplicationBridge.h
+++ b/Engine/Source/Runtime/Experimental/Iris/Core/Public/Iris/ReplicationSystem/ReplicationBridge.h
@@ -266,7 +266,7 @@ protected:
 	inline UE::Net::Private::FObjectReferenceCache* GetObjectReferenceCache() const { return ObjectReferenceCache; }
 
 	/** Creates a group for a level for object filtering purposes. */
-	IRISCORE_API UE::Net::FNetObjectGroupHandle CreateLevelGroup(const UObject* Level);
+	IRISCORE_API UE::Net::FNetObjectGroupHandle CreateLevelGroup(const UObject* Level, FName PackageName);
 
 	/** Destroys the group associated with the level. */
 	IRISCORE_API void DestroyLevelGroup(const UObject* Level);
diff --git a/Engine/Source/Runtime/Experimental/Iris/Core/Public/Iris/ReplicationSystem/ReplicationSystem.h b/Engine/Source/Runtime/Experimental/Iris/Core/Public/Iris/ReplicationSystem/ReplicationSystem.h
--- a/Engine/Source/Runtime/Experimental/Iris/Core/Public/Iris/ReplicationSystem/ReplicationSystem.h
+++ b/Engine/Source/Runtime/Experimental/Iris/Core/Public/Iris/ReplicationSystem/ReplicationSystem.h
@@ -358,10 +358,11 @@ public:
 	 * Create a group which can be used to logically group objects together. The group must be
 	 * destroyed when it's not needed anymore.
 	 * Groups can be used to setup filtering rules on it's members.
+	 * @param GroupName An unique name to identify the group. Passing NAME_None will assign it an autogenerated name
 	 * @return A handle to the group, or InvalidNetObjectGroupHandle if no more groups could be created.
 	 * @see DestroyGroup
 	 */
-	IRISCORE_API FNetObjectGroupHandle CreateGroup();
+	IRISCORE_API FNetObjectGroupHandle CreateGroup(FName GroupName);
 
 	/**
 	 * Destroy a group.
@@ -369,6 +370,12 @@ public:
 	 */
 	IRISCORE_API void DestroyGroup(FNetObjectGroupHandle GroupHandle);
 
+	/**
+	 * Return the handle to a group identified by the given name
+	 * @return A valid handle if the name is used by an existing group or InvalidNetObjectGroupHandle if no groups is found
+	 */
+	IRISCORE_API FNetObjectGroupHandle FindGroup(FName GroupName) const;
+
 	/**
 	 * Add an object to a group.
 	 * @param GroupHandle A valid group handle.
