Iris - Fixed compile errors due to AutomationTest API changes.

#rnx

diff --git a/Engine/Plugins/Runtime/ReplicationSystemTestPlugin/Source/Private/NetworkAutomationTest.cpp b/Engine/Plugins/Runtime/ReplicationSystemTestPlugin/Source/Private/NetworkAutomationTest.cpp
--- a/Engine/Plugins/Runtime/ReplicationSystemTestPlugin/Source/Private/NetworkAutomationTest.cpp
+++ b/Engine/Plugins/Runtime/ReplicationSystemTestPlugin/Source/Private/NetworkAutomationTest.cpp
@@ -94,7 +94,7 @@ FNetworkAutomationTestWrapper::FNetworkAutomationTestWrapper(FNetworkAutomationT
 {
 }
 
-uint32 FNetworkAutomationTestWrapper::GetTestFlags() const
+EAutomationTestFlags FNetworkAutomationTestWrapper::GetTestFlags() const
 {
 	return EAutomationTestFlags::EngineFilter | EAutomationTestFlags::ClientContext | EAutomationTestFlags::ServerContext | EAutomationTestFlags::CommandletContext;
 }
diff --git a/Engine/Plugins/Runtime/ReplicationSystemTestPlugin/Source/Private/NetworkAutomationTest.h b/Engine/Plugins/Runtime/ReplicationSystemTestPlugin/Source/Private/NetworkAutomationTest.h
--- a/Engine/Plugins/Runtime/ReplicationSystemTestPlugin/Source/Private/NetworkAutomationTest.h
+++ b/Engine/Plugins/Runtime/ReplicationSystemTestPlugin/Source/Private/NetworkAutomationTest.h
@@ -73,7 +73,7 @@ public:
 	FNetworkAutomationTestWrapper(FNetworkAutomationTestSuiteFixture& TestSuite, const TCHAR* Name);
 
 private:
-	virtual uint32 GetTestFlags() const override;
+	virtual EAutomationTestFlags GetTestFlags() const override;
 	virtual uint32 GetRequiredDeviceNum() const override;
 
 	virtual FString GetBeautifiedTestName() const override;
