Iris - Initial naive but working version of FInstancedStructNetSerializer to be able to replicate FInstancedStruct. Requires struct type to be resolvable on sender and receiver. Unit tests.
Added FNetSerializerAlignedStorage to allow external NetSerializers to be able to allocate arbitrary aligned memory.

diff --git a/Engine/Plugins/Experimental/StructUtils/Source/StructUtils/Private/InstancedStructNetSerializer.cpp b/Engine/Plugins/Experimental/StructUtils/Source/StructUtils/Private/InstancedStructNetSerializer.cpp
new file mode 100644
--- /dev/null
+++ b/Engine/Plugins/Experimental/StructUtils/Source/StructUtils/Private/InstancedStructNetSerializer.cpp
@@ -0,0 +1,718 @@
+// Copyright Epic Games, Inc. All Rights Reserved.
+
+#include "InstancedStructNetSerializer.h"
+#include UE_INLINE_GENERATED_CPP_BY_NAME(InstancedStructNetSerializer)
+
+#if UE_WITH_IRIS
+
+#include "InstancedStruct.h"
+#include "HAL/IConsoleManager.h"
+#include "Iris/Core/IrisLog.h"
+#include "Iris/Core/NetObjectReference.h"
+#include "Iris/ReplicationState/ReplicationStateDescriptorBuilder.h"
+#include "Iris/ReplicationState/PropertyNetSerializerInfoRegistry.h"
+#include "Iris/Serialization/NetReferenceCollector.h"
+#include "Iris/Serialization/NetSerializerArrayStorage.h"
+#include "Iris/Serialization/NetSerializerDelegates.h"
+#include "Iris/Serialization/NetSerializers.h"
+#include "Templates/IsPODType.h"
+
+FInstancedStructNetSerializerConfig::FInstancedStructNetSerializerConfig()
+: FNetSerializerConfig()
+{
+	ConfigTraits = ENetSerializerConfigTraits::NeedDestruction;
+}
+
+FInstancedStructNetSerializerConfig::~FInstancedStructNetSerializerConfig() = default;
+
+namespace UE::Net
+{
+
+static int32 MaxCachedInstancedStructDescriptorCount = 8;
+static FAutoConsoleVariableRef CVarMaxCachedInstancedStructDescriptors(TEXT("InstancedStruct.MaxCachedReplicationStateDescriptors"), MaxCachedInstancedStructDescriptorCount, TEXT("How many ReplicationStateDescriptors the InstancedStructNetSerializer is allowed to cache for InstancedStructs without a type allow list. Warning: A value <= 0 means an unlimited amount of descriptors."));
+
+const FName NetError_InstancedStructNetSerializer_InvalidStructType("Invalid struct type");
+
+struct FFInstancedStructNetSerializerQuantizedData
+{
+	FNetSerializerAlignedStorage StructData;
+	FNetObjectReference StructType;
+
+	// Not serialized. Fully qualified path. For ReplicationStateDescriptor lookup, validation etc. 
+	FName StructName;
+	// Not serialized. To optimize away some calls like dynamic memory management and object references.
+	EReplicationStateTraits StructDescriptorTraits;
+};
+
+struct FInstancedStructPropertyNetSerializerInfo : public FNamedStructPropertyNetSerializerInfo
+{
+public:
+	FInstancedStructPropertyNetSerializerInfo();
+
+protected:
+	virtual bool CanUseDefaultConfig(const FProperty* Property) const;
+	virtual FNetSerializerConfig* BuildNetSerializerConfig(void* NetSerializerConfigBuffer, const FProperty* Property) const override;
+};
+
+}
+
+template<> struct TIsPODType<UE::Net::FFInstancedStructNetSerializerQuantizedData> { enum { Value = true }; };
+
+namespace UE::Net
+{
+
+struct FInstancedStructNetSerializer
+{
+	// Version
+	static const uint32 Version = 0;
+
+	// Traits
+	static constexpr bool bHasDynamicState = true;
+	static constexpr bool bIsForwardingSerializer = true;
+	static constexpr bool bHasCustomNetReference = true;
+
+	typedef FInstancedStruct SourceType;
+	typedef FFInstancedStructNetSerializerQuantizedData QuantizedType;
+	typedef FInstancedStructNetSerializerConfig ConfigType;
+
+	static void Serialize(FNetSerializationContext&, const FNetSerializeArgs&);
+	static void Deserialize(FNetSerializationContext&, const FNetDeserializeArgs&);
+
+	static void SerializeDelta(FNetSerializationContext&, const FNetSerializeDeltaArgs& Args);
+	static void DeserializeDelta(FNetSerializationContext&, const FNetDeserializeDeltaArgs& Args);
+
+	static void Quantize(FNetSerializationContext&, const FNetQuantizeArgs&);
+	static void Dequantize(FNetSerializationContext&, const FNetDequantizeArgs&);
+
+	static bool IsEqual(FNetSerializationContext&, const FNetIsEqualArgs&);
+	static bool Validate(FNetSerializationContext&, const FNetValidateArgs&);
+
+	static void CloneDynamicState(FNetSerializationContext&, const FNetCloneDynamicStateArgs&);
+	static void FreeDynamicState(FNetSerializationContext&, const FNetFreeDynamicStateArgs&);
+
+	static void CollectNetReferences(FNetSerializationContext&, const FNetCollectReferencesArgs&);
+
+private:
+	// Frees dynamic memory allocated by the struct instance. Zeros the struct storage. Does not free the struct storage. After the call the Value is ready to be re-purposed for a different type of struct.
+	static void FreeStructInstance(FNetSerializationContext&, FInstancedStructNetSerializerConfig*, QuantizedType& Value);
+	// Frees dynamic memory allocated by the struct instance, frees the storage for the struct instance and reset the entire quantized state to default.
+	static void Reset(FNetSerializationContext&, FInstancedStructNetSerializerConfig*, QuantizedType&);
+
+	static void InternalFreeStructInstance(FNetSerializationContext&, FInstancedStructNetSerializerConfig*, QuantizedType&);
+
+	class FNetSerializerRegistryDelegates final : private UE::Net::FNetSerializerRegistryDelegates
+	{
+	public:
+		virtual ~FNetSerializerRegistryDelegates()
+		{
+			UE_NET_UNREGISTER_NETSERIALIZER_INFO(FInstancedStructPropertyNetSerializerInfo);
+		}
+
+	private:
+		virtual void OnPreFreezeNetSerializerRegistry() override
+		{
+			UE_NET_REGISTER_NETSERIALIZER_INFO(FInstancedStructPropertyNetSerializerInfo);
+		}
+
+		UE_NET_IMPLEMENT_NETSERIALIZER_INFO(FInstancedStructPropertyNetSerializerInfo);
+	};
+
+	inline static const FNetSerializer* StructNetSerializer = &UE_NET_GET_SERIALIZER(FStructNetSerializer);
+	inline static const FNetSerializer* ObjectNetSerializer = &UE_NET_GET_SERIALIZER(FObjectNetSerializer);
+	inline static FNetSerializerRegistryDelegates NetSerializerRegistryDelegates;
+};
+
+UE_NET_IMPLEMENT_SERIALIZER(FInstancedStructNetSerializer);
+
+void FInstancedStructNetSerializer::Serialize(FNetSerializationContext& Context, const FNetSerializeArgs& Args)
+{
+	const QuantizedType& Value = *reinterpret_cast<QuantizedType*>(Args.Source);
+
+	FInstancedStructNetSerializerConfig* Config = const_cast<FInstancedStructNetSerializerConfig*>(static_cast<const FInstancedStructNetSerializerConfig*>(Args.NetSerializerConfig));
+	
+	FNetBitStreamWriter* Writer = Context.GetBitStreamWriter();
+
+	FStructNetSerializerConfig StructConfig;
+	if (!Value.StructName.IsNone())
+	{
+		StructConfig.StateDescriptor = Config->DescriptorCache.FindOrAddDescriptor(Value.StructName);
+		ensureMsgf(StructConfig.StateDescriptor.IsValid(), TEXT("Struct type is no longer resolvable: %s. Sending FInstancedStruct as uninitialized."), ToCStr(Value.StructName.ToString()));
+	}
+
+	if (Writer->WriteBool(StructConfig.StateDescriptor.IsValid()))
+	{
+		// Serialize struct type
+		{
+			FNetSerializeArgs SerializeArgs = Args;
+			SerializeArgs.Source = NetSerializerValuePointer(&Value.StructType);
+			SerializeArgs.NetSerializerConfig = ObjectNetSerializer->DefaultConfig;
+			ObjectNetSerializer->Serialize(Context, SerializeArgs);
+		}
+
+		// Serialize struct data
+		{
+			FNetSerializeArgs SerializeArgs = Args;
+			SerializeArgs.Source = NetSerializerValuePointer(Value.StructData.GetData());
+			SerializeArgs.NetSerializerConfig = NetSerializerConfigParam(&StructConfig);
+			StructNetSerializer->Serialize(Context, SerializeArgs);
+		}
+	}
+}
+
+void FInstancedStructNetSerializer::Deserialize(FNetSerializationContext& Context, const FNetDeserializeArgs& Args)
+{
+	QuantizedType& Value = *reinterpret_cast<QuantizedType*>(Args.Target);
+
+	FInstancedStructNetSerializerConfig* Config = const_cast<FInstancedStructNetSerializerConfig*>(static_cast<const FInstancedStructNetSerializerConfig*>(Args.NetSerializerConfig));
+
+	FNetBitStreamReader* Reader = Context.GetBitStreamReader();
+	// Was the instanced struct valid on the sending side?
+	if (Reader->ReadBool())
+	{
+		FNetObjectReference StructType;
+		const UScriptStruct* Struct = nullptr;
+
+		// Deserialize struct type
+		{
+			FNetDeserializeArgs DeserializeArgs = Args;
+			DeserializeArgs.Target = NetSerializerValuePointer(&StructType);
+			DeserializeArgs.NetSerializerConfig = ObjectNetSerializer->DefaultConfig;
+			ObjectNetSerializer->Deserialize(Context, DeserializeArgs);
+		}
+
+		// Dequantize to get the UScriptStruct.
+		// $IRIS TODO : Allow receiving end to skip payloads which it's unable to parse due to missing struct.
+		{
+			UObject* Object = nullptr;
+
+			FNetDequantizeArgs DequantizeArgs = {};
+			DequantizeArgs.Source = NetSerializerValuePointer(&StructType);
+			DequantizeArgs.Target = NetSerializerValuePointer(&Object);
+			DequantizeArgs.NetSerializerConfig = ObjectNetSerializer->DefaultConfig;
+			ObjectNetSerializer->Dequantize(Context, DequantizeArgs);
+
+			if (Object != nullptr)
+			{
+				Struct = Cast<UScriptStruct>(Object);
+				ensureMsgf(Struct != nullptr, TEXT("Unable to cast object %s to UScriptStruct"), ToCStr(Object->GetPathName()));
+				if (Struct == nullptr)
+				{
+					Context.SetError(GNetError_InvalidValue);
+					return;
+				}
+			}
+		}
+
+		// Deserialize struct data
+		if (ensureMsgf(Struct != nullptr, TEXT("Unable to find struct using NetObjectReference %s"), ToCStr(StructType.ToString())))
+		{
+			FStructNetSerializerConfig StructConfig;
+			StructConfig.StateDescriptor = Config->DescriptorCache.FindOrAddDescriptor(Struct);
+			if (!StructConfig.StateDescriptor.IsValid())
+			{
+				ensureMsgf(StructConfig.StateDescriptor.IsValid(), TEXT("Unable to create ReplicationStateDescriptor for struct %s."), ToCStr(Struct->GetPathName()));
+				Context.SetError(GNetError_InvalidValue);
+				return;
+			}
+
+			// If we changed types we need to free the struct, adjust the storage size and update struct info.
+			if (StructType != Value.StructType)
+			{
+				FreeStructInstance(Context, Config, Value);
+				Value.StructData.AdjustSize(Context, StructConfig.StateDescriptor->InternalSize, StructConfig.StateDescriptor->InternalAlignment);
+				Value.StructType = StructType;
+
+				Value.StructDescriptorTraits = StructConfig.StateDescriptor->Traits;
+				const FString& PathNameString = Struct->GetPathName();
+				Value.StructName = FName(PathNameString);
+			}
+
+			FNetDeserializeArgs DeserializeArgs = Args;
+			DeserializeArgs.Target = NetSerializerValuePointer(Value.StructData.GetData());
+			DeserializeArgs.NetSerializerConfig = NetSerializerConfigParam(&StructConfig);
+			StructNetSerializer->Deserialize(Context, DeserializeArgs);
+		}
+		else
+		{
+			Context.SetError(GNetError_InvalidValue);
+			return;
+		}
+	}
+	else
+	{
+		Reset(Context, Config, Value);
+	}
+}
+
+void FInstancedStructNetSerializer::SerializeDelta(FNetSerializationContext& Context, const FNetSerializeDeltaArgs& Args)
+{
+	// Skip DC support for now. Need to figure out how to gracefully handle missing UScriptStruct on the receiving end.
+	Serialize(Context, Args);
+}
+
+void FInstancedStructNetSerializer::DeserializeDelta(FNetSerializationContext& Context, const FNetDeserializeDeltaArgs& Args)
+{
+	// Skip DC support for now. Need to figure out how to gracefully handle missing UScriptStruct on the receiving end.
+	Deserialize(Context, Args);
+}
+
+void FInstancedStructNetSerializer::Quantize(FNetSerializationContext& Context, const FNetQuantizeArgs& Args)
+{
+	const SourceType& Source = *reinterpret_cast<SourceType*>(Args.Source);
+	QuantizedType& Target = *reinterpret_cast<QuantizedType*>(Args.Target);
+
+	FInstancedStructNetSerializerConfig* Config = const_cast<FInstancedStructNetSerializerConfig*>(static_cast<const FInstancedStructNetSerializerConfig*>(Args.NetSerializerConfig));
+
+	if (Source.IsValid())
+	{
+		const UScriptStruct* Struct = Source.GetScriptStruct();
+
+		TRefCountPtr<const FReplicationStateDescriptor> DescriptorRef = Config->DescriptorCache.FindOrAddDescriptor(Struct);
+		if (ensureMsgf(DescriptorRef.IsValid(), TEXT("Unable to create descriptor for struct %s. Unexpected."), ToCStr(Struct->GetFullName())))
+		{
+			const FString& PathNameString = Struct->GetPathName();
+			const FName PathName(PathNameString);
+			// If the struct type is the same as previous instance we don't need to free memory or adjust allocations.
+			if (PathName != Target.StructName)
+			{
+				// We need to free the previous struct data prior to overwriting it. Doing this early.
+				FreeStructInstance(Context, Config, Target);
+
+				// Adjust struct storage size
+				Target.StructData.AdjustSize(Context, DescriptorRef->InternalSize, DescriptorRef->InternalAlignment);
+
+				// Adjust struct name and traits
+				Target.StructName = PathName;
+				Target.StructDescriptorTraits = DescriptorRef->Traits;
+
+				// Quantize the struct type since the receiver need it to be able to serialize the data properly.
+				{
+					FNetQuantizeArgs QuantizeArgs = Args;
+					QuantizeArgs.Source = NetSerializerValuePointer(&Struct);
+					QuantizeArgs.Target = NetSerializerValuePointer(&Target.StructType);
+					QuantizeArgs.NetSerializerConfig = ObjectNetSerializer->DefaultConfig;
+					ObjectNetSerializer->Quantize(Context, QuantizeArgs);
+				}
+			}
+
+			// Quantize the struct instance into the target memory.
+			if (Target.StructData.Num() > 0)
+			{
+				FStructNetSerializerConfig StructConfig;
+				StructConfig.StateDescriptor = DescriptorRef;
+
+				FNetQuantizeArgs QuantizeArgs = Args;
+				QuantizeArgs.Source = NetSerializerValuePointer(Source.GetMemory());
+				QuantizeArgs.Target = NetSerializerValuePointer(Target.StructData.GetData());
+				QuantizeArgs.NetSerializerConfig = NetSerializerConfigParam(&StructConfig);
+				StructNetSerializer->Quantize(Context, QuantizeArgs);
+			}
+
+			return;
+		}
+	}
+
+	// Path taken for uninitialized FInstancedStruct or if an error was detected.
+	Reset(Context, Config, Target);
+}
+
+// $IRIS TODO : Consider implementing Apply to avoid unnecessary memory operations.
+void FInstancedStructNetSerializer::Dequantize(FNetSerializationContext& Context, const FNetDequantizeArgs& Args)
+{
+	const QuantizedType& Source = *reinterpret_cast<QuantizedType*>(Args.Source);
+	SourceType& Target = *reinterpret_cast<SourceType*>(Args.Target);
+
+	FInstancedStructNetSerializerConfig* Config = const_cast<FInstancedStructNetSerializerConfig*>(static_cast<const FInstancedStructNetSerializerConfig*>(Args.NetSerializerConfig));
+
+	if (!Source.StructType.IsValid())
+	{
+		Target.Reset();
+		return;
+	}
+
+	const UScriptStruct* Struct = nullptr;
+	{
+		UObject* Object = nullptr;
+
+		FNetDequantizeArgs DequantizeArgs = {};
+		DequantizeArgs.Source = NetSerializerValuePointer(&Source.StructType);
+		DequantizeArgs.Target = NetSerializerValuePointer(&Object);
+		DequantizeArgs.NetSerializerConfig = ObjectNetSerializer->DefaultConfig;
+		ObjectNetSerializer->Dequantize(Context, DequantizeArgs);
+
+		if (Object != nullptr)
+		{
+			Struct = Cast<UScriptStruct>(Object);
+			ensureMsgf(Struct != nullptr, TEXT("Unable to cast object %s to UScriptStruct"), ToCStr(Object->GetPathName()));
+		}
+	}
+	
+	if (ensureMsgf(Struct != nullptr, TEXT("Unable to find struct using NetObjectReference %s"), ToCStr(Source.StructType.ToString())))
+	{
+		// Re-initialize if the type changes.
+		if (Struct != Target.GetScriptStruct())
+		{
+			Target.InitializeAs(Struct);
+		}
+
+		FStructNetSerializerConfig StructConfig;
+		StructConfig.StateDescriptor = Config->DescriptorCache.FindOrAddDescriptor(Struct);
+
+		if (ensureMsgf(StructConfig.StateDescriptor.IsValid(), TEXT("Unable to create ReplicationStateDescriptor for struct %s."), ToCStr(Source.StructName.ToString())))
+		{
+			FNetDequantizeArgs DequantizeArgs = Args;
+			DequantizeArgs.NetSerializerConfig = &StructConfig;
+			DequantizeArgs.Source = NetSerializerValuePointer(Source.StructData.GetData());
+			DequantizeArgs.Target = NetSerializerValuePointer(Target.GetMutableMemory());
+			StructNetSerializer->Dequantize(Context, DequantizeArgs);
+		}
+	}
+}
+
+bool FInstancedStructNetSerializer::IsEqual(FNetSerializationContext& Context, const FNetIsEqualArgs& Args)
+{
+	if (Args.bStateIsQuantized)
+	{
+		const QuantizedType& Value0 = *reinterpret_cast<const QuantizedType*>(Args.Source0);
+		const QuantizedType& Value1 = *reinterpret_cast<const QuantizedType*>(Args.Source1);
+
+		if (Value0.StructData.Num() != Value1.StructData.Num())
+		{
+			return false;
+		}
+
+		if (Value0.StructType != Value1.StructType)
+		{
+			return false;
+		}
+
+		if (FMemory::Memcmp(Value0.StructData.GetData(), Value1.StructData.GetData(), Value0.StructData.Num()) != 0)
+		{
+			return false;
+		}
+
+		return true;
+	}
+	else
+	{
+		const SourceType& Value0 = *reinterpret_cast<const SourceType*>(Args.Source0);
+		const SourceType& Value1 = *reinterpret_cast<const SourceType*>(Args.Source1);
+		return Value0 == Value1;
+	}
+}
+
+bool FInstancedStructNetSerializer::Validate(FNetSerializationContext& Context, const FNetValidateArgs& Args)
+{
+	const SourceType& Value = *reinterpret_cast<const SourceType*>(Args.Source);
+
+	return true;
+}
+
+void FInstancedStructNetSerializer::CloneDynamicState(FNetSerializationContext& Context, const FNetCloneDynamicStateArgs& Args)
+{
+	const QuantizedType& Source = *reinterpret_cast<const QuantizedType*>(Args.Source);
+	QuantizedType& Target = *reinterpret_cast<QuantizedType*>(Args.Target);
+	FInstancedStructNetSerializerConfig* Config = const_cast<FInstancedStructNetSerializerConfig*>(static_cast<const FInstancedStructNetSerializerConfig*>(Args.NetSerializerConfig));
+
+	Target.StructData.Clone(Context, Source.StructData);
+
+	if (EnumHasAnyFlags(Source.StructDescriptorTraits, EReplicationStateTraits::HasDynamicState))
+	{
+		FStructNetSerializerConfig StructConfig;
+		StructConfig.StateDescriptor = Config->DescriptorCache.FindOrAddDescriptor(Source.StructName);
+
+		if (ensureMsgf(StructConfig.StateDescriptor.IsValid(), TEXT("Unable to create ReplicationStateDescriptor for struct %s."), ToCStr(Source.StructName.ToString())))
+		{
+			FNetCloneDynamicStateArgs CloneArgs = Args;
+			CloneArgs.NetSerializerConfig = &StructConfig;
+			CloneArgs.Source = NetSerializerValuePointer(Source.StructData.GetData());
+			CloneArgs.Target = NetSerializerValuePointer(Target.StructData.GetData());
+			StructNetSerializer->CloneDynamicState(Context, CloneArgs);
+		}
+	}
+}
+
+void FInstancedStructNetSerializer::FreeDynamicState(FNetSerializationContext& Context, const FNetFreeDynamicStateArgs& Args)
+{
+	QuantizedType& Value = *reinterpret_cast<QuantizedType*>(Args.Source);
+	FInstancedStructNetSerializerConfig* Config = const_cast<FInstancedStructNetSerializerConfig*>(static_cast<const FInstancedStructNetSerializerConfig*>(Args.NetSerializerConfig));
+	
+	InternalFreeStructInstance(Context, Config, Value);
+	Value.StructData.Free(Context);
+}
+
+void FInstancedStructNetSerializer::CollectNetReferences(FNetSerializationContext& Context, const FNetCollectReferencesArgs& Args)
+{
+	const QuantizedType& Value = *reinterpret_cast<QuantizedType*>(Args.Source);
+	FInstancedStructNetSerializerConfig* Config = const_cast<FInstancedStructNetSerializerConfig*>(static_cast<const FInstancedStructNetSerializerConfig*>(Args.NetSerializerConfig));
+
+	if (Value.StructType.IsValid())
+	{
+		FNetReferenceCollector& Collector = *reinterpret_cast<FNetReferenceCollector*>(Args.Collector);
+
+		// What's the proper reference type?
+		const FNetReferenceInfo ReferenceInfo(FNetReferenceInfo::EResolveType::ResolveOnClient);
+		Collector.Add(ReferenceInfo, Value.StructType, Args.ChangeMaskInfo);
+	}
+
+	if (EnumHasAnyFlags(Value.StructDescriptorTraits, EReplicationStateTraits::HasObjectReference))
+	{
+		FStructNetSerializerConfig StructConfig;
+		StructConfig.StateDescriptor = Config->DescriptorCache.FindOrAddDescriptor(Value.StructName);
+
+		if (ensureMsgf(StructConfig.StateDescriptor.IsValid(), TEXT("Unable to create ReplicationStateDescriptor for struct %s."), ToCStr(Value.StructName.ToString())))
+		{
+			FNetCollectReferencesArgs CollectReferencesArgs = Args;
+			CollectReferencesArgs.NetSerializerConfig = &StructConfig;
+			CollectReferencesArgs.Source = NetSerializerValuePointer(Value.StructData.GetData());
+			StructNetSerializer->CollectNetReferences(Context, CollectReferencesArgs);
+		}
+	}
+}
+
+void FInstancedStructNetSerializer::FreeStructInstance(FNetSerializationContext& Context, FInstancedStructNetSerializerConfig* Config, QuantizedType& Value)
+{
+	InternalFreeStructInstance(Context, Config, Value);
+	if (Value.StructData.Num() > 0)
+	{
+		FMemory::Memzero(Value.StructData.GetData(), Value.StructData.Num());
+	}
+}
+
+void FInstancedStructNetSerializer::Reset(FNetSerializationContext& Context, FInstancedStructNetSerializerConfig* Config, QuantizedType& Value)
+{
+	InternalFreeStructInstance(Context, Config, Value);
+	Value.StructData.Free(Context);
+
+	FMemory::Memzero(&Value, sizeof(QuantizedType));
+}
+
+void FInstancedStructNetSerializer::InternalFreeStructInstance(FNetSerializationContext& Context, FInstancedStructNetSerializerConfig* Config, QuantizedType& Value)
+{
+	if (Value.StructData.Num() > 0 && EnumHasAnyFlags(Value.StructDescriptorTraits, EReplicationStateTraits::HasDynamicState))
+	{
+		FStructNetSerializerConfig StructConfig;
+		StructConfig.StateDescriptor = Config->DescriptorCache.FindOrAddDescriptor(Value.StructName);
+
+		FNetFreeDynamicStateArgs FreeArgs;
+		FreeArgs.NetSerializerConfig = &StructConfig;
+		FreeArgs.Source = NetSerializerValuePointer(Value.StructData.GetData());
+	}
+}
+
+}
+
+namespace UE::Net
+{
+
+void InitInstancedStructNetSerializerConfig(FInstancedStructNetSerializerConfig* Config, const FProperty* Property)
+{
+	// We want to be explicit about which structs are supported in the config. That requires UE-180981. For now let's allow any UScriptStruct.
+	Config->SupportedTypes.Reset();
+
+	{
+		FString DebugName;
+		DebugName.Reserve(256);
+
+		FFieldVariant Owner = Property->GetOwnerVariant();
+		if (const UObject* Object = Owner.ToUObject())
+		{
+			DebugName.Append(Object->GetName()).AppendChar(TEXT('.'));
+		}
+		else if (const FField* Field = Owner.ToField())
+		{
+			DebugName.Append(Field->GetName()).AppendChar(TEXT('.'));
+		}
+		DebugName.Append(Property->GetName());
+
+		Config->DescriptorCache.SetDebugName(DebugName);
+	}
+
+	// Add supported type info to the cache.
+	Config->DescriptorCache.AddSupportedTypes(TConstArrayView<TSoftObjectPtr<UScriptStruct>>(Config->SupportedTypes));
+
+	const bool bIsAllowingArbitraryStruct = true;
+	if (bIsAllowingArbitraryStruct)
+	{
+		Config->DescriptorCache.SetMaxCachedDescriptorCount(MaxCachedInstancedStructDescriptorCount);
+	}
+}
+
+FInstancedStructPropertyNetSerializerInfo::FInstancedStructPropertyNetSerializerInfo()
+: FNamedStructPropertyNetSerializerInfo(FName("InstancedStruct"), UE_NET_GET_SERIALIZER(FInstancedStructNetSerializer))
+{
+}
+
+bool FInstancedStructPropertyNetSerializerInfo::CanUseDefaultConfig(const FProperty* Property) const
+{
+	// Creating property specific configs so that we can validate and allow only property specific types. This allows us property specific tracking of used types too.
+	return false;
+}
+
+FNetSerializerConfig* FInstancedStructPropertyNetSerializerInfo::BuildNetSerializerConfig(void* NetSerializerConfigBuffer, const FProperty* Property) const
+{
+	FInstancedStructNetSerializerConfig* Config = new (NetSerializerConfigBuffer) FInstancedStructNetSerializerConfig();
+	InitInstancedStructNetSerializerConfig(Config, Property);
+	return Config;
+}
+
+}
+
+namespace UE::Net::Private
+{
+
+FInstancedStructDescriptorCache::FInstancedStructDescriptorCache()
+{
+	// Intentionally left empty for debugging purposes.
+}
+
+FInstancedStructDescriptorCache::~FInstancedStructDescriptorCache()
+{
+	// Intentionally left empty for debugging purposes.
+}
+
+// Name for debugging purposes
+void FInstancedStructDescriptorCache::SetDebugName(const FString& InDebugName)
+{
+	DebugName = InDebugName;
+}
+
+void FInstancedStructDescriptorCache::SetMaxCachedDescriptorCount(int32 MaxCount)
+{
+	if (MaxCount <= 0)
+	{
+		DescriptorLruCache.Empty(0);
+		MaxCachedDescriptorCount = 0;
+	}
+	else
+	{
+		// Clear DescriptorMap which is only used for unlimited MaxCount
+		UE_CLOG(!DescriptorMap.IsEmpty(), LogIris, Warning, TEXT("Clearing DescriptorMap from FIstancedStructDescriptorCache %s"), ToCStr(DebugName));
+		DescriptorMap.Empty();
+		DescriptorLruCache.Empty(MaxCount);
+		MaxCachedDescriptorCount = MaxCount;
+	}
+}
+
+void FInstancedStructDescriptorCache::AddSupportedTypes(const TConstArrayView<TSoftObjectPtr<UScriptStruct>>& InSupportedTypes)
+{
+	for (const TSoftObjectPtr<UScriptStruct>& Type : InSupportedTypes)
+	{
+		SupportedTypes.Add(Type);
+	}
+}
+
+
+bool FInstancedStructDescriptorCache::IsSupportedType(const UScriptStruct* Struct) const
+{
+	if (ensure(Struct != nullptr))
+	{
+		if (SupportedTypes.IsEmpty())
+		{
+			return true;
+		}
+
+		for (const TSoftObjectPtr<UScriptStruct>& SupportedType : SupportedTypes)
+		{
+			if (const UScriptStruct* SupportedStruct = SupportedType.Get())
+			{
+				if (Struct->IsChildOf(SupportedStruct))
+				{
+					return true;
+				}
+			}
+		}
+	}
+
+	return false;
+}
+
+TRefCountPtr<const FReplicationStateDescriptor> FInstancedStructDescriptorCache::FindDescriptor(FName StructPath)
+{
+	FScopeLock Lock(&Mutex);
+	if (MaxCachedDescriptorCount > 0)
+	{
+		return DescriptorLruCache.FindAndTouchRef(StructPath);
+	}
+	else
+	{
+		return DescriptorMap.FindRef(StructPath);
+	}
+}
+
+TRefCountPtr<const FReplicationStateDescriptor> FInstancedStructDescriptorCache::FindOrAddDescriptor(FName StructPath)
+{
+	TRefCountPtr<const FReplicationStateDescriptor> Descriptor;
+
+	Descriptor = FindDescriptor(StructPath);
+	if (Descriptor.IsValid())
+	{
+		return Descriptor;
+	}
+
+	const UObject* Object = StaticLoadObject(UScriptStruct::StaticClass(), nullptr, ToCStr(StructPath.ToString()), nullptr, LOAD_None);
+	if (const UScriptStruct* Struct = Cast<UScriptStruct>(Object))
+	{
+		Descriptor = CreateAndCacheDescriptor(Struct, StructPath);
+	}
+	else
+	{
+		// Cast fail?
+		ensureMsgf(Object == nullptr, TEXT("Unable to cast object %s to UScriptStruct"), ToCStr(Object->GetPathName()));
+	}
+
+	return Descriptor;
+}
+
+TRefCountPtr<const FReplicationStateDescriptor> FInstancedStructDescriptorCache::FindOrAddDescriptor(const UScriptStruct* Struct)
+{
+	TRefCountPtr<const FReplicationStateDescriptor> Descriptor;
+	if (ensure(Struct != nullptr))
+	{
+		const FString& PathNameString = Struct->GetPathName();
+		const FName PathName(PathNameString);
+
+		Descriptor = FindDescriptor(PathName);
+		if (Descriptor.IsValid())
+		{
+			return Descriptor;
+		}
+
+		if (!IsSupportedType(Struct))
+		{
+			return Descriptor;
+		}
+
+		// Create descriptor and add it to the cache.
+		Descriptor = CreateAndCacheDescriptor(Struct, PathName);
+	}
+
+	return Descriptor;
+}
+
+TRefCountPtr<const FReplicationStateDescriptor> FInstancedStructDescriptorCache::CreateAndCacheDescriptor(const UScriptStruct* Struct, FName StructPath)
+{
+	TRefCountPtr<const FReplicationStateDescriptor> Descriptor;
+
+	FReplicationStateDescriptorBuilder::FParameters Params;
+	Descriptor = FReplicationStateDescriptorBuilder::CreateDescriptorForStruct(Struct, Params);
+
+	{
+		FScopeLock Lock(&Mutex);
+		if (MaxCachedDescriptorCount > 0)
+		{
+			DescriptorLruCache.Add(StructPath, Descriptor);
+		}
+		else
+		{
+			DescriptorMap.Add(StructPath, Descriptor);
+		}
+	}
+
+	return Descriptor;
+}
+
+}
+
+#endif
+
diff --git a/Engine/Plugins/Experimental/StructUtils/Source/StructUtils/Public/InstancedStructNetSerializer.h b/Engine/Plugins/Experimental/StructUtils/Source/StructUtils/Public/InstancedStructNetSerializer.h
new file mode 100644
--- /dev/null
+++ b/Engine/Plugins/Experimental/StructUtils/Source/StructUtils/Public/InstancedStructNetSerializer.h
@@ -0,0 +1,91 @@
+// Copyright Epic Games, Inc. All Rights Reserved.
+
+#pragma once
+
+#include "Iris/ReplicationState/ReplicationStateDescriptor.h"
+#include "Iris/Serialization/NetSerializer.h"
+#include "Containers/LruCache.h"
+#include "UObject/ObjectMacros.h"
+
+#include "InstancedStructNetSerializer.generated.h"
+
+class UScriptStruct;
+
+namespace UE::Net::Private
+{
+
+class FInstancedStructDescriptorCache
+{
+public:
+	FInstancedStructDescriptorCache();
+	~FInstancedStructDescriptorCache();
+
+	// Name for debugging purposes
+	void SetDebugName(const FString& DebugName);
+
+	// Set max cached descriptor count. The most recently used descriptors will be kept. MaxCount <= 0 means no limit which is the default.
+	void SetMaxCachedDescriptorCount(int32 MaxCount);
+
+	// 
+	void AddSupportedTypes(const TConstArrayView<TSoftObjectPtr<UScriptStruct>>& SupportedTypes);
+
+	bool IsSupportedType(const UScriptStruct* Struct) const;
+
+	// Find descriptor for struct with fully qualified name.
+	TRefCountPtr<const FReplicationStateDescriptor> FindDescriptor(FName StructPath);
+
+	// Find or create descriptor for struct with fully qualified name.
+	TRefCountPtr<const FReplicationStateDescriptor> FindOrAddDescriptor(FName StructPath);
+
+	// Find or create descriptor for struct.
+	TRefCountPtr<const FReplicationStateDescriptor> FindOrAddDescriptor(const UScriptStruct* Struct);
+
+private:
+	TRefCountPtr<const FReplicationStateDescriptor> CreateAndCacheDescriptor(const UScriptStruct* Struct, FName StructPath);
+
+	FCriticalSection Mutex;
+	// LRU cache for descriptors for limited descriptor counts.
+	TLruCache<FName, TRefCountPtr<const FReplicationStateDescriptor>> DescriptorLruCache;
+	// Map struct name -> FReplicationStateDescriptor for unlimited descriptor counts. 
+	TMap<FName, TRefCountPtr<const FReplicationStateDescriptor>> DescriptorMap;
+	// Supported types. An empty set indicates all UScriptStructs are supported.
+	TSet<TSoftObjectPtr<UScriptStruct>> SupportedTypes;
+	FString DebugName;
+	int32 MaxCachedDescriptorCount = 0;
+};
+
+}
+
+USTRUCT()
+struct FInstancedStructNetSerializerConfig : public FNetSerializerConfig
+{
+	GENERATED_BODY()
+
+	FInstancedStructNetSerializerConfig();
+	~FInstancedStructNetSerializerConfig();
+
+	FInstancedStructNetSerializerConfig(const FInstancedStructNetSerializerConfig&) = delete;
+	FInstancedStructNetSerializerConfig& operator=(const FInstancedStructNetSerializerConfig&) = delete;
+
+	// The property is for serialization support. We store the supported types differently in the descriptor cache.
+	UPROPERTY()
+	TArray<TSoftObjectPtr<UScriptStruct>> SupportedTypes;
+
+	UE::Net::Private::FInstancedStructDescriptorCache DescriptorCache;
+};
+
+template<>
+struct TStructOpsTypeTraits<FInstancedStructNetSerializerConfig> : public TStructOpsTypeTraitsBase2<FInstancedStructNetSerializerConfig>
+{
+	enum
+	{
+		WithCopy = false
+	};
+};
+
+namespace UE::Net
+{
+
+UE_NET_DECLARE_SERIALIZER(FInstancedStructNetSerializer, STRUCTUTILS_API);
+
+}
diff --git a/Engine/Plugins/Experimental/StructUtils/Source/StructUtils/StructUtils.Build.cs b/Engine/Plugins/Experimental/StructUtils/Source/StructUtils/StructUtils.Build.cs
--- a/Engine/Plugins/Experimental/StructUtils/Source/StructUtils/StructUtils.Build.cs
+++ b/Engine/Plugins/Experimental/StructUtils/Source/StructUtils/StructUtils.Build.cs
@@ -21,6 +21,8 @@ namespace UnrealBuildTool.Rules
 			{
 				PrivateDependencyModuleNames.Add("Engine");
 			}
+
+			SetupIrisSupport(Target);
 		}
 	}
 }
diff --git a/Engine/Plugins/Runtime/ReplicationSystemTestPlugin/Source/Private/Tests/Serialization/TestInstancedStructNetSerializer.cpp b/Engine/Plugins/Runtime/ReplicationSystemTestPlugin/Source/Private/Tests/Serialization/TestInstancedStructNetSerializer.cpp
new file mode 100644
--- /dev/null
+++ b/Engine/Plugins/Runtime/ReplicationSystemTestPlugin/Source/Private/Tests/Serialization/TestInstancedStructNetSerializer.cpp
@@ -0,0 +1,768 @@
+// Copyright Epic Games, Inc. All Rights Reserved.
+
+#include "Tests/Serialization/TestInstancedStructNetSerializer.h"
+#include "Tests/Serialization/TestNetSerializerFixture.h"
+#include "Tests/ReplicationSystem/ReplicationSystemServerClientTestFixture.h"
+#include "Iris/ReplicationSystem/ReplicationFragmentUtil.h"
+#include "Iris/ReplicationSystem/ReplicationSystemInternal.h"
+#include "Iris/ReplicationState/ReplicationStateDescriptorBuilder.h"
+#include "Iris/Serialization/InternalNetSerializationContext.h"
+#include "Iris/Serialization/NetReferenceCollector.h"
+#include "Iris/Serialization/NetSerializers.h"
+#include "InstancedStructNetSerializer.h"
+#include "Net/UnrealNetwork.h"
+
+namespace UE::Net
+{
+
+FTestMessage& operator<<(FTestMessage& Message, const FInstancedStruct& InstancedStruct)
+{
+	FString String;
+	InstancedStruct.ExportTextItem(String, FInstancedStruct(), nullptr, 0, nullptr);
+	return Message << String;
+}
+
+}
+
+namespace UE::Net::Private
+{
+
+FTestMessage& PrintInstancedStructNetSerializerConfig(FTestMessage& Message, const FNetSerializerConfig& InConfig)
+{
+	const FInstancedStructNetSerializerConfig& Config = static_cast<const FInstancedStructNetSerializerConfig&>(InConfig);
+	for (const auto& SupportedType : Config.SupportedTypes)
+	{
+		Message << SupportedType.Get()->GetFullName();
+	}
+	return Message;
+}
+
+class FTestInstancedStructNetSerializerFixture : public FReplicationSystemServerClientTestFixture
+{
+public:
+	FTestInstancedStructNetSerializerFixture() : FReplicationSystemServerClientTestFixture() {}
+
+protected:
+	virtual void SetUp() override;
+	virtual void TearDown() override;
+
+	// Composable operations for testing the serializer
+	void Serialize();
+	void Deserialize();
+	void SerializeDelta();
+	void DeserializeDelta();
+	void Quantize();
+	void QuantizeTwoStates();
+	void Dequantize();
+	bool IsEqual(bool bQuantized);
+	void Clone();
+	void Validate();
+	void FreeQuantizedState();
+
+	// Instantiates a struct with a modified property on InstancedStruct0
+	void SetNonDefaultInstanceState();
+
+	// Adds multiple elements, at least one uninitialized and at least two initialized with at least one modified property, to InstancedStructArray0
+	void SetNonDefaultArrayState();
+
+	bool AreEqual(const FInstancedStruct* Value0, const FInstancedStruct* Value1);
+
+	const FNetSerializer* GetInstancedStructNetSerializer(bool bIsArray) const;
+	const FNetSerializerConfig* GetInstancedStructNetSerializerConfig(bool bIsArray) const;
+
+protected:
+	FNetSerializationContext NetSerializationContext;
+	FInternalNetSerializationContext InternalNetSerializationContext;
+
+	FInstancedStruct InstancedStruct0;
+	FInstancedStruct InstancedStruct1;
+	TArray<FInstancedStruct> InstancedStructArray0;
+	TArray<FInstancedStruct> InstancedStructArray1;
+
+	FStructNetSerializerConfig InstancedStructNetSerializerConfig;
+	FStructNetSerializerConfig InstancedStructArrayNetSerializerConfig;
+
+	alignas(16) uint8 QuantizedBuffer[2][128];
+	alignas(16) uint8 ClonedQuantizedBuffer[2][128];
+	alignas(16) uint8 BitStreamBuffer[2048];
+
+	bool bHasQuantizedState = false;
+	bool bHasClonedQuantizedState = false;
+	bool bIsTestingArray = false;
+
+	uint32 QuantizedStateCount = 0;
+	uint32 ClonedQuantizedStateCount = 0;
+
+	FNetBitStreamWriter Writer;
+	FNetBitStreamReader Reader;
+};
+
+}
+
+namespace UE::Net::Private
+{
+
+// Instance tests
+UE_NET_TEST_FIXTURE(FTestInstancedStructNetSerializerFixture, TestQuantizeUninitialized)
+{
+	Quantize();
+}
+
+UE_NET_TEST_FIXTURE(FTestInstancedStructNetSerializerFixture, TestQuantizeInitialized)
+{
+	SetNonDefaultInstanceState();
+	Quantize();
+}
+
+UE_NET_TEST_FIXTURE(FTestInstancedStructNetSerializerFixture, TestDequantizeUninitialized)
+{
+	Quantize();
+	Dequantize();
+	UE_NET_ASSERT_EQ(InstancedStruct0, InstancedStruct1);
+}
+
+UE_NET_TEST_FIXTURE(FTestInstancedStructNetSerializerFixture, TestDequantizeInitialized)
+{
+	SetNonDefaultInstanceState();
+	Quantize();
+	Dequantize();
+	UE_NET_ASSERT_EQ(InstancedStruct0, InstancedStruct1);
+}
+
+UE_NET_TEST_FIXTURE(FTestInstancedStructNetSerializerFixture, TestSerializeUninitialized)
+{
+	Quantize();
+	Serialize();
+}
+
+UE_NET_TEST_FIXTURE(FTestInstancedStructNetSerializerFixture, TestSerializeInitialized)
+{
+	SetNonDefaultInstanceState();
+	Quantize();
+	Serialize();
+}
+
+UE_NET_TEST_FIXTURE(FTestInstancedStructNetSerializerFixture, TestDeserializeUninitialized)
+{
+	Quantize();
+	Serialize();
+	FreeQuantizedState();
+	Deserialize();
+}
+
+UE_NET_TEST_FIXTURE(FTestInstancedStructNetSerializerFixture, TestDeserializeInitialized)
+{
+	SetNonDefaultInstanceState();
+	Quantize();
+	Serialize();
+	FreeQuantizedState();
+	Deserialize();
+}
+
+UE_NET_TEST_FIXTURE(FTestInstancedStructNetSerializerFixture, TestDequantizeSerializedUninitializedState)
+{
+	Quantize();
+	Serialize();
+	FreeQuantizedState();
+	Deserialize();
+	Dequantize();
+	UE_NET_ASSERT_EQ(InstancedStruct0, InstancedStruct1);
+}
+
+UE_NET_TEST_FIXTURE(FTestInstancedStructNetSerializerFixture, TestDequantizeSerializedInitializedState)
+{
+	SetNonDefaultInstanceState();
+	Quantize();
+	Serialize();
+	FreeQuantizedState();
+	Deserialize();
+	Dequantize();
+	UE_NET_ASSERT_EQ(InstancedStruct0, InstancedStruct1);
+}
+
+UE_NET_TEST_FIXTURE(FTestInstancedStructNetSerializerFixture, TestSerializeDeltaEqualStates)
+{
+	QuantizeTwoStates();
+	SerializeDelta();
+}
+
+UE_NET_TEST_FIXTURE(FTestInstancedStructNetSerializerFixture, TestSerializeDeltaNonEqualStates)
+{
+	SetNonDefaultInstanceState();
+	QuantizeTwoStates();
+	SerializeDelta();
+}
+
+UE_NET_TEST_FIXTURE(FTestInstancedStructNetSerializerFixture, TestDeserializeDeltaEqualStates)
+{
+	QuantizeTwoStates();
+	SerializeDelta();
+	DeserializeDelta();
+}
+
+UE_NET_TEST_FIXTURE(FTestInstancedStructNetSerializerFixture, TestDeserializeDeltaNonEqualStates)
+{
+	SetNonDefaultInstanceState();
+	QuantizeTwoStates();
+	SerializeDelta();
+	DeserializeDelta();
+}
+
+UE_NET_TEST_FIXTURE(FTestInstancedStructNetSerializerFixture, TestDequantizeDeltaSerializedState)
+{
+	SetNonDefaultInstanceState();
+	QuantizeTwoStates();
+	SerializeDelta();
+	DeserializeDelta();
+	Dequantize();
+
+	UE_NET_ASSERT_EQ(InstancedStruct0, InstancedStruct1);
+}
+
+UE_NET_TEST_FIXTURE(FTestInstancedStructNetSerializerFixture, TestCollectReferencesUninitialized)
+{
+	Quantize();
+
+	FNetReferenceCollector Collector;
+
+	FNetCollectReferencesArgs Args = {};
+	Args.NetSerializerConfig = NetSerializerConfigParam(GetInstancedStructNetSerializerConfig(bIsTestingArray));
+	Args.Source = NetSerializerValuePointer(&QuantizedBuffer[0]);
+	Args.Collector = NetSerializerValuePointer(&Collector);
+	GetInstancedStructNetSerializer(bIsTestingArray)->CollectNetReferences(NetSerializationContext, Args);
+
+	UE_NET_ASSERT_EQ(Collector.GetCollectedReferences().Num(), 0);
+}
+
+UE_NET_TEST_FIXTURE(FTestInstancedStructNetSerializerFixture, TestCollectReferencesStructNoRef)
+{
+	InstancedStruct0.InitializeAs<FStructForInstancedStructTestD>();
+
+	Quantize();
+
+	FNetReferenceCollector Collector(ENetReferenceCollectorTraits::IncludeInvalidReferences);
+
+	FNetCollectReferencesArgs Args = {};
+	Args.NetSerializerConfig = NetSerializerConfigParam(GetInstancedStructNetSerializerConfig(bIsTestingArray));
+	Args.Source = NetSerializerValuePointer(&QuantizedBuffer[0]);
+	Args.Collector = NetSerializerValuePointer(&Collector);
+
+	GetInstancedStructNetSerializer(bIsTestingArray)->CollectNetReferences(NetSerializationContext, Args);
+
+	UE_NET_ASSERT_EQ(Collector.GetCollectedReferences().Num(), 1);
+}
+
+UE_NET_TEST_FIXTURE(FTestInstancedStructNetSerializerFixture, TestCollectReferencesStructWithRef)
+{
+	InstancedStruct0.InitializeAs<FStructForInstancedStructTestWithObjectReference>();
+	InstancedStruct0.GetMutable<FStructForInstancedStructTestWithObjectReference>().SomeObject = FStructForInstancedStructTestWithObjectReference::StaticStruct();
+
+	Quantize();
+
+	FNetReferenceCollector Collector(ENetReferenceCollectorTraits::IncludeInvalidReferences);
+
+	FNetCollectReferencesArgs Args = {};
+	Args.NetSerializerConfig = NetSerializerConfigParam(GetInstancedStructNetSerializerConfig(bIsTestingArray));
+	Args.Source = NetSerializerValuePointer(&QuantizedBuffer[0]);
+	Args.Collector = NetSerializerValuePointer(&Collector);
+
+	GetInstancedStructNetSerializer(bIsTestingArray)->CollectNetReferences(NetSerializationContext, Args);
+
+	UE_NET_ASSERT_GE(Collector.GetCollectedReferences().Num(), 2);
+}
+
+UE_NET_TEST_FIXTURE(FTestInstancedStructNetSerializerFixture, TestIsEqualExternal)
+{
+	constexpr bool bUseQuantizedState = false;
+
+	// Default state compared to default state
+	InstancedStruct0.Reset();
+	InstancedStruct1.Reset();
+	UE_NET_ASSERT_TRUE(IsEqual(bUseQuantizedState));
+
+	// Non-default state compared to default state
+	SetNonDefaultInstanceState();
+	UE_NET_ASSERT_FALSE(IsEqual(bUseQuantizedState));
+
+	// Non-default state compared to non-default state
+	InstancedStruct1 = InstancedStruct0;
+	UE_NET_ASSERT_TRUE(IsEqual(bUseQuantizedState));
+}
+
+UE_NET_TEST_FIXTURE(FTestInstancedStructNetSerializerFixture, TestIsEqualQuantized)
+{
+	constexpr bool bUseQuantizedState = true;
+
+	// Default state compared to default state
+	InstancedStruct0.Reset();
+	Quantize();
+	Clone();
+	UE_NET_ASSERT_TRUE(IsEqual(bUseQuantizedState));
+
+	// Non-default state compared to default state
+	SetNonDefaultInstanceState();
+	Quantize();
+	UE_NET_ASSERT_FALSE(IsEqual(bUseQuantizedState));
+
+	// Non-default state compared to non-default state
+	Clone();
+	UE_NET_ASSERT_TRUE(IsEqual(bUseQuantizedState));
+}
+
+UE_NET_TEST_FIXTURE(FTestInstancedStructNetSerializerFixture, TestValidate)
+{
+	Validate();
+}
+
+// Array tests. There's no custom FInstancedStructArrayNetSerializer so we just add the one test until we require in-depth testing.
+UE_NET_TEST_FIXTURE(FTestInstancedStructNetSerializerFixture, TestDequantizeSerializedInitializedArrayState)
+{
+	bIsTestingArray = true;
+
+	SetNonDefaultArrayState();
+	Quantize();
+	Serialize();
+	FreeQuantizedState();
+	Deserialize();
+	Dequantize();
+	UE_NET_ASSERT_TRUE(InstancedStruct0 == InstancedStruct1);
+}
+
+// End-to-end tests
+UE_NET_TEST_FIXTURE(FTestInstancedStructNetSerializerFixture, ModifyInstance)
+{
+	// Add a client
+	FReplicationSystemTestClient* Client = CreateClient();
+
+	// Spawn object on server
+	UInstancedStructNetSerializerTestObject* ServerObject = Server->CreateObject<UInstancedStructNetSerializerTestObject>();
+	
+	ServerObject->InstancedStruct.InitializeAs<FStructForInstancedStructTestB>();
+	ServerObject->InstancedStruct.GetMutable<FStructForInstancedStructTestB>().SomeFloat = 12.0f;
+
+	// Replicate
+	Server->UpdateAndSend({ Client });
+
+	auto ClientObject = Client->GetObjectAs<UInstancedStructNetSerializerTestObject>(ServerObject->NetRefHandle);
+	UE_NET_ASSERT_NE(ClientObject, nullptr);
+	CA_ASSUME(ClientObject != nullptr);
+	UE_NET_ASSERT_EQ(ClientObject->InstancedStruct, ServerObject->InstancedStruct);
+	UE_NET_ASSERT_EQ(ClientObject->InstancedStruct.Get<FStructForInstancedStructTestB>().SomeFloat, ServerObject->InstancedStruct.Get<FStructForInstancedStructTestB>().SomeFloat);
+
+	// Modify
+	ServerObject->InstancedStruct.GetMutable<FStructForInstancedStructTestB>().SomeFloat += 1.0f;
+
+	// Replicate
+	Server->UpdateAndSend({ Client });
+
+	// Verify that we detected the modification
+	UE_NET_ASSERT_EQ(ClientObject->InstancedStruct.Get<FStructForInstancedStructTestB>().SomeFloat, ServerObject->InstancedStruct.Get<FStructForInstancedStructTestB>().SomeFloat);
+
+	// Switch type
+	ServerObject->InstancedStruct.InitializeAs<FStructForInstancedStructTestA>();
+	ServerObject->InstancedStruct.GetMutable<FStructForInstancedStructTestA>().SomeUint16 = 100;
+
+	// Replicate
+	Server->UpdateAndSend({ Client });
+
+	UE_NET_ASSERT_EQ(ClientObject->InstancedStruct, ServerObject->InstancedStruct);
+	UE_NET_ASSERT_EQ(ClientObject->InstancedStruct.Get<FStructForInstancedStructTestA>().SomeUint16, ServerObject->InstancedStruct.Get<FStructForInstancedStructTestA>().SomeUint16);
+}
+
+UE_NET_TEST_FIXTURE(FTestInstancedStructNetSerializerFixture, ModifyArray)
+{
+	// Add a client
+	FReplicationSystemTestClient* Client = CreateClient();
+
+	// Spawn objects on server
+	UInstancedStructNetSerializerTestObject* ServerObject = Server->CreateObject<UInstancedStructNetSerializerTestObject>();
+
+	// Add entries to the array
+	ServerObject->InstancedStructArray.Add(FInstancedStruct::Make<FStructForInstancedStructTestA>());
+	ServerObject->InstancedStructArray.Add(FInstancedStruct::Make<FStructForInstancedStructTestB>());
+	ServerObject->InstancedStructArray.Add(FInstancedStruct::Make<FStructForInstancedStructTestC>());
+	ServerObject->InstancedStructArray.Add(FInstancedStruct::Make<FStructForInstancedStructTestD>());
+	ServerObject->InstancedStructArray[1].GetMutable<FStructForInstancedStructTestB>().SomeFloat = 13.0f;
+	
+	// Replicate
+	Server->UpdateAndSend({ Client });
+	
+	auto ClientObject = Client->GetObjectAs<UInstancedStructNetSerializerTestObject>(ServerObject->NetRefHandle);
+	UE_NET_ASSERT_NE(ClientObject, nullptr);
+	CA_ASSUME(ClientObject != nullptr);
+	UE_NET_ASSERT_EQ(ClientObject->InstancedStructArray.Num(), ServerObject->InstancedStructArray.Num());
+	UE_NET_ASSERT_EQ(ClientObject->InstancedStructArray[1].Get<FStructForInstancedStructTestB>().SomeFloat, 13.0f);
+
+	// Modify value and see that it is replicated as expected
+	ServerObject->InstancedStructArray[1].GetMutable<FStructForInstancedStructTestB>().SomeFloat += 2.0f;
+
+	// Replicate
+	Server->UpdateAndSend({ Client });
+
+	// Verified that the client got the modified value
+	UE_NET_ASSERT_EQ(ClientObject->InstancedStructArray[1].Get<FStructForInstancedStructTestB>().SomeFloat, ServerObject->InstancedStructArray[1].Get<FStructForInstancedStructTestB>().SomeFloat);
+
+	// Switch type
+	ServerObject->InstancedStructArray[2].InitializeAs<FStructForInstancedStructTestA>();
+	ServerObject->InstancedStructArray[2].GetMutable<FStructForInstancedStructTestA>().SomeUint16 += 1;
+
+	// Replicate
+	Server->UpdateAndSend({ Client });
+
+	// Verified that the client got the modified value
+	UE_NET_ASSERT_EQ(ClientObject->InstancedStructArray[2], ClientObject->InstancedStructArray[2]);
+	UE_NET_ASSERT_EQ(ClientObject->InstancedStructArray[2].Get<FStructForInstancedStructTestA>().SomeUint16, ServerObject->InstancedStructArray[2].Get<FStructForInstancedStructTestA>().SomeUint16);
+}
+
+// Fixture implementation
+void FTestInstancedStructNetSerializerFixture::SetUp()
+{
+	FReplicationSystemServerClientTestFixture::SetUp();
+
+	// Init default serialization context
+	InternalNetSerializationContext.ReplicationSystem = Server->ReplicationSystem;
+
+	FReplicationSystemInternal* ReplicationSystemInternal = Server->GetReplicationSystem()->GetReplicationSystemInternal();
+
+	FInternalNetSerializationContext TempInternalNetSerializationContext;
+	FInternalNetSerializationContext::FInitParameters TempInternalNetSerializationContextInitParams;
+	TempInternalNetSerializationContextInitParams.ReplicationSystem = Server->ReplicationSystem;
+	TempInternalNetSerializationContextInitParams.ObjectResolveContext.RemoteNetTokenStoreState = ReplicationSystemInternal->GetNetTokenStore().GetLocalNetTokenStoreState();
+	TempInternalNetSerializationContext.Init(TempInternalNetSerializationContextInitParams);
+
+	InternalNetSerializationContext = MoveTemp(TempInternalNetSerializationContext);
+	NetSerializationContext.SetInternalContext(&InternalNetSerializationContext);
+
+	FMemory::Memzero(QuantizedBuffer, sizeof(QuantizedBuffer));
+
+	bHasQuantizedState = false;
+	bHasClonedQuantizedState = false;
+
+	if (!InstancedStructNetSerializerConfig.StateDescriptor.IsValid())
+	{
+		FReplicationStateDescriptorBuilder::FParameters Params;
+		InstancedStructNetSerializerConfig.StateDescriptor = FReplicationStateDescriptorBuilder::CreateDescriptorForStruct(FTestInstancedStruct::StaticStruct(), Params);
+	}
+
+	if (!InstancedStructArrayNetSerializerConfig.StateDescriptor.IsValid())
+	{
+		FReplicationStateDescriptorBuilder::FParameters Params;
+		InstancedStructArrayNetSerializerConfig.StateDescriptor = FReplicationStateDescriptorBuilder::CreateDescriptorForStruct(FTestInstancedStructArray::StaticStruct(), Params);
+	}
+}
+
+void FTestInstancedStructNetSerializerFixture::TearDown()
+{
+	InstancedStruct0.Reset();
+	InstancedStruct1.Reset();
+	InstancedStructArray0.Reset();
+	InstancedStructArray1.Reset();
+
+	FreeQuantizedState();
+
+	FReplicationSystemServerClientTestFixture::TearDown();
+}
+
+void FTestInstancedStructNetSerializerFixture::Serialize()
+{
+	// Must have run quantize before this
+	UE_NET_ASSERT_TRUE(bHasQuantizedState);
+
+	// Serialize data
+	Writer.InitBytes(BitStreamBuffer, sizeof(BitStreamBuffer));
+	FNetSerializationContext Context(&Writer);
+	Context.SetInternalContext(NetSerializationContext.GetInternalContext());
+
+	FNetSerializeArgs Args = {};
+	Args.NetSerializerConfig = NetSerializerConfigParam(GetInstancedStructNetSerializerConfig(bIsTestingArray));
+	Args.Source = NetSerializerValuePointer(&QuantizedBuffer[0]);
+	GetInstancedStructNetSerializer(bIsTestingArray)->Serialize(Context, Args);
+
+	Writer.CommitWrites();
+
+	UE_NET_ASSERT_FALSE(Context.HasError());
+	UE_NET_ASSERT_GT(Writer.GetPosBits(), 0U);
+}
+
+void FTestInstancedStructNetSerializerFixture::Deserialize()
+{
+	// Check pre-conditions
+	UE_NET_ASSERT_FALSE(bHasQuantizedState);
+	UE_NET_ASSERT_GT(Writer.GetPosBytes(), 0U);
+
+	Reader.InitBits(BitStreamBuffer, Writer.GetPosBits());
+
+	FNetSerializationContext Context(&Reader);
+	Context.SetInternalContext(NetSerializationContext.GetInternalContext());
+
+	FNetDeserializeArgs Args = {};
+	Args.NetSerializerConfig = NetSerializerConfigParam(GetInstancedStructNetSerializerConfig(bIsTestingArray));
+	Args.Target = NetSerializerValuePointer(&QuantizedBuffer[0]);
+	GetInstancedStructNetSerializer(bIsTestingArray)->Deserialize(Context, Args);
+
+	bHasQuantizedState = true;
+
+	UE_NET_ASSERT_FALSE(Context.HasErrorOrOverflow());
+	UE_NET_ASSERT_GT(Reader.GetPosBits(), 0U);
+}
+
+void FTestInstancedStructNetSerializerFixture::SerializeDelta()
+{
+	// Check pre-conditions
+	UE_NET_ASSERT_TRUE(bHasQuantizedState);
+	UE_NET_ASSERT_EQ(QuantizedStateCount, 2U);
+
+	// Serialize data
+	Writer.InitBytes(BitStreamBuffer, sizeof(BitStreamBuffer));
+	FNetSerializationContext Context(&Writer);
+	Context.SetInternalContext(NetSerializationContext.GetInternalContext());
+
+	FNetSerializeDeltaArgs Args = {};
+	Args.NetSerializerConfig = NetSerializerConfigParam(GetInstancedStructNetSerializerConfig(bIsTestingArray));
+	Args.Source = NetSerializerValuePointer(&QuantizedBuffer[0]);
+	Args.Prev = NetSerializerValuePointer(&QuantizedBuffer[1]);
+	GetInstancedStructNetSerializer(bIsTestingArray)->SerializeDelta(Context, Args);
+
+	Writer.CommitWrites();
+
+	UE_NET_ASSERT_FALSE(Context.HasErrorOrOverflow());
+	UE_NET_ASSERT_GT(Writer.GetPosBits(), 0U);
+}
+
+void FTestInstancedStructNetSerializerFixture::DeserializeDelta()
+{
+	// Check pre-conditions
+	UE_NET_ASSERT_GT(Writer.GetPosBytes(), 0U);
+
+	Reader.InitBits(BitStreamBuffer, Writer.GetPosBits());
+
+	FNetSerializationContext Context(&Reader);
+	Context.SetInternalContext(NetSerializationContext.GetInternalContext());
+
+	FNetDeserializeDeltaArgs Args = {};
+	Args.NetSerializerConfig = NetSerializerConfigParam(GetInstancedStructNetSerializerConfig(bIsTestingArray));
+	Args.Target = NetSerializerValuePointer(&QuantizedBuffer[0]);
+	Args.Prev = NetSerializerValuePointer(&QuantizedBuffer[1]);
+	GetInstancedStructNetSerializer(bIsTestingArray)->DeserializeDelta(Context, Args);
+
+	bHasQuantizedState = true;
+	QuantizedStateCount = 1;
+
+	UE_NET_ASSERT_FALSE(Context.HasErrorOrOverflow());
+	UE_NET_ASSERT_GT(Reader.GetPosBits(), 0U);
+}
+
+void FTestInstancedStructNetSerializerFixture::Quantize()
+{
+	FNetQuantizeArgs Args = {};
+	Args.NetSerializerConfig = NetSerializerConfigParam(GetInstancedStructNetSerializerConfig(bIsTestingArray));
+	Args.Target = NetSerializerValuePointer(&QuantizedBuffer[0]);
+	Args.Source = bIsTestingArray ? NetSerializerValuePointer(&InstancedStructArray0) : NetSerializerValuePointer(&InstancedStruct0);
+	GetInstancedStructNetSerializer(bIsTestingArray)->Quantize(NetSerializationContext, Args);
+
+	bHasQuantizedState = true;
+	QuantizedStateCount = 1;
+
+	UE_NET_ASSERT_FALSE(NetSerializationContext.HasError());
+}
+
+void FTestInstancedStructNetSerializerFixture::QuantizeTwoStates()
+{
+	Quantize();
+
+	FNetQuantizeArgs Args = {};
+	Args.NetSerializerConfig = NetSerializerConfigParam(GetInstancedStructNetSerializerConfig(bIsTestingArray));
+	Args.Target = NetSerializerValuePointer(&QuantizedBuffer[1]);
+	Args.Source = bIsTestingArray ? NetSerializerValuePointer(&InstancedStructArray1) : NetSerializerValuePointer(&InstancedStruct1);
+	GetInstancedStructNetSerializer(bIsTestingArray)->Quantize(NetSerializationContext, Args);
+
+	bHasQuantizedState = true;
+	QuantizedStateCount = 2;
+
+	UE_NET_ASSERT_FALSE(NetSerializationContext.HasError());
+}
+
+void FTestInstancedStructNetSerializerFixture::Clone()
+{
+	// Check pre-conditions
+	UE_NET_ASSERT_TRUE(bHasQuantizedState);
+
+	FMemory::Memcpy(ClonedQuantizedBuffer[0], QuantizedBuffer[0], sizeof(QuantizedBuffer[0]));
+
+	FNetCloneDynamicStateArgs Args = {};
+	Args.Source = NetSerializerValuePointer(&QuantizedBuffer[0]);
+	Args.Target = NetSerializerValuePointer(&ClonedQuantizedBuffer[0]);
+	Args.NetSerializerConfig = NetSerializerConfigParam(GetInstancedStructNetSerializerConfig(bIsTestingArray));
+	GetInstancedStructNetSerializer(bIsTestingArray)->CloneDynamicState(NetSerializationContext, Args);
+
+	bHasClonedQuantizedState = true;
+	ClonedQuantizedStateCount = 1;
+}
+
+void FTestInstancedStructNetSerializerFixture::FreeQuantizedState()
+{
+	FNetFreeDynamicStateArgs Args = {};
+	Args.NetSerializerConfig = NetSerializerConfigParam(GetInstancedStructNetSerializerConfig(bIsTestingArray));
+
+	if (bHasQuantizedState)
+	{
+		for (uint32 StateIt = 0, StateEndIt = QuantizedStateCount; StateIt != StateEndIt; ++StateIt)
+		{
+			Args.Source = NetSerializerValuePointer(&QuantizedBuffer[StateIt]);
+			GetInstancedStructNetSerializer(bIsTestingArray)->FreeDynamicState(NetSerializationContext, Args);
+
+			FMemory::Memzero(&QuantizedBuffer[StateIt], sizeof(QuantizedBuffer[StateIt]));
+		}
+		bHasQuantizedState = false;
+	}
+
+	if (bHasClonedQuantizedState)
+	{
+		for (uint32 StateIt = 0, StateEndIt = ClonedQuantizedStateCount; StateIt != StateEndIt; ++StateIt)
+		{
+			Args.Source = NetSerializerValuePointer(&ClonedQuantizedBuffer[StateIt]);
+			GetInstancedStructNetSerializer(bIsTestingArray)->FreeDynamicState(NetSerializationContext, Args);
+
+			FMemory::Memzero(&ClonedQuantizedBuffer[StateIt], sizeof(ClonedQuantizedBuffer[StateIt]));
+		}
+		bHasClonedQuantizedState = false;
+	}
+}
+
+void FTestInstancedStructNetSerializerFixture::Dequantize()
+{
+	UE_NET_ASSERT_TRUE(bHasQuantizedState);
+
+	FNetDequantizeArgs Args = {};
+	Args.NetSerializerConfig = NetSerializerConfigParam(GetInstancedStructNetSerializerConfig(bIsTestingArray));
+	Args.Source = NetSerializerValuePointer(&QuantizedBuffer[0]);
+	Args.Target = bIsTestingArray ? NetSerializerValuePointer(&InstancedStructArray1) : NetSerializerValuePointer(&InstancedStruct1);
+	GetInstancedStructNetSerializer(bIsTestingArray)->Dequantize(NetSerializationContext, Args);
+}
+
+bool FTestInstancedStructNetSerializerFixture::IsEqual(bool bQuantized)
+{
+	if (bQuantized)
+	{
+		UE_NET_EXPECT_TRUE(bHasQuantizedState);
+		if (!bHasQuantizedState)
+		{
+			return false;
+		}
+
+		UE_NET_EXPECT_TRUE(bHasClonedQuantizedState);
+		if (!bHasClonedQuantizedState)
+		{
+			return false;
+		}
+	}
+
+	FNetIsEqualArgs Args = {};
+	if (bQuantized)
+	{
+		Args.Source0 = NetSerializerValuePointer(&QuantizedBuffer[0]);
+		Args.Source1 = NetSerializerValuePointer(&ClonedQuantizedBuffer[0]);
+	}
+	else
+	{
+		Args.Source0 = bIsTestingArray ? NetSerializerValuePointer(&InstancedStructArray0) : NetSerializerValuePointer(&InstancedStruct0);
+		Args.Source1 = bIsTestingArray ? NetSerializerValuePointer(&InstancedStructArray1) : NetSerializerValuePointer(&InstancedStruct1);
+	}
+	Args.bStateIsQuantized = bQuantized;
+	Args.NetSerializerConfig = NetSerializerConfigParam(GetInstancedStructNetSerializerConfig(bIsTestingArray));
+	return GetInstancedStructNetSerializer(bIsTestingArray)->IsEqual(NetSerializationContext, Args);
+}
+
+void FTestInstancedStructNetSerializerFixture::Validate()
+{
+	FNetValidateArgs Args = {};
+	Args.NetSerializerConfig = NetSerializerConfigParam(GetInstancedStructNetSerializerConfig(bIsTestingArray));
+	Args.Source = bIsTestingArray ? NetSerializerValuePointer(&InstancedStructArray0) : NetSerializerValuePointer(&InstancedStruct0);
+
+	GetInstancedStructNetSerializer(bIsTestingArray)->Validate(NetSerializationContext, Args);
+}
+
+void FTestInstancedStructNetSerializerFixture::SetNonDefaultInstanceState()
+{
+	InstancedStruct0.InitializeAs<FStructForInstancedStructTestA>();
+	InstancedStruct0.GetMutable<FStructForInstancedStructTestA>().SomeUint16 += 4711U;
+}
+
+void FTestInstancedStructNetSerializerFixture::SetNonDefaultArrayState()
+{
+	InstancedStructArray0.Empty(7);
+	InstancedStructArray0.SetNum(7);
+
+	InstancedStructArray0[0].InitializeAs<FStructForInstancedStructTestA>();
+
+	InstancedStructArray0[1].InitializeAs<FStructForInstancedStructTestB>();
+	InstancedStructArray0[1].GetMutable<FStructForInstancedStructTestB>().SomeFloat = 1234.5f;
+
+	InstancedStructArray0[2].Reset();
+
+	InstancedStructArray0[3].InitializeAs<FStructForInstancedStructTestC>();
+	InstancedStructArray0[3].GetMutable<FStructForInstancedStructTestC>().SomeBool ^= true;
+
+	InstancedStructArray0[4].Reset();
+
+	InstancedStructArray0[5].InitializeAs<FStructForInstancedStructTestWithArray>();
+
+	InstancedStructArray0[6].InitializeAs<FStructForInstancedStructTestWithObjectReference>();
+}
+
+const FNetSerializer* FTestInstancedStructNetSerializerFixture::GetInstancedStructNetSerializer(bool bIsArray) const
+{
+	if (bIsArray)
+	{
+		if (const FReplicationStateDescriptor* Desc = InstancedStructArrayNetSerializerConfig.StateDescriptor.GetReference())
+		{
+			return Desc->MemberSerializerDescriptors[0].Serializer;
+		}
+	}
+	else
+	{
+		if (const FReplicationStateDescriptor* Desc = InstancedStructNetSerializerConfig.StateDescriptor.GetReference())
+		{
+			return Desc->MemberSerializerDescriptors[0].Serializer;
+		}
+	}
+
+	return nullptr;
+}
+
+const FNetSerializerConfig* FTestInstancedStructNetSerializerFixture::GetInstancedStructNetSerializerConfig(bool bIsArray) const
+{
+	if (bIsArray)
+	{
+		if (const FReplicationStateDescriptor* Desc = InstancedStructArrayNetSerializerConfig.StateDescriptor.GetReference())
+		{
+			return Desc->MemberSerializerDescriptors[0].SerializerConfig;
+		}
+	}
+	else
+	{
+		if (const FReplicationStateDescriptor* Desc = InstancedStructNetSerializerConfig.StateDescriptor.GetReference())
+		{
+			return Desc->MemberSerializerDescriptors[0].SerializerConfig;
+		}
+	}
+
+	return nullptr;
+}
+
+}
+
+//
+void UInstancedStructNetSerializerTestObject::GetLifetimeReplicatedProps(TArray<class FLifetimeProperty>& OutLifetimeProps) const
+{
+	FDoRepLifetimeParams Params;
+	Params.bIsPushBased = false;
+
+	DOREPLIFETIME_WITH_PARAMS_FAST(ThisClass, InstancedStruct, Params);
+	DOREPLIFETIME_WITH_PARAMS_FAST(ThisClass, InstancedStructArray, Params);
+}
+
+void UInstancedStructNetSerializerTestObject::RegisterReplicationFragments(UE::Net::FFragmentRegistrationContext& Context, UE::Net::EFragmentRegistrationFlags RegistrationFlags)
+{
+	UE::Net::FReplicationFragmentUtil::CreateAndRegisterFragmentsForObject(this, Context, RegistrationFlags);
+}
diff --git a/Engine/Plugins/Runtime/ReplicationSystemTestPlugin/Source/Private/Tests/Serialization/TestInstancedStructNetSerializer.h b/Engine/Plugins/Runtime/ReplicationSystemTestPlugin/Source/Private/Tests/Serialization/TestInstancedStructNetSerializer.h
new file mode 100644
--- /dev/null
+++ b/Engine/Plugins/Runtime/ReplicationSystemTestPlugin/Source/Private/Tests/Serialization/TestInstancedStructNetSerializer.h
@@ -0,0 +1,96 @@
+// Copyright Epic Games, Inc. All Rights Reserved.
+
+#pragma once
+
+#include "UObject/ObjectMacros.h"
+#include "InstancedStruct.h"
+#include "Tests/ReplicationSystem/ReplicatedTestObject.h"
+#include "TestInstancedStructNetSerializer.generated.h"
+
+USTRUCT()
+struct FTestInstancedStruct
+{
+	GENERATED_BODY()
+
+	UPROPERTY()
+	FInstancedStruct InstancedStruct;
+};
+
+USTRUCT()
+struct FTestInstancedStructArray
+{
+	GENERATED_BODY()
+
+	UPROPERTY()
+	TArray<FInstancedStruct> InstancedStructArray;
+};
+
+USTRUCT()
+struct FStructForInstancedStructTestA
+{
+	GENERATED_BODY()
+
+	UPROPERTY()
+	uint16 SomeUint16 = 0;
+};
+
+USTRUCT()
+struct FStructForInstancedStructTestB
+{
+	GENERATED_BODY()
+
+	UPROPERTY()
+	float SomeFloat = 0.0f;
+};
+
+USTRUCT()
+struct FStructForInstancedStructTestC
+{
+	GENERATED_BODY()
+
+	UPROPERTY()
+	bool SomeBool = false;
+};
+
+USTRUCT()
+struct FStructForInstancedStructTestD
+{
+	GENERATED_BODY()
+
+	// Intentionally has no properties
+};
+
+USTRUCT()
+struct FStructForInstancedStructTestWithArray
+{
+	GENERATED_BODY()
+
+	UPROPERTY()
+	TArray<FStructForInstancedStructTestB> ArrayOfTestB;
+};
+
+USTRUCT()
+struct FStructForInstancedStructTestWithObjectReference
+{
+	GENERATED_BODY()
+
+	UPROPERTY();
+	TObjectPtr<UObject> SomeObject = nullptr;
+};
+
+UCLASS()
+class UInstancedStructNetSerializerTestObject : public UReplicatedTestObject
+{
+	GENERATED_BODY()
+
+protected:
+	virtual void RegisterReplicationFragments(UE::Net::FFragmentRegistrationContext& Fragments, UE::Net::EFragmentRegistrationFlags RegistrationFlags) override;
+
+public:
+	UPROPERTY(Replicated, Transient)
+	FInstancedStruct InstancedStruct;
+
+	UPROPERTY(Replicated, Transient)
+	TArray<FInstancedStruct> InstancedStructArray;
+};
+
diff --git a/Engine/Plugins/Runtime/ReplicationSystemTestPlugin/Source/ReplicationSystemTestPlugin.Build.cs b/Engine/Plugins/Runtime/ReplicationSystemTestPlugin/Source/ReplicationSystemTestPlugin.Build.cs
--- a/Engine/Plugins/Runtime/ReplicationSystemTestPlugin/Source/ReplicationSystemTestPlugin.Build.cs
+++ b/Engine/Plugins/Runtime/ReplicationSystemTestPlugin/Source/ReplicationSystemTestPlugin.Build.cs
@@ -33,6 +33,7 @@ namespace UnrealBuildTool.Rules
 					"CoreUObject",
 					"Engine",
 					"NetCore",
+					"StructUtils"
 				}
 				);
 
diff --git a/Engine/Source/Runtime/Experimental/Iris/Core/Private/Iris/Serialization/NetSerializerArrayStorage.cpp b/Engine/Source/Runtime/Experimental/Iris/Core/Private/Iris/Serialization/NetSerializerArrayStorage.cpp
--- a/Engine/Source/Runtime/Experimental/Iris/Core/Private/Iris/Serialization/NetSerializerArrayStorage.cpp
+++ b/Engine/Source/Runtime/Experimental/Iris/Core/Private/Iris/Serialization/NetSerializerArrayStorage.cpp
@@ -19,3 +19,77 @@ void FElementAllocationPolicy::Free(FNetSerializationContext& Context, void* Ptr
 
 }
 
+
+namespace UE::Net
+{
+
+void FNetSerializerAlignedStorage::AdjustSize(FNetSerializationContext& Context, SizeType InNum, SizeType InAlignment)
+{
+	if (!InNum)
+	{
+		Free(Context);
+		return;
+	}
+
+	// If the allocation isn't properly aligned or if the allocation is too small we make a new allocation.
+	if ((InNum > StorageMaxCapacity) || !IsAligned(Data, InAlignment))
+	{
+		// Can't use realloc as the alignment of the original allocation isn't tracked.
+		void* NewData = Context.GetInternalContext()->Alloc(InNum, InAlignment);
+		FMemory::Memzero(NewData, InNum);
+		// Copy old data
+		FMemory::Memcpy(NewData, Data, StorageNum);
+		Context.GetInternalContext()->Free(Data);
+
+		Data = static_cast<uint8*>(NewData);
+		StorageNum = InNum;
+		StorageMaxCapacity = InNum;
+		StorageAlignment = InAlignment;
+	}
+	// Requested data size fits the current allocation
+	else
+	{
+		// Clear capacity we're not using anymore. If we're growing we don't need to clear as it has already been cleared.
+		if (InNum < StorageNum)
+		{
+			FMemory::Memzero(static_cast<void*>(Data + InNum), StorageNum - InNum);
+		}
+		StorageNum = InNum;
+		StorageNum = InAlignment;
+	}
+}
+
+void FNetSerializerAlignedStorage::Free(FNetSerializationContext& Context)
+{
+	if (Data != nullptr)
+	{
+		Context.GetInternalContext()->Free(Data);
+	}
+
+	Data = nullptr;
+	StorageNum = 0;
+	StorageMaxCapacity = 0;
+	StorageAlignment = 0;
+}
+
+void FNetSerializerAlignedStorage::Clone(FNetSerializationContext& Context, const FNetSerializerAlignedStorage& Source)
+{
+	// Only allocate and copy the exact amount of memory needed.
+	if (Source.StorageNum > 0)
+	{
+		Data = static_cast<uint8*>(Context.GetInternalContext()->Alloc(Source.StorageNum, Source.StorageAlignment));
+		StorageNum = Source.StorageNum;
+		StorageMaxCapacity = Source.StorageNum;
+		StorageAlignment = Source.StorageAlignment;
+		FMemory::Memcpy(Data, Source.Data, Source.StorageNum);
+	}
+	else
+	{
+		Data = nullptr;
+		StorageNum = 0;
+		StorageMaxCapacity = 0;
+		StorageAlignment = 0;
+	}
+}
+
+}
diff --git a/Engine/Source/Runtime/Experimental/Iris/Core/Public/Iris/Serialization/NetSerializerArrayStorage.h b/Engine/Source/Runtime/Experimental/Iris/Core/Public/Iris/Serialization/NetSerializerArrayStorage.h
--- a/Engine/Source/Runtime/Experimental/Iris/Core/Public/Iris/Serialization/NetSerializerArrayStorage.h
+++ b/Engine/Source/Runtime/Experimental/Iris/Core/Public/Iris/Serialization/NetSerializerArrayStorage.h
@@ -2,10 +2,11 @@
 
 #pragma once
 
-#include "CoreTypes.h"
+#include "HAL/Platform.h"
 #include "Iris/IrisConfig.h"
+#include "Iris/Serialization/NetSerializationContext.h"
+#include "Templates/AlignmentTemplates.h"
 #include "Templates/IsPODType.h"
-#include "NetSerializationContext.h"
 
 namespace UE::Net
 {
@@ -14,7 +15,7 @@ namespace AllocationPolicies
 {
 
 /**
- * Default AllicationPolicy for NetSerializationStorage helper class
+ * Default AllocationPolicy for FNetSerializerArrayStorage helper class
  */
 class FElementAllocationPolicy
 {
@@ -166,9 +167,55 @@ private:
 	SizeType ArrayMaxCapacity;
 };
 
+/**
+ * Helper class to manage storage of untyped dynamic storage with arbitrary alignment in quantized data.
+ * If storage is used to store elements that themselves uses dynamic allocations the user must take care to properly free and clone such elements properly prior to and/or after changing sizes.
+ * FNetSerializerAlignedStorage is only intended to be used from within Quantized replication states that requires dynamic storage as it has very specific expectations and limitations.
+ * NOTE: A NetSerializer using FNetSerializerAlignedStorage MUST specify the trait bHasDynamicState and forward appropriate calls to Clone and Free to this class.
+ * A zero constructed state is considered valid.
+ */
+class FNetSerializerAlignedStorage
+{
+public:
+	// Use whatever SizeType that is typically used by FNetSerializerArrayStorage
+	typedef typename AllocationPolicies::FElementAllocationPolicy::SizeType SizeType;
 
-// FNetSerializerArrayStorage Implementation
+public:
+	/**
+	 * AdjustSize - Adjust the size of FNetSerializerArrayStorage as needed
+	 * The state storage in which the storage is located is expected to be in a valid state which can either be zero initialized or a previous valid state
+	 */
+	IRISCORE_API void AdjustSize(FNetSerializationContext& Context, SizeType InNum, SizeType InAlignment);
+
+	/**
+	 * Free - Free allocated memory and reset state. it is valid to pass in a zero-initialized state
+	 * The state storage in which the storage is located is expected to be in a valid state which can either be zero initialized or a previous valid state.
+	 * Should typically only be called from within a the implemntation of NetSerializer::FreeDynamicState()
+	 */
+	IRISCORE_API void Free(FNetSerializationContext& Context);
+
+	/**
+	 * Clone - Clone dynamic storage from Source
+	 * No assumptions of the validity of the target is made, as this is typically called AFTER a memcopy is made to the target state
+	 * which invalidates all dynamic data
+	 * Should typically only be called from within a the implemntation of NetSerializer::CloneDynamicState()
+	*/
+	IRISCORE_API void Clone(FNetSerializationContext& Context, const FNetSerializerAlignedStorage& Source);
+
+	inline const uint8* GetData() const;
+	inline uint8* GetData();
+	inline SizeType Num() const;
+	inline SizeType GetAlignment() const;
+
+private:
+	uint8* Data = nullptr;
+	SizeType StorageNum = 0;
+	SizeType StorageMaxCapacity = 0;
+	// An alignment of 0 isn't valid but the value is only used if there's actual data involved.
+	SizeType StorageAlignment = 0;
+};
 
+// FNetSerializerArrayStorage Implementation
 template <typename QuantizedElementType, typename AllocationPolicy>
 FNetSerializerArrayStorage<QuantizedElementType, AllocationPolicy>::FNetSerializerArrayStorage()
 : ArrayNum(0)
@@ -214,6 +261,27 @@ void FNetSerializerArrayStorage<QuantizedElementType, AllocationPolicy>::Clone(F
 	ArrayNum = SourceNum;
 }
 
+// FNetSerializerAlignedStorage implementation
+
+inline const uint8* FNetSerializerAlignedStorage::GetData() const
+{
+	return Data;
+}
+
+inline uint8* FNetSerializerAlignedStorage::GetData()
+{
+	return Data;
+}
+
+inline FNetSerializerAlignedStorage::SizeType FNetSerializerAlignedStorage::Num() const
+{
+	return StorageNum;
+}
+
+inline FNetSerializerAlignedStorage::SizeType FNetSerializerAlignedStorage::GetAlignment() const
+{
+	return StorageAlignment;
+}
 
 namespace AllocationPolicies
 {
