Iris - Perform all delta compression baseline invalidation based on conditional changes in FReplicationConditionals. Invalidate baselines for the entire object hierarchy if a condition change are affecting properties in subobjects too.

diff --git a/Engine/Plugins/Runtime/ReplicationSystemTestPlugin/Source/Private/Tests/ReplicationSystem/ReplicatedTestObject.cpp b/Engine/Plugins/Runtime/ReplicationSystemTestPlugin/Source/Private/Tests/ReplicationSystem/ReplicatedTestObject.cpp
--- a/Engine/Plugins/Runtime/ReplicationSystemTestPlugin/Source/Private/Tests/ReplicationSystem/ReplicatedTestObject.cpp
+++ b/Engine/Plugins/Runtime/ReplicationSystemTestPlugin/Source/Private/Tests/ReplicationSystem/ReplicatedTestObject.cpp
@@ -364,6 +364,11 @@ void UTestReplicatedIrisLifetimeConditionalsPropertyState::GetLifetimeReplicated
 	DOREPLIFETIME_CONDITION(ThisClass, NeverInt, COND_Never);
 	DOREPLIFETIME_CONDITION(ThisClass, SkipReplayInt, COND_SkipReplay);
 	DOREPLIFETIME_CONDITION(ThisClass, ReplayOnlyInt, COND_ReplayOnly);
+
+	DOREPLIFETIME_CONDITION(ThisClass, SimulatedOnlyIntArray, COND_SimulatedOnly);
+	DOREPLIFETIME_CONDITION(ThisClass, AutonomousOnlyIntArray, COND_AutonomousOnly);
+	DOREPLIFETIME_CONDITION(ThisClass, SimulatedOrPhysicsIntArray, COND_SimulatedOrPhysics);
+	DOREPLIFETIME_CONDITION(ThisClass, OwnerOnlyIntArray, COND_OwnerOnly);
 }
 
 //////////////////////////////////////////////////////////////////////////
diff --git a/Engine/Plugins/Runtime/ReplicationSystemTestPlugin/Source/Private/Tests/ReplicationSystem/ReplicatedTestObject.h b/Engine/Plugins/Runtime/ReplicationSystemTestPlugin/Source/Private/Tests/ReplicationSystem/ReplicatedTestObject.h
--- a/Engine/Plugins/Runtime/ReplicationSystemTestPlugin/Source/Private/Tests/ReplicationSystem/ReplicatedTestObject.h
+++ b/Engine/Plugins/Runtime/ReplicationSystemTestPlugin/Source/Private/Tests/ReplicationSystem/ReplicatedTestObject.h
@@ -258,6 +258,16 @@ public:
 	UPROPERTY(Replicated, Transient)
 	int32 ReplayOnlyInt = 0;
 
+	// Arrays
+	UPROPERTY(Replicated, Transient)
+	TArray<int32> SimulatedOnlyIntArray;
+	UPROPERTY(Replicated, Transient)
+	TArray<int32> AutonomousOnlyIntArray;
+	UPROPERTY(Replicated, Transient)
+	TArray<int32> SimulatedOrPhysicsIntArray;
+	UPROPERTY(Replicated, Transient)
+	TArray<int32> OwnerOnlyIntArray;
+ 
 	// Network data only for test
 	TArray<UE::Net::FReplicationFragment*> ReplicationFragments;
 };
diff --git a/Engine/Plugins/Runtime/ReplicationSystemTestPlugin/Source/Private/Tests/ReplicationSystem/ReplicationSystemServerClientTestFixture.h b/Engine/Plugins/Runtime/ReplicationSystemTestPlugin/Source/Private/Tests/ReplicationSystem/ReplicationSystemServerClientTestFixture.h
--- a/Engine/Plugins/Runtime/ReplicationSystemTestPlugin/Source/Private/Tests/ReplicationSystem/ReplicationSystemServerClientTestFixture.h
+++ b/Engine/Plugins/Runtime/ReplicationSystemTestPlugin/Source/Private/Tests/ReplicationSystem/ReplicationSystemServerClientTestFixture.h
@@ -131,7 +131,7 @@ public:
 	UTestReplicatedIrisObject* CreateObject(const UObjectReplicationBridge::FCreateNetRefHandleParams& Params, UTestReplicatedIrisObject::FComponents* ComponentsToCreate=nullptr);
 	UTestReplicatedIrisObject* CreateObject(uint32 NumComponents, uint32 NumIrisComponents);
 	UTestReplicatedIrisObject* CreateSubObject(FNetRefHandle Owner, uint32 NumComponents, uint32 NumIrisComponents);
-	UTestReplicatedIrisObject* CreateObject(const UTestReplicatedIrisObject::FComponents& Components);
+	UTestReplicatedIrisObject* CreateObject(const UTestReplicatedIrisObject::FComponents& Components = UTestReplicatedIrisObject::FComponents());
 	UTestReplicatedIrisObject* CreateSubObject(FNetRefHandle Owner, const UTestReplicatedIrisObject::FComponents& Components);
 	UTestReplicatedIrisObject* CreateObjectWithDynamicState(uint32 NumComponents, uint32 NumIrisComponents, uint32 NumDynamicStateComponents);
 
diff --git a/Engine/Plugins/Runtime/ReplicationSystemTestPlugin/Source/Private/Tests/ReplicationSystem/TestObjectDeltaSerialization.cpp b/Engine/Plugins/Runtime/ReplicationSystemTestPlugin/Source/Private/Tests/ReplicationSystem/TestObjectDeltaSerialization.cpp
--- a/Engine/Plugins/Runtime/ReplicationSystemTestPlugin/Source/Private/Tests/ReplicationSystem/TestObjectDeltaSerialization.cpp
+++ b/Engine/Plugins/Runtime/ReplicationSystemTestPlugin/Source/Private/Tests/ReplicationSystem/TestObjectDeltaSerialization.cpp
@@ -2,6 +2,7 @@
 
 #include "NetworkAutomationTest.h"
 #include "NetworkAutomationTestMacros.h"
+#include "UObject/CoreNetTypes.h"
 #include "Iris/ReplicationSystem/ReplicationSystem.h"
 #include "Iris/ReplicationSystem/ReplicationSystemTypes.h"
 #include "Iris/ReplicationSystem/Conditionals/ReplicationCondition.h"
@@ -305,6 +306,171 @@ UE_NET_TEST_FIXTURE(FTestObjectDeltaSerialization, SimulatedIsProperlyReplicated
 	UE_NET_ASSERT_EQ(ClientObject->ConnectionFilteredComponents[0]->SimulatedOrPhysicsNoReplayInt, ServerObject->ConnectionFilteredComponents[0]->SimulatedOrPhysicsNoReplayInt);
 }
 
+UE_NET_TEST_FIXTURE(FTestObjectDeltaSerialization, SubObjectSimulatedIsProperlyReplicatedAfterBeingAutonomous)
+{
+	FReplicationSystemTestClient* Client = CreateClient();
+
+	// Spawn object with subobject on server
+	UTestReplicatedIrisObject* ServerObject = Server->CreateObject();
+	UTestReplicatedIrisObject* ServerSubObject = Server->CreateSubObject(ServerObject->NetRefHandle, {.ConnectionFilteredComponentCount = 1});
+	Server->ReplicationSystem->SetDeltaCompressionStatus(ServerSubObject->NetRefHandle, ENetObjectDeltaCompressionStatus::Allow);
+
+	// Set the root object to be autonomous.
+	Server->ReplicationSystem->SetReplicationConditionConnectionFilter(ServerObject->NetRefHandle, EReplicationCondition::RoleAutonomous, Client->ConnectionIdOnServer, true);
+
+	// Send and deliver packet
+	Server->UpdateAndSend({Client});
+
+	// Add some values to the simulated only array
+	for (int32 Value : {4711, 1337})
+	{
+		ServerSubObject->ConnectionFilteredComponents[0]->SimulatedOnlyIntArray.Add(Value);
+	}
+
+	// In order to actually send the object we need to modify a non-SimulatedOnly property too
+	ServerSubObject->ConnectionFilteredComponents[0]->NoneInt += 1;
+
+	// We want a new baseline to be created
+	AllowNewBaselineCreation();
+
+	// Send and deliver packet
+	Server->UpdateAndSend({Client});
+
+	UTestReplicatedIrisObject* ClientSubObject = Cast<UTestReplicatedIrisObject>(Client->GetReplicationBridge()->GetReplicatedObject(ServerSubObject->NetRefHandle));
+
+	// Verify the simulated only array is unaffected by the server changes
+	UE_NET_ASSERT_TRUE(ClientSubObject->ConnectionFilteredComponents[0]->SimulatedOnlyIntArray.IsEmpty());
+
+	// Set the client to no longer be autonomous, meaning it should be "simulated".
+	Server->ReplicationSystem->SetReplicationConditionConnectionFilter(ServerObject->NetRefHandle, EReplicationCondition::RoleAutonomous, Client->ConnectionIdOnServer, false);
+
+	// Add some values to the simulated only array
+	for (int32 Value : {1, -1})
+	{
+		ServerSubObject->ConnectionFilteredComponents[0]->SimulatedOnlyIntArray.Add(Value);
+	}
+
+	// Send and deliver packet
+	Server->UpdateAndSend({Client});
+
+	// Check that the array is properly replicated in its entirety
+	UE_NET_ASSERT_EQ(ClientSubObject->ConnectionFilteredComponents[0]->SimulatedOnlyIntArray.Num(), ServerSubObject->ConnectionFilteredComponents[0]->SimulatedOnlyIntArray.Num());
+	for (int Index = 0; int32 Value : ServerSubObject->ConnectionFilteredComponents[0]->SimulatedOnlyIntArray)
+	{
+		UE_NET_ASSERT_EQ(ClientSubObject->ConnectionFilteredComponents[0]->SimulatedOnlyIntArray[Index], Value);
+		++Index;
+	}
+}
+
+UE_NET_TEST_FIXTURE(FTestObjectDeltaSerialization, SubObjectAutonomousIsProperlyReplicatedAfterBeingSimulated)
+{
+	FReplicationSystemTestClient* Client = CreateClient();
+
+	// Spawn object with subobject on server
+	UTestReplicatedIrisObject* ServerObject = Server->CreateObject();
+	UTestReplicatedIrisObject* ServerSubObject = Server->CreateSubObject(ServerObject->NetRefHandle, {.ConnectionFilteredComponentCount = 1});
+	Server->ReplicationSystem->SetDeltaCompressionStatus(ServerSubObject->NetRefHandle, ENetObjectDeltaCompressionStatus::Allow);
+
+	// Send and deliver packet
+	Server->UpdateAndSend({Client});
+
+	// Add some values to the autonomous only array
+	for (int32 Value : {4711, 1337})
+	{
+		ServerSubObject->ConnectionFilteredComponents[0]->AutonomousOnlyIntArray.Add(Value);
+	}
+
+	// In order to actually send the object we need to modify a non-AutonomousOnly property too
+	ServerSubObject->ConnectionFilteredComponents[0]->NoneInt += 1;
+
+	// We want a new baseline to be created
+	AllowNewBaselineCreation();
+
+	// Send and deliver packet
+	Server->UpdateAndSend({Client});
+
+	UTestReplicatedIrisObject* ClientSubObject = Cast<UTestReplicatedIrisObject>(Client->GetReplicationBridge()->GetReplicatedObject(ServerSubObject->NetRefHandle));
+
+	// Verify the autonomous only array is unaffected by the server changes
+	UE_NET_ASSERT_TRUE(ClientSubObject->ConnectionFilteredComponents[0]->AutonomousOnlyIntArray.IsEmpty());
+
+	// Set the root object to be autonomous.
+	Server->ReplicationSystem->SetReplicationConditionConnectionFilter(ServerObject->NetRefHandle, EReplicationCondition::RoleAutonomous, Client->ConnectionIdOnServer, true);
+
+	// Add some values to the autonomous only array
+	for (int32 Value : {1, -1})
+	{
+		ServerSubObject->ConnectionFilteredComponents[0]->AutonomousOnlyIntArray.Add(Value);
+	}
+
+	// Send and deliver packet
+	Server->UpdateAndSend({Client});
+
+	// Check that the array is properly replicated in its entirety
+	UE_NET_ASSERT_EQ(ClientSubObject->ConnectionFilteredComponents[0]->AutonomousOnlyIntArray.Num(), ServerSubObject->ConnectionFilteredComponents[0]->AutonomousOnlyIntArray.Num());
+	for (int Index = 0; int32 Value : ServerSubObject->ConnectionFilteredComponents[0]->AutonomousOnlyIntArray)
+	{
+		UE_NET_ASSERT_EQ(ClientSubObject->ConnectionFilteredComponents[0]->AutonomousOnlyIntArray[Index], Value);
+		++Index;
+	}
+}
+
+UE_NET_TEST_FIXTURE(FTestObjectDeltaSerialization, SubObjectPhysicsIsProperlyReplicatedWhenEnabled)
+{
+	FReplicationSystemTestClient* Client = CreateClient();
+
+	// Spawn object with subobject on server
+	UTestReplicatedIrisObject* ServerObject = Server->CreateObject();
+	UTestReplicatedIrisObject* ServerSubObject = Server->CreateSubObject(ServerObject->NetRefHandle, {.ConnectionFilteredComponentCount = 1});
+	Server->ReplicationSystem->SetDeltaCompressionStatus(ServerSubObject->NetRefHandle, ENetObjectDeltaCompressionStatus::Allow);
+
+	// Set the root object to be autonomous. This should prevent SimulatedOrPhysics properties to be replicated.
+	Server->ReplicationSystem->SetReplicationConditionConnectionFilter(ServerObject->NetRefHandle, EReplicationCondition::RoleAutonomous, Client->ConnectionIdOnServer, true);
+
+	// Send and deliver packet
+	Server->UpdateAndSend({Client});
+
+	// Add some values to the physics array
+	for (int32 Value : {4711, 1337})
+	{
+		ServerSubObject->ConnectionFilteredComponents[0]->SimulatedOrPhysicsIntArray.Add(Value);
+	}
+
+	// In order to actually send the object we need to modify a non-physics property too
+	ServerSubObject->ConnectionFilteredComponents[0]->NoneInt += 1;
+
+	// We want a new baseline to be created
+	AllowNewBaselineCreation();
+
+	// Send and deliver packet
+	Server->UpdateAndSend({Client});
+
+	UTestReplicatedIrisObject* ClientSubObject = Cast<UTestReplicatedIrisObject>(Client->GetReplicationBridge()->GetReplicatedObject(ServerSubObject->NetRefHandle));
+
+	// Verify the physics array is unaffected by the server changes
+	UE_NET_ASSERT_TRUE(ClientSubObject->ConnectionFilteredComponents[0]->SimulatedOrPhysicsIntArray.IsEmpty());
+
+	// Enable physics replication
+	Server->ReplicationSystem->SetReplicationCondition(ServerObject->NetRefHandle, EReplicationCondition::ReplicatePhysics, true);
+
+	// Add some values to the physics array
+	for (int32 Value : {1, -1})
+	{
+		ServerSubObject->ConnectionFilteredComponents[0]->SimulatedOrPhysicsIntArray.Add(Value);
+	}
+
+	// Send and deliver packet
+	Server->UpdateAndSend({Client});
+
+	// Check that the array is properly replicated in its entirety
+	UE_NET_ASSERT_EQ(ClientSubObject->ConnectionFilteredComponents[0]->SimulatedOrPhysicsIntArray.Num(), ServerSubObject->ConnectionFilteredComponents[0]->SimulatedOrPhysicsIntArray.Num());
+	for (int Index = 0; int32 Value : ServerSubObject->ConnectionFilteredComponents[0]->SimulatedOrPhysicsIntArray)
+	{
+		UE_NET_ASSERT_EQ(ClientSubObject->ConnectionFilteredComponents[0]->SimulatedOrPhysicsIntArray[Index], Value);
+		++Index;
+	}
+}
+
 UE_NET_TEST_FIXTURE(FTestObjectDeltaSerialization, ToOwnerStateIsReplicatedToOwnerAfterBeingNonOwner)
 {
 	FReplicationSystemTestClient* Client = CreateClient();
@@ -348,6 +514,58 @@ UE_NET_TEST_FIXTURE(FTestObjectDeltaSerialization, ToOwnerStateIsReplicatedToOwn
 	UE_NET_ASSERT_EQ(ClientObject->ConnectionFilteredComponents[0]->ToOwnerB, ServerObject->ConnectionFilteredComponents[0]->ToOwnerB);
 }
 
+UE_NET_TEST_FIXTURE(FTestObjectDeltaSerialization, SubObjectToOwnerStateIsReplicatedToOwnerAfterBeingNonOwner)
+{
+	FReplicationSystemTestClient* Client = CreateClient();
+
+	// Spawn object with subobject on server
+	UTestReplicatedIrisObject* ServerObject = Server->CreateObject();
+	UTestReplicatedIrisObject* ServerSubObject = Server->CreateSubObject(ServerObject->NetRefHandle, {.ConnectionFilteredComponentCount = 1});
+	Server->ReplicationSystem->SetDeltaCompressionStatus(ServerSubObject->NetRefHandle, ENetObjectDeltaCompressionStatus::Allow);
+
+	// Send and deliver packet
+	Server->UpdateAndSend({Client});
+
+	// Add some values to the owner only array
+	for (int32 Value : {4711, 1337})
+	{
+		ServerSubObject->ConnectionFilteredComponents[0]->OwnerOnlyIntArray.Add(Value);
+	}
+
+	// In order to actually send the object we need to modify a non-owner property too
+	ServerSubObject->ConnectionFilteredComponents[0]->NoneInt += 1;
+
+	// We want a new baseline to be created
+	AllowNewBaselineCreation();
+
+	// Send and deliver packet
+	Server->UpdateAndSend({Client});
+
+	UTestReplicatedIrisObject* ClientSubObject = Cast<UTestReplicatedIrisObject>(Client->GetReplicationBridge()->GetReplicatedObject(ServerSubObject->NetRefHandle));
+
+	// Verify the owner only array is unaffected by the server changes
+	UE_NET_ASSERT_TRUE(ClientSubObject->ConnectionFilteredComponents[0]->OwnerOnlyIntArray.IsEmpty());
+
+	// Set owner
+	Server->ReplicationSystem->SetOwningNetConnection(ServerObject->NetRefHandle, Client->ConnectionIdOnServer);
+
+	// Add some values to the owner only array
+	for (int32 Value : {1, -1})
+	{
+		ServerSubObject->ConnectionFilteredComponents[0]->OwnerOnlyIntArray.Add(Value);
+	}
+
+	// Send and deliver packet
+	Server->UpdateAndSend({Client});
+
+	// Check that the array is properly replicated in its entirety
+	UE_NET_ASSERT_EQ(ClientSubObject->ConnectionFilteredComponents[0]->OwnerOnlyIntArray.Num(), ServerSubObject->ConnectionFilteredComponents[0]->OwnerOnlyIntArray.Num());
+	for (int Index = 0; int32 Value : ServerSubObject->ConnectionFilteredComponents[0]->OwnerOnlyIntArray)
+	{
+		UE_NET_ASSERT_EQ(ClientSubObject->ConnectionFilteredComponents[0]->OwnerOnlyIntArray[Index], Value);
+		++Index;
+	}
+}
 UE_NET_TEST_FIXTURE(FTestObjectDeltaSerialization, InFlightChangesForDisabledConditionAreNotResent)
 {
 	FReplicationSystemTestClient* Client = CreateClient();
@@ -504,6 +722,7 @@ UE_NET_TEST_FIXTURE(FTestObjectDeltaSerialization, LostChangesDuringPendingDestr
 	UE_NET_ASSERT_EQ(ClientObject->DynamicStateComponents[0]->IntArray[0], ServerObject->DynamicStateComponents[0]->IntArray[0]);
 }
 
+
 // FTestObjectDeltaSerialization implementation
 void FTestObjectDeltaSerialization::SetUp()
 {
diff --git a/Engine/Source/Runtime/Experimental/Iris/Core/Private/Iris/ReplicationSystem/Conditionals/ReplicationConditionals.cpp b/Engine/Source/Runtime/Experimental/Iris/Core/Private/Iris/ReplicationSystem/Conditionals/ReplicationConditionals.cpp
--- a/Engine/Source/Runtime/Experimental/Iris/Core/Private/Iris/ReplicationSystem/Conditionals/ReplicationConditionals.cpp
+++ b/Engine/Source/Runtime/Experimental/Iris/Core/Private/Iris/ReplicationSystem/Conditionals/ReplicationConditionals.cpp
@@ -120,25 +120,40 @@ bool FReplicationConditionals::SetConditionConnectionFilter(FInternalNetRefIndex
 		const uint32 ConnIdForBaselineInvalidation = (bEnable ? ConnectionId : ObjectInfo->AutonomousConnectionId);
 		ObjectInfo->AutonomousConnectionId = uint16(AutonomousConnectionId);
 
-		// TODO: Make sure we also invalidate baselines for subobjects
-		BaselineInvalidationTracker->InvalidateBaselines(ObjectIndex, ConnIdForBaselineInvalidation);
-
 		MarkRemoteRoleDirty(ObjectIndex);
 		// Mark object as having dirty global conditional that should be evaluated before next send
 		ObjectsWithDirtyLifetimeConditionals.SetBit(ObjectIndex);
+
+		InvalidateBaselinesForObjectHierarchy(ObjectIndex, TConstArrayView<uint32>(&ConnIdForBaselineInvalidation, 1));
 	}
 
 	return true;
 }
 
-void FReplicationConditionals::SetOwningConnection(FInternalNetRefIndex ObjectIndex, uint32 ConnectionId)
+void FReplicationConditionals::SetOwningConnection(FInternalNetRefIndex ObjectIndex, uint32 OwningConnectionId)
 {
-	if (ReplicationFiltering->GetOwningConnection(ObjectIndex) != ConnectionId)
+	const uint32 OldOwningConnectionId = ReplicationFiltering->GetOwningConnection(ObjectIndex);
+	if (OldOwningConnectionId != OwningConnectionId && (OwningConnectionId == InvalidConnectionId || ReplicationConnections->IsValidConnection(OwningConnectionId)))
 	{
-		UE_LOG(LogIrisConditionals, Verbose, TEXT("SetOwningConnection on object %u. Connection: %u"), ObjectIndex, ConnectionId);
+		UE_LOG(LogIrisConditionals, Verbose, TEXT("SetOwningConnection on object %u. Connection: %u"), ObjectIndex, OwningConnectionId);
 
 		// Mark object as having dirty global conditional that should be evaluated before next send
 		ObjectsWithDirtyLifetimeConditionals.SetBit(ObjectIndex);
+
+		// Invalidate baselines for connections affected by the owner change.
+		{
+			const uint32 ConnectionIdToInvalidateCandidates[] = {OldOwningConnectionId, OwningConnectionId};
+			uint32 ConnectionIdsToInvalidate[UE_ARRAY_COUNT(ConnectionIdToInvalidateCandidates)];
+			uint32 ConnectionIdCount = 0;
+			for (uint32 ConnectionId : ConnectionIdToInvalidateCandidates)
+			{
+				if (ConnectionId != InvalidConnectionId)
+				{
+					ConnectionIdsToInvalidate[ConnectionIdCount++] = ConnectionId;
+				}
+			}
+			InvalidateBaselinesForObjectHierarchy(ObjectIndex, TConstArrayView<uint32>(ConnectionIdsToInvalidate, ConnectionIdCount));
+		}
 	}
 }
 
@@ -172,8 +187,8 @@ bool FReplicationConditionals::SetCondition(FInternalNetRefIndex ObjectIndex, ER
 			UE_LOG(LogIrisConditionals, Verbose, TEXT("SetCondition object %s. EReplicationCondition::ReplicatePhysics: %u"), *NetRefHandleManager->PrintObjectFromIndex(ObjectIndex), bEnable ? 1U : 0U);
 
 			// We only care to track this change if the condition is enabled.
-			// TODO: Make sure we also invalidate baselines for subobjects
-			BaselineInvalidationTracker->InvalidateBaselines(ObjectIndex, BaselineInvalidationTracker->InvalidateBaselineForAllConnections);
+			const uint32 ConnIdForBaselineInvalidation = BaselineInvalidationTracker->InvalidateBaselineForAllConnections;
+			InvalidateBaselinesForObjectHierarchy(ObjectIndex, TConstArrayView<uint32>(&ConnIdForBaselineInvalidation, 1));
 
 			// Mark object as having dirty global conditional that should be evaluated before next send
 			ObjectsWithDirtyLifetimeConditionals.SetBit(ObjectIndex);
@@ -296,7 +311,6 @@ bool FReplicationConditionals::SetPropertyCustomCondition(FInternalNetRefIndex O
 		}
 
 		const FReplicationStateMemberChangeMaskDescriptor& ChangeMaskDescriptor = StateDescriptor->MemberChangeMaskDescriptors[RepIndexToMemberIndexDescriptor.MemberIndex];
-		checkSlow(ChangeMaskDescriptor.BitCount == 1);
 		FNetBitArrayView ConditionalChangeMask = UE::Net::Private::GetMemberConditionalChangeMask(Fragment.ExternalSrcBuffer, StateDescriptor);
 
 		if (bIsActive)
@@ -1065,4 +1079,35 @@ void FReplicationConditionals::MarkPropertyDirty(FInternalNetRefIndex ObjectInde
 	UE_LOG(LogIris, Warning, TEXT("Trying to mark non-existing property with RepIndex %u in protocol %s as dirty"), RepIndex, ToCStr(Protocol->DebugName));
 }
 
+void FReplicationConditionals::InvalidateBaselinesForObjectHierarchy(uint32 ObjectIndex, const TConstArrayView<uint32>& ConnectionsToInvalidate)
+{
+	// Invalidate baselines for root object
+	{
+		const FNetRefHandleManager::FReplicatedObjectData& ObjectData = NetRefHandleManager->GetReplicatedObjectDataNoCheck(ObjectIndex);
+		if (EnumHasAnyFlags(ObjectData.Protocol->ProtocolTraits, EReplicationProtocolTraits::HasLifetimeConditionals))
+		{
+			for (const uint32 ConnId : ConnectionsToInvalidate)
+			{
+				BaselineInvalidationTracker->InvalidateBaselines(ObjectIndex, ConnId);
+			}
+		}
+	}
+
+	// Invalidate baselines for subobjects
+	for (const FInternalNetRefIndex SubObjectIndex : NetRefHandleManager->GetSubObjects(ObjectIndex))
+	{
+		const FNetRefHandleManager::FReplicatedObjectData& SubObjectData = NetRefHandleManager->GetReplicatedObjectDataNoCheck(SubObjectIndex);
+		if (SubObjectData.IsSubObject())
+		{
+			if (EnumHasAnyFlags(SubObjectData.Protocol->ProtocolTraits, EReplicationProtocolTraits::HasLifetimeConditionals))
+			{
+				for (const uint32 ConnId : ConnectionsToInvalidate)
+				{
+					BaselineInvalidationTracker->InvalidateBaselines(SubObjectIndex, ConnId);
+				}
+			}
+		}
+	}
+}
+
 }
diff --git a/Engine/Source/Runtime/Experimental/Iris/Core/Private/Iris/ReplicationSystem/Conditionals/ReplicationConditionals.h b/Engine/Source/Runtime/Experimental/Iris/Core/Private/Iris/ReplicationSystem/Conditionals/ReplicationConditionals.h
--- a/Engine/Source/Runtime/Experimental/Iris/Core/Private/Iris/ReplicationSystem/Conditionals/ReplicationConditionals.h
+++ b/Engine/Source/Runtime/Experimental/Iris/Core/Private/Iris/ReplicationSystem/Conditionals/ReplicationConditionals.h
@@ -136,6 +136,9 @@ private:
 	void MarkRemoteRoleDirty(FInternalNetRefIndex ObjectIndex);
 	uint16 GetRemoteRoleRepIndex(const FReplicationProtocol* Protocol);
 
+	// Invalidates baselines for root object and subobjects with lifetime conditionals.
+	void InvalidateBaselinesForObjectHierarchy(uint32 ObjectIndex, const TConstArrayView<uint32>& ConnectionsToInvalidate);
+
 private:
 	static constexpr uint16 InvalidRepIndex = 65535U;
 
diff --git a/Engine/Source/Runtime/Experimental/Iris/Core/Private/Iris/ReplicationSystem/Filtering/ReplicationFiltering.cpp b/Engine/Source/Runtime/Experimental/Iris/Core/Private/Iris/ReplicationSystem/Filtering/ReplicationFiltering.cpp
--- a/Engine/Source/Runtime/Experimental/Iris/Core/Private/Iris/ReplicationSystem/Filtering/ReplicationFiltering.cpp
+++ b/Engine/Source/Runtime/Experimental/Iris/Core/Private/Iris/ReplicationSystem/Filtering/ReplicationFiltering.cpp
@@ -15,7 +15,6 @@
 #include "Iris/ReplicationSystem/ReplicationProtocol.h"
 #include "Iris/ReplicationSystem/ReplicationSystem.h"
 #include "Iris/ReplicationSystem/ReplicationSystemInternal.h"
-#include "Iris/ReplicationSystem/DeltaCompression/DeltaCompressionBaselineInvalidationTracker.h"
 #include "Iris/ReplicationSystem/Filtering/NetObjectGroups.h"
 #include "Iris/ReplicationSystem/Filtering/NetObjectFilterDefinitions.h"
 #include "Iris/ReplicationSystem/Filtering/ReplicationFilteringConfig.h"
@@ -194,7 +193,6 @@ void FReplicationFiltering::Init(FReplicationFilteringInitParams& Params)
 
 	Connections = Params.Connections;
 	NetRefHandleManager = Params.NetRefHandleManager;
-	BaselineInvalidationTracker = Params.BaselineInvalidationTracker;
 	Groups = Params.Groups;
 
 	MaxInternalNetRefIndex = Params.MaxInternalNetRefIndex;
@@ -439,8 +437,6 @@ void FReplicationFiltering::SetOwningConnection(FInternalNetRefIndex ObjectIndex
 			bHasDirtyConnectionFilter = 1;
 			ObjectsWithDirtyConnectionFilter.SetBit(ObjectIndex);
 		}
-
-		InvalidateBaselinesForObject(ObjectIndex, ConnectionId, OldConnectionId);
 	}
 }
 
@@ -2693,42 +2689,6 @@ void FReplicationFiltering::RemoveFromDynamicFilter(uint32 ObjectIndex, uint32 F
 	HysteresisState.ClearFromHysteresis(ObjectIndex);
 }
 
-void FReplicationFiltering::InvalidateBaselinesForObject(uint32 ObjectIndex, uint32 NewOwningConnectionId, uint32 PrevOwningConnectionId)
-{
-	{
-		const FNetRefHandleManager::FReplicatedObjectData& ObjectData = NetRefHandleManager->GetReplicatedObjectDataNoCheck(ObjectIndex);
-		// $IRIS TODO Only invalidate baselines if the object has conditions based on owner/non-owner.
-		if (EnumHasAnyFlags(ObjectData.Protocol->ProtocolTraits, EReplicationProtocolTraits::HasLifetimeConditionals))
-		{
-			for (const uint32 ConnId : {NewOwningConnectionId, PrevOwningConnectionId})
-			{
-				if (ConnId)
-				{
-					BaselineInvalidationTracker->InvalidateBaselines(ObjectIndex, ConnId);
-				}
-			}
-		}
-	}
-
-	// Invalidate baselines for subobjects
-	for (const FInternalNetRefIndex SubObjectIndex : NetRefHandleManager->GetSubObjects(ObjectIndex))
-	{
-		const FNetRefHandleManager::FReplicatedObjectData& SubObjectData = NetRefHandleManager->GetReplicatedObjectDataNoCheck(SubObjectIndex);
-		if (SubObjectData.IsSubObject())
-		{
-			if (EnumHasAnyFlags(SubObjectData.Protocol->ProtocolTraits, EReplicationProtocolTraits::HasLifetimeConditionals))
-			{
-				for (const uint32 ConnId : {NewOwningConnectionId, PrevOwningConnectionId})
-				{
-					if (ConnId)
-					{
-						BaselineInvalidationTracker->InvalidateBaselines(ObjectIndex, ConnId);
-					}
-				}
-			}
-		}
-	}
-}
 
 TArrayView<FNetObjectFilteringInfo> FReplicationFiltering::GetNetObjectFilteringInfos()
 {
diff --git a/Engine/Source/Runtime/Experimental/Iris/Core/Private/Iris/ReplicationSystem/Filtering/ReplicationFiltering.h b/Engine/Source/Runtime/Experimental/Iris/Core/Private/Iris/ReplicationSystem/Filtering/ReplicationFiltering.h
--- a/Engine/Source/Runtime/Experimental/Iris/Core/Private/Iris/ReplicationSystem/Filtering/ReplicationFiltering.h
+++ b/Engine/Source/Runtime/Experimental/Iris/Core/Private/Iris/ReplicationSystem/Filtering/ReplicationFiltering.h
@@ -18,7 +18,6 @@ namespace UE::Net
 	typedef uint32 FNetObjectFilterHandle;
 	namespace Private
 	{
-		class FDeltaCompressionBaselineInvalidationTracker;
 		class FNetRefHandleManager;
 		class FNetObjectGroups;
 		class FReplicationConnections;
@@ -43,10 +42,9 @@ private:
 
 struct FReplicationFilteringInitParams
 {
-	TObjectPtr<UReplicationSystem> ReplicationSystem;
+	TObjectPtr<UReplicationSystem> ReplicationSystem = nullptr;
 	const FNetRefHandleManager* NetRefHandleManager = nullptr;
 	FNetObjectGroups* Groups = nullptr;
-	FDeltaCompressionBaselineInvalidationTracker* BaselineInvalidationTracker = nullptr;
 	FReplicationConnections* Connections = nullptr;
 	FInternalNetRefIndex MaxInternalNetRefIndex = 0;
 	uint32 MaxGroupCount = 0;
@@ -267,8 +265,6 @@ private:
 	void NotifyFiltersOfDirtyObjects();
 	void BatchNotifyFiltersOfDirtyObjects(FUpdateDirtyObjectsBatchHelper& BatchHelper, const uint32* ObjectIndices, uint32 ObjectCount);
 
-	void InvalidateBaselinesForObject(uint32 ObjectIndex, uint32 NewOwningConnectionId, uint32 PrevOwningConnectionId);
-
 	/** Returns all the filtering infos. */
 	TArrayView<FNetObjectFilteringInfo> GetNetObjectFilteringInfos();
 
@@ -318,9 +314,6 @@ private:
 	// Groups
 	FNetObjectGroups* Groups = nullptr;
 
-	// Baseline invalidation tracker
-	FDeltaCompressionBaselineInvalidationTracker* BaselineInvalidationTracker = nullptr;
-
 	// Connection specifics
 	FReplicationConnections* Connections = nullptr;
 	TArray<FPerConnectionInfo> ConnectionInfos;
diff --git a/Engine/Source/Runtime/Experimental/Iris/Core/Private/Iris/ReplicationSystem/ReplicationSystem.cpp b/Engine/Source/Runtime/Experimental/Iris/Core/Private/Iris/ReplicationSystem/ReplicationSystem.cpp
--- a/Engine/Source/Runtime/Experimental/Iris/Core/Private/Iris/ReplicationSystem/ReplicationSystem.cpp
+++ b/Engine/Source/Runtime/Experimental/Iris/Core/Private/Iris/ReplicationSystem/ReplicationSystem.cpp
@@ -211,7 +211,6 @@ public:
 			InitParams.Connections = &ReplicationSystemInternal.GetConnections();
 			InitParams.NetRefHandleManager = &NetRefHandleManager;
 			InitParams.Groups = &ReplicationSystemInternal.GetGroups();
-			InitParams.BaselineInvalidationTracker = &ReplicationSystemInternal.GetDeltaCompressionBaselineInvalidationTracker();
 			InitParams.MaxInternalNetRefIndex = CurrentMaxInternalNetRefIndex;
 			InitParams.MaxGroupCount = Params.MaxNetObjectGroupCount;
 			ReplicationFiltering.Init(InitParams);
