Iris
* Fix dormant objects getting replicated if one of their push model property was set dirty or force net update was called on them.
* Added dormancy iris unit test

#rnx

diff --git a/Engine/Plugins/Runtime/ReplicationSystemTestPlugin/Source/Private/Tests/ReplicationSystem/ReplicatedTestObject.cpp b/Engine/Plugins/Runtime/ReplicationSystemTestPlugin/Source/Private/Tests/ReplicationSystem/ReplicatedTestObject.cpp
--- a/Engine/Plugins/Runtime/ReplicationSystemTestPlugin/Source/Private/Tests/ReplicationSystem/ReplicatedTestObject.cpp
+++ b/Engine/Plugins/Runtime/ReplicationSystemTestPlugin/Source/Private/Tests/ReplicationSystem/ReplicatedTestObject.cpp
@@ -603,6 +603,49 @@ void UTestReplicatedObjectWithRepNotifies::OnRep_IntB(int32 OldInt)
 	PrevIntBStoredInOnRep = OldInt;
 }
 
+
+//////////////////////////////////////////////////////////////////////////
+// Implementation for UTestReplicatedIrisPushModelObject
+//////////////////////////////////////////////////////////////////////////
+void UTestReplicatedIrisPushModelObject::GetLifetimeReplicatedProps(TArray< class FLifetimeProperty >& OutLifetimeProps) const
+{
+	FDoRepLifetimeParams Params{ .bIsPushBased = true };
+	DOREPLIFETIME_WITH_PARAMS(UTestReplicatedIrisPushModelObject, IntA, Params);
+	DOREPLIFETIME_WITH_PARAMS(UTestReplicatedIrisPushModelObject, IntB, Params);
+}
+
+void UTestReplicatedIrisPushModelObject::RegisterReplicationFragments(UE::Net::FFragmentRegistrationContext& Context, UE::Net::EFragmentRegistrationFlags RegistrationFlags)
+{
+	// Base object owns the fragment in this case
+	{
+		this->ReplicationFragments.Reset();
+		UE::Net::FReplicationFragmentUtil::CreateAndRegisterFragmentsForObject(this, Context, RegistrationFlags, &this->ReplicationFragments);
+	}
+}
+
+
+void UTestReplicatedIrisPushModelObject::SetIntA(int32 InValue)
+{
+	IntA = InValue;
+	MARK_PROPERTY_DIRTY_FROM_NAME(UTestReplicatedIrisPushModelObject, IntA, this);
+}
+
+int32 UTestReplicatedIrisPushModelObject::GetIntA() const
+{
+	return IntA;
+}
+
+void UTestReplicatedIrisPushModelObject::SetIntB(int32 InValue)
+{
+	IntB = InValue;
+	MARK_PROPERTY_DIRTY_FROM_NAME(UTestReplicatedIrisPushModelObject, IntB, this);
+}
+
+int32 UTestReplicatedIrisPushModelObject::GetIntB() const
+{
+	return IntB;
+}
+
 //////////////////////////////////////////////////////////////////////////
 // Implementation for UTestReplicatedIrisObjectWithDynamicCondition
 //////////////////////////////////////////////////////////////////////////
diff --git a/Engine/Plugins/Runtime/ReplicationSystemTestPlugin/Source/Private/Tests/ReplicationSystem/ReplicatedTestObject.h b/Engine/Plugins/Runtime/ReplicationSystemTestPlugin/Source/Private/Tests/ReplicationSystem/ReplicatedTestObject.h
--- a/Engine/Plugins/Runtime/ReplicationSystemTestPlugin/Source/Private/Tests/ReplicationSystem/ReplicatedTestObject.h
+++ b/Engine/Plugins/Runtime/ReplicationSystemTestPlugin/Source/Private/Tests/ReplicationSystem/ReplicatedTestObject.h
@@ -492,6 +492,37 @@ public:
 	virtual void RegisterReplicationFragments(UE::Net::FFragmentRegistrationContext& Context, UE::Net::EFragmentRegistrationFlags RegistrationFlags) override;
 };
 
+/**
+ * Replicated object with PushModel properties
+ */
+UCLASS()
+class UTestReplicatedIrisPushModelObject : public UReplicatedTestObject
+{
+	GENERATED_BODY()
+
+public:
+
+	// Network interface must be part of base.
+	virtual void RegisterReplicationFragments(UE::Net::FFragmentRegistrationContext& Fragments, UE::Net::EFragmentRegistrationFlags RegistrationFlags) override;
+
+	void SetIntA(int32 InValue);
+	int32 GetIntA() const;
+
+	void SetIntB(int32 InValue);
+	int32 GetIntB() const;
+
+private:
+	UPROPERTY(Transient, Replicated)
+	int32 IntA;
+
+	UPROPERTY(Transient, Replicated)
+	int32 IntB;
+
+public:
+	// Network data only for test
+	TArray<UE::Net::FReplicationFragment*> ReplicationFragments;
+};
+
 /**
 * Example of type specific ReplicationBridge, it currently only works with UTestReplicationSystem_TestClass
 */
diff --git a/Engine/Plugins/Runtime/ReplicationSystemTestPlugin/Source/Private/Tests/ReplicationSystem/TestDormancy.cpp b/Engine/Plugins/Runtime/ReplicationSystemTestPlugin/Source/Private/Tests/ReplicationSystem/TestDormancy.cpp
new file mode 100644
--- /dev/null
+++ b/Engine/Plugins/Runtime/ReplicationSystemTestPlugin/Source/Private/Tests/ReplicationSystem/TestDormancy.cpp
@@ -0,0 +1,75 @@
+// Copyright Epic Games, Inc. All Rights Reserved.
+
+#include "NetworkAutomationTest.h"
+#include "NetworkAutomationTestMacros.h"
+#include "Iris/ReplicationSystem/ReplicationSystemInternal.h"
+#include "Tests/ReplicationSystem/ReplicationSystemTestFixture.h"
+#include "Tests/ReplicationSystem/ReplicationSystemServerClientTestFixture.h"
+#include "Tests/ReplicationSystem/RPC/ReplicatedTestObjectWithRPC.h"
+#include "Tests/ReplicationSystem/RPC/RPCTestFixture.h"
+
+namespace UE::Net::Private
+{
+
+UE_NET_TEST_FIXTURE(FReplicationSystemServerClientTestFixture, TestDormancyWithPushModel)
+{
+	// Add a client
+	FReplicationSystemTestClient* Client = CreateClient();
+
+	// Spawn object on server
+	UTestReplicatedIrisPushModelObject* ServerObject = Server->CreateObject<UTestReplicatedIrisPushModelObject>();
+
+	// Send and deliver packet
+	Server->UpdateAndSend({Client});
+
+	// Object should be replicated on the client
+	UTestReplicatedIrisPushModelObject* ClientObject = Cast<UTestReplicatedIrisPushModelObject>(Client->GetReplicationBridge()->GetReplicatedObject(ServerObject->NetRefHandle));
+	UE_NET_ASSERT_TRUE(ClientObject != nullptr);
+
+	// Set a property dirty
+	ServerObject->SetIntA(1);
+
+	// Property should have been received
+	Server->UpdateAndSend({ Client });	
+	UE_NET_ASSERT_EQ(ClientObject->GetIntA(), 1);
+
+	// Set object dormant
+	Server->GetReplicationBridge()->SetObjectWantsToBeDormant(ServerObject->NetRefHandle, true);
+
+	// Set a property dirty on the same frame as setting the object dormant
+	ServerObject->SetIntA(2);
+
+	// Property should have been received
+	Server->UpdateAndSend({ Client });
+	UE_NET_ASSERT_EQ(ClientObject->GetIntA(), 2);
+
+	// Set a property dirty again
+	ServerObject->SetIntA(3);
+
+	// The property on the client should not be updated
+	Server->UpdateAndSend({ Client });
+	UE_NET_ASSERT_NE(ClientObject->GetIntA(), 3);
+
+	// Flush dormancy on the object
+	Server->GetReplicationBridge()->ForceUpdateWantsToBeDormantObject(ServerObject->NetRefHandle);
+
+	// The property on the client should now be received
+	Server->UpdateAndSend({ Client });
+	UE_NET_ASSERT_EQ(ClientObject->GetIntA(), 3);
+
+	// Set a property dirty again
+	ServerObject->SetIntA(4);
+
+	// The property on the client should not be updated
+	Server->UpdateAndSend({ Client });
+	UE_NET_ASSERT_NE(ClientObject->GetIntA(), 4);
+
+	// Wake object from dormancy
+	Server->GetReplicationBridge()->SetObjectWantsToBeDormant(ServerObject->NetRefHandle, false);
+
+	// The property on the client should be received now
+	Server->UpdateAndSend({ Client });
+	UE_NET_ASSERT_EQ(ClientObject->GetIntA(), 4);
+}
+
+} // end namespace UE::Net::Private
\ No newline at end of file
diff --git a/Engine/Source/Runtime/Experimental/Iris/Core/Private/Iris/ReplicationSystem/ObjectReplicationBridge.cpp b/Engine/Source/Runtime/Experimental/Iris/Core/Private/Iris/ReplicationSystem/ObjectReplicationBridge.cpp
--- a/Engine/Source/Runtime/Experimental/Iris/Core/Private/Iris/ReplicationSystem/ObjectReplicationBridge.cpp
+++ b/Engine/Source/Runtime/Experimental/Iris/Core/Private/Iris/ReplicationSystem/ObjectReplicationBridge.cpp
@@ -1003,14 +1003,14 @@ void UObjectReplicationBridge::BuildPollList(UE::Net::FNetBitArrayView ObjectsCo
 		ObjectsConsideredForPolling.Copy(RelevantObjects);
 	}
 
-	// Mask off objects pending dormancy as we do not want to poll/pre-update them unless they are marked for flush or are dirty
+	// Mask off objects pending dormancy as we do not want to poll/pre-update them unless they are marked for flush
 	if (bUseDormancyToFilterPolling)
 	{
 		IRIS_PROFILER_SCOPE(BuildPollList_Dormancy);
 
-		// Mask off objects pending dormancy that are not dirty
+		// Mask off objects pending dormancy
 		const FNetBitArrayView AccumulatedDirtyObjects = DirtyNetObjectTracker.GetAccumulatedDirtyNetObjects();
-		ObjectsConsideredForPolling.CombineMultiple(FNetBitArrayView::AndNotOp, WantToBeDormantObjects, FNetBitArrayView::AndNotOp, AccumulatedDirtyObjects);
+		ObjectsConsideredForPolling.Combine(WantToBeDormantObjects, FNetBitArrayView::AndNotOp);
 
 		FNetBitArrayView ForceNetUpdateObjects = ReplicationSystemInternal->GetDirtyNetObjectTracker().GetForceNetUpdateObjects();
 
