Replication system tests fix compilation errors
#rnx

diff --git a/Engine/Plugins/Runtime/ReplicationSystemTestPlugin/Source/Private/Tests/ReplicationState/TestPropertyReplicationState.h b/Engine/Plugins/Runtime/ReplicationSystemTestPlugin/Source/Private/Tests/ReplicationState/TestPropertyReplicationState.h
--- a/Engine/Plugins/Runtime/ReplicationSystemTestPlugin/Source/Private/Tests/ReplicationState/TestPropertyReplicationState.h
+++ b/Engine/Plugins/Runtime/ReplicationSystemTestPlugin/Source/Private/Tests/ReplicationState/TestPropertyReplicationState.h
@@ -130,7 +130,7 @@ public:
 	UTestPropertyReplicationState_TestClassWithTArray() : UReplicatedTestObject() {}
 
 	UPROPERTY(ReplicatedUsing=OnRep_ReferencedObjects)
-	TArray<UObject*> ReferencedObjects;
+	TArray<TObjectPtr<UObject>> ReferencedObjects;
 
 	UPROPERTY(Replicated)
 	uint32 ForceReplication = 0;
