Iris
Fix for compile errors in TestPropertyReplicationState

diff --git a/Engine/Plugins/Runtime/ReplicationSystemTestPlugin/Source/Private/Tests/ReplicationState/TestPropertyReplicationState.cpp b/Engine/Plugins/Runtime/ReplicationSystemTestPlugin/Source/Private/Tests/ReplicationState/TestPropertyReplicationState.cpp
--- a/Engine/Plugins/Runtime/ReplicationSystemTestPlugin/Source/Private/Tests/ReplicationState/TestPropertyReplicationState.cpp
+++ b/Engine/Plugins/Runtime/ReplicationSystemTestPlugin/Source/Private/Tests/ReplicationState/TestPropertyReplicationState.cpp
@@ -542,7 +542,7 @@ UE_NET_TEST_FIXTURE(FTestPropertyReplicationStateContext, TestArrayOnRepWithUnre
 			ClientObject = Cast<UTestPropertyReplicationState_TestClassWithTArray>(Client->GetReplicationBridge()->GetReplicatedObject(ServerObject->NetRefHandle));
 			UE_NET_ASSERT_NE(ClientObject, nullptr);
 			UE_NET_ASSERT_FALSE(ClientObject->ReferencedObjects.IsEmpty());
-			UE_NET_ASSERT_NE(ClientObject->ReferencedObjects[0], nullptr);
+			UE_NET_ASSERT_NE(ClientObject->ReferencedObjects[0].Get(), nullptr);
 		}
 
 		// Step 2. Modify the referenced objects array such that it contains both resolvable and unresolvable references.
@@ -556,9 +556,9 @@ UE_NET_TEST_FIXTURE(FTestPropertyReplicationStateContext, TestArrayOnRepWithUnre
 
 			ClientObject = Cast<UTestPropertyReplicationState_TestClassWithTArray>(Client->GetReplicationBridge()->GetReplicatedObject(ServerObject->NetRefHandle));
 			UE_NET_ASSERT_EQ(ClientObject->ReferencedObjects.Num(), 3);
-			UE_NET_ASSERT_EQ(ClientObject->ReferencedObjects[0], Client->GetReplicationBridge()->GetReplicatedObject(ServerReferencedObjectA->NetRefHandle));
-			UE_NET_ASSERT_EQ(ClientObject->ReferencedObjects[1], nullptr);
-			UE_NET_ASSERT_EQ(ClientObject->ReferencedObjects[2], nullptr);
+			UE_NET_ASSERT_EQ(ClientObject->ReferencedObjects[0].Get(), Client->GetReplicationBridge()->GetReplicatedObject(ServerReferencedObjectA->NetRefHandle));
+			UE_NET_ASSERT_EQ(ClientObject->ReferencedObjects[1].Get(), nullptr);
+			UE_NET_ASSERT_EQ(ClientObject->ReferencedObjects[2].Get(), nullptr);
 		}
 
 		// Step 3. Modify non-array property on server and null out reference on client. OnRep call depending on whether we're applying previously received state with unresolved references or not.
@@ -582,7 +582,7 @@ UE_NET_TEST_FIXTURE(FTestPropertyReplicationStateContext, TestArrayOnRepWithUnre
 			Server->UpdateAndSend({ Client });
 
 			UE_NET_ASSERT_TRUE(ClientObject->bOnRepWasCalled);
-			UE_NET_ASSERT_EQ(ClientObject->ReferencedObjects[1], Client->GetReplicationBridge()->GetReplicatedObject(ServerReferencedObjectC->NetRefHandle));
+			UE_NET_ASSERT_EQ(ClientObject->ReferencedObjects[1].Get(), Client->GetReplicationBridge()->GetReplicatedObject(ServerReferencedObjectC->NetRefHandle));
 		}
 
 		// Step 5. Allow object D to be replicated. Expecting OnRep.
@@ -594,7 +594,7 @@ UE_NET_TEST_FIXTURE(FTestPropertyReplicationStateContext, TestArrayOnRepWithUnre
 			Server->UpdateAndSend({ Client });
 
 			UE_NET_ASSERT_TRUE(ClientObject->bOnRepWasCalled);
-			UE_NET_ASSERT_EQ(ClientObject->ReferencedObjects[2], Client->GetReplicationBridge()->GetReplicatedObject(ServerReferencedObjectD->NetRefHandle));
+			UE_NET_ASSERT_EQ(ClientObject->ReferencedObjects[2].Get(), Client->GetReplicationBridge()->GetReplicatedObject(ServerReferencedObjectD->NetRefHandle));
 		}
 
 		// Step 6. Resize array on server. Expecting OnRep.
