Iris - Fixed off by one error in FReplicationSystemTestNode::ConvertPollPeriodIntoFrequency

#rnx

diff --git a/Engine/Plugins/Runtime/ReplicationSystemTestPlugin/Source/Private/Tests/ReplicationSystem/ReplicationSystemServerClientTestFixture.cpp b/Engine/Plugins/Runtime/ReplicationSystemTestPlugin/Source/Private/Tests/ReplicationSystem/ReplicationSystemServerClientTestFixture.cpp
--- a/Engine/Plugins/Runtime/ReplicationSystemTestPlugin/Source/Private/Tests/ReplicationSystem/ReplicationSystemServerClientTestFixture.cpp
+++ b/Engine/Plugins/Runtime/ReplicationSystemTestPlugin/Source/Private/Tests/ReplicationSystem/ReplicationSystemServerClientTestFixture.cpp
@@ -404,7 +404,7 @@ uint32 FReplicationSystemTestNode::GetReplicationSystemId() const
 
 float FReplicationSystemTestNode::ConvertPollPeriodIntoFrequency(uint32 PollPeriod) const
 {
-	const float PollFrequency = ReplicationBridge->GetMaxTickRate() / (float)PollPeriod;
+	const float PollFrequency = ReplicationBridge->GetMaxTickRate() / (float)(PollPeriod + 1);
 	return PollFrequency;
 }
 
