Iris
- Fixed bad init code for ReplicationSystemTest

diff --git a/Engine/Plugins/Runtime/ReplicationSystemTestPlugin/Source/Private/Tests/ReplicationSystem/ReplicationSystemServerClientTestFixture.cpp b/Engine/Plugins/Runtime/ReplicationSystemTestPlugin/Source/Private/Tests/ReplicationSystem/ReplicationSystemServerClientTestFixture.cpp
--- a/Engine/Plugins/Runtime/ReplicationSystemTestPlugin/Source/Private/Tests/ReplicationSystem/ReplicationSystemServerClientTestFixture.cpp
+++ b/Engine/Plugins/Runtime/ReplicationSystemTestPlugin/Source/Private/Tests/ReplicationSystem/ReplicationSystemServerClientTestFixture.cpp
@@ -86,8 +86,8 @@ void FReplicationSystemTestNode::Setup(bool bIsServer, const TCHAR* Name, FRepli
 		NetTokenStore->Init(NetTokenStoreInitParams);
 
 		// Register data stores for supported types, $TODO: make this configurable.
-		NetTokenDataStores.Add(MakeUnique<FStringTokenStore>(*NetTokenStore));
-		NetTokenDataStores.Add(MakeUnique<FNameTokenStore>(*NetTokenStore));	
+		NetTokenStore->CreateAndRegisterDataStore<FStringTokenStore>();
+		NetTokenStore->CreateAndRegisterDataStore<FNameTokenStore>();
 	}
 
 	ReplicationBridge = NewObject<UReplicatedTestObjectBridge>();
diff --git a/Engine/Plugins/Runtime/ReplicationSystemTestPlugin/Source/Private/Tests/ReplicationSystem/ReplicationSystemServerClientTestFixture.h b/Engine/Plugins/Runtime/ReplicationSystemTestPlugin/Source/Private/Tests/ReplicationSystem/ReplicationSystemServerClientTestFixture.h
--- a/Engine/Plugins/Runtime/ReplicationSystemTestPlugin/Source/Private/Tests/ReplicationSystem/ReplicationSystemServerClientTestFixture.h
+++ b/Engine/Plugins/Runtime/ReplicationSystemTestPlugin/Source/Private/Tests/ReplicationSystem/ReplicationSystemServerClientTestFixture.h
@@ -168,7 +168,6 @@ public:
 
 public:
 	TUniquePtr<UE::Net::FNetTokenStore> NetTokenStore;
-	TArray<TUniquePtr<UE::Net::FNetTokenDataStore>> NetTokenDataStores;
 	UReplicationSystem* ReplicationSystem;
 	UReplicatedTestObjectBridge* ReplicationBridge;
 	TArray<TStrongObjectPtr<UObject>> CreatedObjects;
diff --git a/Engine/Plugins/Runtime/ReplicationSystemTestPlugin/Source/Private/Tests/ReplicationSystem/ReplicationSystemTestFixture.h b/Engine/Plugins/Runtime/ReplicationSystemTestPlugin/Source/Private/Tests/ReplicationSystem/ReplicationSystemTestFixture.h
--- a/Engine/Plugins/Runtime/ReplicationSystemTestPlugin/Source/Private/Tests/ReplicationSystem/ReplicationSystemTestFixture.h
+++ b/Engine/Plugins/Runtime/ReplicationSystemTestPlugin/Source/Private/Tests/ReplicationSystem/ReplicationSystemTestFixture.h
@@ -42,8 +42,8 @@ protected:
 			NetTokenStore->Init(NetTokenStoreInitParams);
 
 			// Register data stores for supported types, $TODO: make this configurable.
-			NetTokenDataStores.Add(MakeUnique<FStringTokenStore>(*NetTokenStore));
-			NetTokenDataStores.Add(MakeUnique<FNameTokenStore>(*NetTokenStore));	
+			NetTokenStore->CreateAndRegisterDataStore<FStringTokenStore>();
+			NetTokenStore->CreateAndRegisterDataStore<FNameTokenStore>();
 		}
 
 		ReplicationBridge = NewObject<UReplicatedTestObjectBridge>();
@@ -125,7 +125,6 @@ protected:
 	}
 
 	TUniquePtr<UE::Net::FNetTokenStore> NetTokenStore;
-	TArray<TUniquePtr<UE::Net::FNetTokenDataStore>> NetTokenDataStores;
 	UReplicationSystem* ReplicationSystem;
 	UReplicatedTestObjectBridge* ReplicationBridge;
 	
