Iris - Implemented TestMessage stream operator for ENetFilterStatus

#rnx

diff --git a/Engine/Plugins/Runtime/ReplicationSystemTestPlugin/Source/Private/IrisTestMessageStreamOperators .cpp b/Engine/Plugins/Runtime/ReplicationSystemTestPlugin/Source/Private/IrisTestMessageStreamOperators .cpp
new file mode 100644
--- /dev/null
+++ b/Engine/Plugins/Runtime/ReplicationSystemTestPlugin/Source/Private/IrisTestMessageStreamOperators .cpp	
@@ -0,0 +1,14 @@
+// Copyright Epic Games, Inc. All Rights Reserved.
+
+#include "IrisTestMessageStreamOperators.h"
+#include "Iris/ReplicationSystem/Filtering/NetObjectFilter.h"
+
+namespace UE::Net
+{
+
+FTestMessage& operator<<(FTestMessage& Message, ENetFilterStatus Status)
+{
+	return Message << (Status == ENetFilterStatus::Disallow ? TEXT("ENetFilterStatus::Disallow") : TEXT("ENetFilterStatus::Allow"));
+}
+
+}
diff --git a/Engine/Plugins/Runtime/ReplicationSystemTestPlugin/Source/Private/IrisTestMessageStreamOperators.h b/Engine/Plugins/Runtime/ReplicationSystemTestPlugin/Source/Private/IrisTestMessageStreamOperators.h
new file mode 100644
--- /dev/null
+++ b/Engine/Plugins/Runtime/ReplicationSystemTestPlugin/Source/Private/IrisTestMessageStreamOperators.h
@@ -0,0 +1,18 @@
+// Copyright Epic Games, Inc. All Rights Reserved.
+
+#pragma once
+
+#include "TestMessage.h"
+
+namespace UE::Net
+{
+	enum class ENetFilterStatus : uint32;
+}
+
+
+namespace UE::Net
+{
+
+REPLICATIONSYSTEMTESTPLUGIN_API FTestMessage& operator<<(FTestMessage& Message, ENetFilterStatus);
+
+}
