Iris
- Removed use of forced inlined exports flag from FieldPathNetSerializerTest

diff --git a/Engine/Plugins/Runtime/ReplicationSystemTestPlugin/Source/Private/Tests/Serialization/TestFieldPathNetSerializer.cpp b/Engine/Plugins/Runtime/ReplicationSystemTestPlugin/Source/Private/Tests/Serialization/TestFieldPathNetSerializer.cpp
--- a/Engine/Plugins/Runtime/ReplicationSystemTestPlugin/Source/Private/Tests/Serialization/TestFieldPathNetSerializer.cpp
+++ b/Engine/Plugins/Runtime/ReplicationSystemTestPlugin/Source/Private/Tests/Serialization/TestFieldPathNetSerializer.cpp
@@ -7,6 +7,10 @@
 #include "Iris/Serialization/NetBitStreamWriter.h"
 #include "Iris/Serialization/InternalNetSerializers.h"
 #include "Iris/Serialization/InternalNetSerializationContext.h"
+#include "Iris/Serialization/NetExportContext.h"
+#include "Iris/ReplicationSystem/NetExports.h"
+#include "Iris/ReplicationSystem/ObjectReferenceCache.h"
+#include "Iris/Serialization/NetReferenceCollector.h"
 
 namespace UE::Net
 {
@@ -40,6 +44,7 @@ protected:
 	void Dequantize(FNetSerializationContext& Context, const void* StateBuffer, FFieldPath& Target);
 	bool IsEqual(FNetSerializationContext& Context, const FFieldPath& FieldPath0, const FFieldPath& FieldPath1);
 	void FreeDynamicState(FNetSerializationContext& Context, void* StateBuffer);
+	void CollectAndAppendExports(FNetSerializationContext& Context, void* StateBuffer);
 
 	// Tests
 	void TestQuantize();
@@ -68,6 +73,10 @@ protected:
 	FNetBitStreamReader Reader;
 	FNetBitStreamWriter Writer;
 
+	// Export context
+	UE::Net::Private::FNetExports NetExports;
+	const FNetSerializerChangeMaskParam InitStateChangeMaskInfo = { 0 };
+
 	enum : uint32
 	{
 		BufferSize = 1024,
@@ -109,8 +118,7 @@ void FTestFieldPathNetSerializer::SetUp()
 			TempServerInternalContextInitParams.ObjectResolveContext.RemoteNetTokenStoreState = ConnectionInfo.RemoteNetTokenStoreState;
 			TempServerInternalContextInitParams.ObjectResolveContext.ConnectionId = ConnectionInfo.ConnectionId;
 			TempServerInternalContext.Init(TempServerInternalContextInitParams);
-			// Since we do explicit serialization involving object references we set the context to allow for inlined exports
-			TempServerInternalContext.bInlineObjectReferenceExports = 1U;
+
 			ServerInternalContext = MoveTemp(TempServerInternalContext);
 		}
 
@@ -123,8 +131,6 @@ void FTestFieldPathNetSerializer::SetUp()
 			TempClientInternalContextInitParams.ObjectResolveContext.RemoteNetTokenStoreState = ConnectionInfo.RemoteNetTokenStoreState;
 			TempClientInternalContextInitParams.ObjectResolveContext.ConnectionId = ConnectionInfo.ConnectionId;
 			TempClientInternalContext.Init(TempClientInternalContextInitParams);
-			// Since we do explicit serialization involving object references we set the context to allow for inlined exports
-			TempClientInternalContext.bInlineObjectReferenceExports = 1U;
 
 			ClientInternalContext = MoveTemp(TempClientInternalContext);
 		}
@@ -270,6 +276,30 @@ void FTestFieldPathNetSerializer::FreeDynamicState(FNetSerializationContext& Con
 	FieldPathNetSerializer->FreeDynamicState(Context, Args);
 }
 
+void FTestFieldPathNetSerializer::CollectAndAppendExports(FNetSerializationContext& Context, void* StateBuffer)
+{
+	FNetExportContext* ExportContext = Context.GetExportContext();
+	if (!ExportContext)
+	{
+		return;
+	}
+
+	FNetReferenceCollector Collector(ENetReferenceCollectorTraits::OnlyCollectReferencesThatCanBeExported);
+
+	FNetCollectReferencesArgs Args = {};
+	Args.Version = FieldPathNetSerializer->Version;
+	Args.NetSerializerConfig = NetSerializerConfigParam(&SerializerConfig);
+	Args.Source = NetSerializerValuePointer(StateBuffer);
+	Args.Collector = NetSerializerValuePointer(&Collector);
+
+	FieldPathNetSerializer->CollectNetReferences(Context, Args);
+
+	for (const FNetReferenceCollector::FReferenceInfo& Info : MakeArrayView(Collector.GetCollectedReferences()))
+	{
+		Context.GetInternalContext()->ObjectReferenceCache->AddPendingExport(*ExportContext, Info.Reference);
+	}
+}
+
 void FTestFieldPathNetSerializer::TestQuantize()
 {
 	FNetQuantizeArgs QuantizeArgs = {};
@@ -370,6 +400,10 @@ void FTestFieldPathNetSerializer::TestSerialize()
 	DeserializeArgs.NetSerializerConfig = NetSerializerConfigParam(&SerializerConfig);
 	DeserializeArgs.Target = NetSerializerValuePointer(StateBuffer1);
 
+	// Setup export context, we do not actually need to export anything as we are writing and reading serverside
+	FNetExportContext::FBatchExports CurrentPacketBatchExports;
+	FNetExports::FExportScope ExportScope = NetExports.MakeExportScope(ServerNetSerializationContext, CurrentPacketBatchExports);
+
 	TArrayView<const FFieldPath> AllTheValues[] = { MakeArrayView(TestValues), MakeArrayView(UnresolvableTestValues) };
 	for (const TArrayView<const FFieldPath>& Values : AllTheValues)
 	{
@@ -416,18 +450,37 @@ void FTestFieldPathNetSerializer::TestRoundtripSerialize()
 	{
 		for (const FFieldPath& FieldPath : Values)
 		{
-			Quantize(ServerNetSerializationContext, FieldPath, StateBuffer0);
+			// Setup export scope to capture and serialize exports
+			FNetExportContext::FBatchExports CurrentPacketBatchExports;
+			FNetExports::FExportScope ExportScope = NetExports.MakeExportScope(ServerNetSerializationContext, CurrentPacketBatchExports);
 
-			// Need to send some info to client in order for it to be able to dequantize.
-			Server->SendAndDeliverTo(Client, DeliverPacket);
+			Quantize(ServerNetSerializationContext, FieldPath, StateBuffer0);
 
 			ServerNetSerializationContext.GetBitStreamWriter()->InitBytes(BitStreamBuffer0, sizeof(BitStreamBuffer0));
 
+			// Serialize data
 			FieldPathNetSerializer->Serialize(ServerNetSerializationContext, SerializeArgs);
+			
+			// Collect object references as well
+			CollectAndAppendExports(ServerNetSerializationContext, StateBuffer0);
+
+			// Write exports
+			const uint32 ExportsPos = ServerNetSerializationContext.GetBitStreamWriter()->GetPosBits();
+			ServerNetSerializationContext.GetInternalContext()->ObjectReferenceCache->WritePendingExports(ServerNetSerializationContext, 0);
+			
+			// Finalize
 			ServerNetSerializationContext.GetBitStreamWriter()->CommitWrites();
 			UE_NET_ASSERT_FALSE(ServerNetSerializationContext.HasErrorOrOverflow());
 
+			// Init client read
 			ClientNetSerializationContext.GetBitStreamReader()->InitBits(BitStreamBuffer0, ServerNetSerializationContext.GetBitStreamWriter()->GetPosBits());
+
+			// Read exports
+			ClientNetSerializationContext.GetBitStreamReader()->Seek(ExportsPos);
+			ClientNetSerializationContext.GetInternalContext()->ObjectReferenceCache->ReadExports(ClientNetSerializationContext, nullptr);
+			ClientNetSerializationContext.GetBitStreamReader()->Seek(0U);
+
+			// Deserialize data
 			FieldPathNetSerializer->Deserialize(ClientNetSerializationContext, DeserializeArgs);
 			UE_NET_ASSERT_FALSE(ClientNetSerializationContext.HasErrorOrOverflow());
 
@@ -464,6 +517,10 @@ void FTestFieldPathNetSerializer::TestSerializeDelta()
 	IsEqualArgs.Source1 = NetSerializerValuePointer(StateBuffer2);
 	IsEqualArgs.bStateIsQuantized = true;
 
+	// Setup export context, we do not actually need to export anything as we are writing and reading serverside
+	FNetExportContext::FBatchExports CurrentPacketBatchExports;
+	FNetExports::FExportScope ExportScope = NetExports.MakeExportScope(ServerNetSerializationContext, CurrentPacketBatchExports);
+
 	TArrayView<const FFieldPath> AllTheValues[] = { MakeArrayView(TestValues), MakeArrayView(UnresolvableTestValues) };
 	for (const TArrayView<const FFieldPath>& Values : AllTheValues)
 	{
@@ -478,7 +535,7 @@ void FTestFieldPathNetSerializer::TestSerializeDelta()
 
 				ServerNetSerializationContext.GetBitStreamWriter()->InitBytes(BitStreamBuffer0, sizeof(BitStreamBuffer0));
 
-				FieldPathNetSerializer->SerializeDelta(ServerNetSerializationContext, SerializeDeltaArgs);
+				FieldPathNetSerializer->SerializeDelta(ServerNetSerializationContext, SerializeDeltaArgs);	
 				ServerNetSerializationContext.GetBitStreamWriter()->CommitWrites();
 				UE_NET_ASSERT_FALSE(ServerNetSerializationContext.HasErrorOrOverflow());
 
