Iris - Fixed issue with owner and connection filtered objects and late joining clients.

diff --git a/Engine/Plugins/Runtime/ReplicationSystemTestPlugin/Source/Private/Tests/ReplicationSystem/RPC/TestObjectRPC.cpp b/Engine/Plugins/Runtime/ReplicationSystemTestPlugin/Source/Private/Tests/ReplicationSystem/RPC/TestObjectRPC.cpp
--- a/Engine/Plugins/Runtime/ReplicationSystemTestPlugin/Source/Private/Tests/ReplicationSystem/RPC/TestObjectRPC.cpp
+++ b/Engine/Plugins/Runtime/ReplicationSystemTestPlugin/Source/Private/Tests/ReplicationSystem/RPC/TestObjectRPC.cpp
@@ -2,11 +2,14 @@
 
 #include "RPCTestFixture.h"
 #include "ReplicatedTestObjectWithRPC.h"
+#include "Containers/BitArray.h"
 #include "Iris/Serialization/NetBitStreamReader.h"
 #include "Iris/Serialization/NetBitStreamWriter.h"
 #include "Iris/ReplicationSystem/ReplicationSystem.h"
-
+#include "Iris/ReplicationSystem/Filtering/NetObjectFilter.h"
 #include "Iris/ReplicationSystem/ReplicationSystemInternal.h"
+#include "Tests/EnsureScope.h"
+#include "Tests/CheckScope.h"
 
 namespace UE::Net::Private
 {
@@ -376,4 +379,100 @@ namespace UE::Net::Private
 		Server->UpdateAndSend({Client});
 	}
 
+	// This test is written to verify a bug with owner filtering and late join. There will be an ensure if this test isn't working.
+	UE_NET_TEST_FIXTURE(FRPCTestFixture, TestSubObjectMulticastRPCIsNotReplicatedToNonOwningConnection)
+	{
+		constexpr SIZE_T OwningClientIndex = 0;
+		constexpr SIZE_T NonOwningClientIndex = 1;
+
+		// Add a client
+		FReplicationSystemTestClient* ClientArray[2] = {};
+		ClientArray[OwningClientIndex] = CreateClient();
+
+		// Spawn owner filtered object with on server
+		UTestReplicatedObjectWithRPC* ServerRootObject = Server->CreateObject<UTestReplicatedObjectWithRPC>();
+		ServerRootObject->Init(Server->GetReplicationSystem());
+
+		const FNetRefHandle ServerRootObjectHandle = ServerRootObject->NetRefHandle;
+		
+		// Turn on owner filter
+		Server->ReplicationSystem->SetFilter(ServerRootObjectHandle, ToOwnerFilterHandle);
+		Server->ReplicationSystem->SetOwningNetConnection(ServerRootObjectHandle, ClientArray[OwningClientIndex]->ConnectionIdOnServer);
+
+		UTestReplicatedObjectWithRPC* ServerSubObject = Server->CreateSubObject<UTestReplicatedObjectWithRPC>(ServerRootObjectHandle);
+		ServerSubObject->Init(Server->GetReplicationSystem());
+		ServerSubObject->SetRootObject(ServerRootObject);
+
+		// Send and deliver packet
+		Server->UpdateAndSend({ClientArray[OwningClientIndex]});
+
+		// Late join with second client
+		ClientArray[NonOwningClientIndex] = CreateClient();
+
+		// Call a multicast RPC server->client on the subobject
+		ServerSubObject->NetMulticast_MultiCastRPC();
+
+		// Send and deliver packet
+		{
+			FEnsureScope EnsureScope;
+			FCheckScope CheckScope;
+
+			Server->UpdateAndSend(ClientArray);
+
+			// No ensures or checks should trigger
+			UE_NET_ASSERT_EQ(EnsureScope.GetCount(), 0);
+			UE_NET_ASSERT_EQ(CheckScope.GetCount(), 0);
+		}
+	}
+
+	// This test is written to verify a bug with connection filtering and late join. There will be an ensure if this test isn't working.
+	UE_NET_TEST_FIXTURE(FRPCTestFixture, TestSubObjectMulticastRPCIsNotReplicatedToFilteredOutConnection)
+	{
+		constexpr SIZE_T OwningClientIndex = 0;
+		constexpr SIZE_T NonOwningClientIndex = 1;
+
+		// Add a client
+		FReplicationSystemTestClient* ClientArray[2] = {};
+		ClientArray[OwningClientIndex] = CreateClient();
+
+		// Spawn connection filtered object with on server
+		UTestReplicatedObjectWithRPC* ServerRootObject = Server->CreateObject<UTestReplicatedObjectWithRPC>();
+		ServerRootObject->Init(Server->GetReplicationSystem());
+
+		const FNetRefHandle ServerRootObjectHandle = ServerRootObject->NetRefHandle;
+		
+		// Set connection filter
+		{
+			TBitArray<> AllowedConnections;
+			AllowedConnections.Add(false, ClientArray[OwningClientIndex]->ConnectionIdOnServer + 1);
+			AllowedConnections[ClientArray[OwningClientIndex]->ConnectionIdOnServer] = true;
+			Server->ReplicationSystem->SetConnectionFilter(ServerRootObjectHandle, AllowedConnections, ENetFilterStatus::Allow);
+		}
+	
+		UTestReplicatedObjectWithRPC* ServerSubObject = Server->CreateSubObject<UTestReplicatedObjectWithRPC>(ServerRootObjectHandle);
+		ServerSubObject->Init(Server->GetReplicationSystem());
+		ServerSubObject->SetRootObject(ServerRootObject);
+
+		// Send and deliver packet
+		Server->UpdateAndSend({ClientArray[OwningClientIndex]});
+
+		// Late join with second client
+		ClientArray[NonOwningClientIndex] = CreateClient();
+
+		// Call a multicast RPC server->client on the subobject
+		ServerSubObject->NetMulticast_MultiCastRPC();
+
+		// Send and deliver packet
+		{
+			FEnsureScope EnsureScope;
+			FCheckScope CheckScope;
+
+			Server->UpdateAndSend(ClientArray);
+
+			// No ensures or checks should trigger
+			UE_NET_ASSERT_EQ(EnsureScope.GetCount(), 0);
+			UE_NET_ASSERT_EQ(CheckScope.GetCount(), 0);
+		}
+	}
+
 }
\ No newline at end of file
diff --git a/Engine/Source/Runtime/Experimental/Iris/Core/Private/Iris/ReplicationSystem/Filtering/ReplicationFiltering.cpp b/Engine/Source/Runtime/Experimental/Iris/Core/Private/Iris/ReplicationSystem/Filtering/ReplicationFiltering.cpp
--- a/Engine/Source/Runtime/Experimental/Iris/Core/Private/Iris/ReplicationSystem/Filtering/ReplicationFiltering.cpp
+++ b/Engine/Source/Runtime/Experimental/Iris/Core/Private/Iris/ReplicationSystem/Filtering/ReplicationFiltering.cpp
@@ -805,6 +805,10 @@ void FReplicationFiltering::InitNewConnections()
 			{
 				const bool bIsOwner = (ConnectionId == ObjectIndexToOwningConnection[ObjectIndex]);
 				ConnectionInfo.ConnectionFilteredObjects.SetBitValue(ObjectIndex, bIsOwner);
+				for (const FInternalNetRefIndex SubObjectIndex : NetRefHandleManager->GetSubObjects(ObjectIndex))
+				{
+					ConnectionInfo.ConnectionFilteredObjects.SetBitValue(SubObjectIndex, bIsOwner);
+				}
 			};
 
 			ObjectsWithOwnerFilter.ForAllSetBits(MaskObjectToOwner);
@@ -817,7 +821,12 @@ void FReplicationFiltering::InitNewConnections()
 				const PerObjectInfoIndexType ObjectInfoIndex = ObjectIndexToPerObjectInfoIndex[ObjectIndex];
 				const FPerObjectInfo* ObjectInfo = GetPerObjectInfo(ObjectInfoIndex);
 				const ENetFilterStatus ReplicationStatus = GetConnectionFilterStatus(*ObjectInfo, ConnectionId);
-				ConnectionInfo.ConnectionFilteredObjects.SetBitValue(ObjectIndex, ReplicationStatus == ENetFilterStatus::Allow);
+				const bool bIsAllowedToReplicate = (ReplicationStatus == ENetFilterStatus::Allow);
+				ConnectionInfo.ConnectionFilteredObjects.SetBitValue(ObjectIndex, bIsAllowedToReplicate);
+				for (const FInternalNetRefIndex SubObjectIndex : NetRefHandleManager->GetSubObjects(ObjectIndex))
+				{
+					ConnectionInfo.ConnectionFilteredObjects.SetBitValue(SubObjectIndex, bIsAllowedToReplicate);
+				}
 			};
 
 			ObjectsWithPerObjectInfo.ForAllSetBits(MaskObjectToConnection);
