Fixed lowleveltest failing when testing an expected failure.

diff --git a/Engine/Plugins/Runtime/ReplicationSystemTestPlugin/Source/Private/Tests/ReplicationSystem/TestReplicationSystem.cpp b/Engine/Plugins/Runtime/ReplicationSystemTestPlugin/Source/Private/Tests/ReplicationSystem/TestReplicationSystem.cpp
--- a/Engine/Plugins/Runtime/ReplicationSystemTestPlugin/Source/Private/Tests/ReplicationSystem/TestReplicationSystem.cpp
+++ b/Engine/Plugins/Runtime/ReplicationSystemTestPlugin/Source/Private/Tests/ReplicationSystem/TestReplicationSystem.cpp
@@ -9,6 +9,7 @@
 #include "Iris/Core/IrisLog.h"
 #include "Iris/Metrics/NetMetrics.h"
 #include "Misc/ScopeExit.h"
+#include "Logging/LogScopedVerbosityOverride.h"
 
 namespace UE::Net::Private
 {
@@ -534,8 +535,12 @@ UE_NET_TEST_FIXTURE(FTestNetTokensFixture, NetToken)
 	Server->SendAndDeliverTo(Client, false);
 	Server->PostSendUpdate();
 
-	// Verify that we cannot resolve the token on the client
-	UE_NET_ASSERT_NE(TokenStringA, FString(ClientStringTokenStore->ResolveRemoteToken(StringTokenA, *ClientRemoteNetTokenState)));
+	{
+		LOG_SCOPE_VERBOSITY_OVERRIDE(LogNetToken, ELogVerbosity::Fatal);
+
+		// Verify that we cannot resolve the token on the client
+		UE_NET_ASSERT_NE(TokenStringA, FString(ClientStringTokenStore->ResolveRemoteToken(StringTokenA, *ClientRemoteNetTokenState)));
+	}
 
 	// Send and deliver packet
 	Server->PreSendUpdate();
@@ -571,7 +576,10 @@ UE_NET_TEST_FIXTURE(FTestNetTokensFixture, NetTokenResendWithFullPacket)
 
 	// Verify that we can resolve the token first token on the client even though second one should not fit
 	UE_NET_ASSERT_EQ(TokenStringA, FString(ClientStringTokenStore->ResolveRemoteToken(StringTokenA, *ClientRemoteNetTokenState)));
-	UE_NET_ASSERT_NE(TokenStringB, FString(ClientStringTokenStore->ResolveRemoteToken(StringTokenB, *ClientRemoteNetTokenState)));
+	{
+		LOG_SCOPE_VERBOSITY_OVERRIDE(LogNetToken, ELogVerbosity::Fatal);
+		UE_NET_ASSERT_NE(TokenStringB, FString(ClientStringTokenStore->ResolveRemoteToken(StringTokenB, *ClientRemoteNetTokenState)));
+	}
 
 	// Restore packet size and make sure that we get the second token through
 	Server->SetMaxSendPacketSize(1024U);
@@ -608,8 +616,12 @@ UE_NET_TEST_FIXTURE(FTestNetTokensFixture, NetTokenResendWithFullPacketAfterFirs
 	Server->DeliverTo(Client, false);
 
 	// Verify that tokens has not been received
-	UE_NET_ASSERT_NE(TestStringA, FString(ClientStringTokenStore->ResolveRemoteToken(StringTokenA, *ClientRemoteNetTokenState)));
-	UE_NET_ASSERT_NE(TestStringB, FString(ClientStringTokenStore->ResolveRemoteToken(StringTokenB, *ClientRemoteNetTokenState)));
+	{
+		LOG_SCOPE_VERBOSITY_OVERRIDE(LogNetToken, ELogVerbosity::Fatal);
+
+		UE_NET_ASSERT_NE(TestStringA, FString(ClientStringTokenStore->ResolveRemoteToken(StringTokenA, *ClientRemoteNetTokenState)));
+		UE_NET_ASSERT_NE(TestStringB, FString(ClientStringTokenStore->ResolveRemoteToken(StringTokenB, *ClientRemoteNetTokenState)));
+	}
 
 	// Send and deliver packet which now should contain two entries in the resend queue
 	Server->SetMaxSendPacketSize(1024);
