Iris
* Implementing RegisterReplicationFragments in a replicated class is no longer mandatory.
* Now we detect if the user implemented is own fragment definition and if not we call FReplicationFragmentUtil::CreateAndRegisterFragmentsForObject to do it automatically for him.

UE-225675

diff --git a/Engine/Plugins/Runtime/ReplicationSystemTestPlugin/Source/Private/Tests/ReplicationState/TestPropertyReplicationState.cpp b/Engine/Plugins/Runtime/ReplicationSystemTestPlugin/Source/Private/Tests/ReplicationState/TestPropertyReplicationState.cpp
--- a/Engine/Plugins/Runtime/ReplicationSystemTestPlugin/Source/Private/Tests/ReplicationState/TestPropertyReplicationState.cpp
+++ b/Engine/Plugins/Runtime/ReplicationSystemTestPlugin/Source/Private/Tests/ReplicationState/TestPropertyReplicationState.cpp
@@ -62,6 +62,11 @@ void UTestPropertyReplicationState_TestClassWithTArray::OnRep_ReferencedObjects(
 	bOnRepWasCalled = true;
 }
 
+void UTestPropertyReplicationState_NoRegisterFragments::GetLifetimeReplicatedProps(TArray<class FLifetimeProperty> & OutLifetimeProps) const
+{
+	DOREPLIFETIME(ThisClass, IntA);
+}
+
 namespace UE::Net::Private
 {
 
@@ -623,4 +628,33 @@ UE_NET_TEST_FIXTURE(FTestPropertyReplicationStateContext, TestArrayOnRepWithUnre
 	}
 }
 
+// Test that automatic fragment registration works for any class that doesn't implement RegisterReplicationFragments
+UE_NET_TEST_FIXTURE(FTestPropertyReplicationStateContext, TestClassDefaultRegisterFragments)
+{
+	using namespace UE::Net;
+
+	// Add a client
+	FReplicationSystemTestClient* Client = CreateClient();
+
+	// Spawn object
+	UTestPropertyReplicationState_NoRegisterFragments* ServerObject = Server->CreateObject<UTestPropertyReplicationState_NoRegisterFragments>();
+
+	// Create replica on client
+	Server->UpdateAndSend({ Client });
+
+	// Find the replica
+	const FNetRefHandle ServerHandle = Server->GetReplicationBridge()->GetReplicatedRefHandle(ServerObject);
+	UTestPropertyReplicationState_NoRegisterFragments* ClientObject = Cast<UTestPropertyReplicationState_NoRegisterFragments>(Client->GetReplicationBridge()->GetReplicatedObject(ServerHandle));
+	UE_NET_ASSERT_NE(ClientObject, nullptr);
+
+	// Dirty a replicated property
+	ServerObject->IntA = 0xAA;
+
+	// Update it for the client
+	Server->UpdateAndSend({ Client });
+
+	// Make sure it was received
+	UE_NET_ASSERT_EQ(ClientObject->IntA, 0xAA);
+}
+
 }
diff --git a/Engine/Plugins/Runtime/ReplicationSystemTestPlugin/Source/Private/Tests/ReplicationState/TestPropertyReplicationState.h b/Engine/Plugins/Runtime/ReplicationSystemTestPlugin/Source/Private/Tests/ReplicationState/TestPropertyReplicationState.h
--- a/Engine/Plugins/Runtime/ReplicationSystemTestPlugin/Source/Private/Tests/ReplicationState/TestPropertyReplicationState.h
+++ b/Engine/Plugins/Runtime/ReplicationSystemTestPlugin/Source/Private/Tests/ReplicationState/TestPropertyReplicationState.h
@@ -144,3 +144,15 @@ protected:
 	void OnRep_ReferencedObjects();
 
 };
+
+
+UCLASS()
+class UTestPropertyReplicationState_NoRegisterFragments : public UReplicatedTestObject
+{
+	GENERATED_BODY()
+
+public:
+
+	UPROPERTY(Transient, Replicated)
+	int IntA = 0;
+};
\ No newline at end of file
diff --git a/Engine/Source/Runtime/Experimental/Iris/Core/Private/Iris/ReplicationSystem/ObjectReplicationBridge.cpp b/Engine/Source/Runtime/Experimental/Iris/Core/Private/Iris/ReplicationSystem/ObjectReplicationBridge.cpp
--- a/Engine/Source/Runtime/Experimental/Iris/Core/Private/Iris/ReplicationSystem/ObjectReplicationBridge.cpp
+++ b/Engine/Source/Runtime/Experimental/Iris/Core/Private/Iris/ReplicationSystem/ObjectReplicationBridge.cpp
@@ -28,6 +28,7 @@
 #include "Iris/ReplicationSystem/ReplicationConnections.h"
 #include "Iris/ReplicationSystem/ReplicationFragment.h"
 #include "Iris/ReplicationSystem/ReplicationFragmentInternal.h"
+#include "Iris/ReplicationSystem/ReplicationFragmentUtil.h"
 #include "Iris/ReplicationState/ReplicationStateDescriptorBuilder.h"
 #include "Iris/ReplicationState/ReplicationStateUtil.h"
 #include "Iris/ReplicationSystem/ReplicationSystem.h"
@@ -114,9 +115,13 @@ namespace UE::Net::Private
 {
 	void CallRegisterReplicationFragments(UObject* Object, FFragmentRegistrationContext& Context, EFragmentRegistrationFlags RegistrationFlags)
 	{
-#if UE_WITH_IRIS
 		Object->RegisterReplicationFragments(Context, RegistrationFlags);
-#endif
+
+		if (!Context.WasRegistered())
+		{
+			FReplicationFragmentUtil::CreateAndRegisterFragmentsForObject(Object, Context, RegistrationFlags);
+			ensure(Context.WasRegistered());
+		}
 	}
 
 	struct FReplicationInstanceProtocolDeleter
diff --git a/Engine/Source/Runtime/Experimental/Iris/Core/Public/Iris/ReplicationSystem/ReplicationFragment.h b/Engine/Source/Runtime/Experimental/Iris/Core/Public/Iris/ReplicationSystem/ReplicationFragment.h
--- a/Engine/Source/Runtime/Experimental/Iris/Core/Public/Iris/ReplicationSystem/ReplicationFragment.h
+++ b/Engine/Source/Runtime/Experimental/Iris/Core/Public/Iris/ReplicationSystem/ReplicationFragment.h
@@ -251,6 +251,12 @@ public:
 	/** Returns true when the netobject knows it won't contain any replicated properties or RPCs */
 	bool IsFragmentlessNetObject() const { return bIsAFragmentlessNetObject; }
 
+	/** Returns true if the fragments (or the lack of) were registered by the instance */
+	bool WasRegistered() const
+	{
+		return Fragments.Num() > 0 || bIsAFragmentlessNetObject;
+	}
+
 	/** Returns the number of fragments registered */
 	int32 NumFragments() const { return Fragments.Num(); }
 
