Iris
- Fixed UHT to properly create bindings for derived fastarrays.

diff --git a/Engine/Plugins/Runtime/ReplicationSystemTestPlugin/Source/Private/Tests/ReplicationState/TestFastArrayReplicationState.cpp b/Engine/Plugins/Runtime/ReplicationSystemTestPlugin/Source/Private/Tests/ReplicationState/TestFastArrayReplicationState.cpp
--- a/Engine/Plugins/Runtime/ReplicationSystemTestPlugin/Source/Private/Tests/ReplicationState/TestFastArrayReplicationState.cpp
+++ b/Engine/Plugins/Runtime/ReplicationSystemTestPlugin/Source/Private/Tests/ReplicationState/TestFastArrayReplicationState.cpp
@@ -827,6 +827,7 @@ UE_NET_TEST_FIXTURE(FReplicationSystemServerClientTestFixture, TestPostReplicate
 
 	// Verify that we hit the callback
 	UE_NET_ASSERT_TRUE(ClientFastArray.bHitPostReplicatedReceive);
+	UE_NET_ASSERT_TRUE(ClientFastArray.bHitDerivedPostReplicatedReceive);
 	UE_NET_ASSERT_TRUE(ClientFastArray.bPostReplicatedReceiveWasHitWithUnresolvedReferences);
 
 	// Reset
diff --git a/Engine/Plugins/Runtime/ReplicationSystemTestPlugin/Source/Private/Tests/ReplicationState/TestFastArrayReplicationState.h b/Engine/Plugins/Runtime/ReplicationSystemTestPlugin/Source/Private/Tests/ReplicationState/TestFastArrayReplicationState.h
--- a/Engine/Plugins/Runtime/ReplicationSystemTestPlugin/Source/Private/Tests/ReplicationState/TestFastArrayReplicationState.h
+++ b/Engine/Plugins/Runtime/ReplicationSystemTestPlugin/Source/Private/Tests/ReplicationState/TestFastArrayReplicationState.h
@@ -85,6 +85,8 @@ struct FTestFastArrayReplicationState_FastArray : public FTestFastArrayReplicati
 {
 	GENERATED_BODY()
 
+	bool bHitDerivedPostReplicatedReceive = false;
+
 	FTestFastArrayReplicationState_FastArray() : FTestFastArrayReplicationState_FastArraySerializer()
 	{
 	}
@@ -93,6 +95,12 @@ struct FTestFastArrayReplicationState_FastArray : public FTestFastArrayReplicati
 	{
 		return FFastArraySerializer::FastArrayDeltaSerialize<FTestFastArrayReplicationState_FastArrayItem, FTestFastArrayReplicationState_FastArray>(Items, DeltaParms, *this);
 	}
+
+	void PostReplicatedReceive(const FFastArraySerializer::FPostReplicatedReceiveParameters& Parameters)
+	{
+		Super::PostReplicatedReceive(Parameters);
+		bHitDerivedPostReplicatedReceive = true;
+	}
 };
 
 template<>
diff --git a/Engine/Source/Programs/Shared/EpicGames.UHT/Exporters/CodeGen/UhtCodeGenerator.cs b/Engine/Source/Programs/Shared/EpicGames.UHT/Exporters/CodeGen/UhtCodeGenerator.cs
--- a/Engine/Source/Programs/Shared/EpicGames.UHT/Exporters/CodeGen/UhtCodeGenerator.cs
+++ b/Engine/Source/Programs/Shared/EpicGames.UHT/Exporters/CodeGen/UhtCodeGenerator.cs
@@ -352,6 +352,13 @@ namespace EpicGames.UHT.Exporters.CodeGen
 				// To fulfill that requirement the struct should be derived from FFastArraySerializer and have a single replicated TArrayProperty
 				if (scriptStructObj.IsChildOf(FastArraySerializer))
 				{
+					// If Super is a valid fastarray we mark this struct as a FastArrayProperty as well
+					if (scriptStructObj.Super != null && ObjectInfos[scriptStructObj.Super.ObjectTypeIndex].FastArrayProperty != null)
+					{
+						objectInfo.FastArrayProperty = ObjectInfos[scriptStructObj.Super.ObjectTypeIndex].FastArrayProperty;
+					}
+
+					// A valid fastarray cannot have any additional replicated properties.
 					foreach (UhtType child in scriptStructObj.Children)
 					{
 						if (child is UhtProperty property)
