Added ForAllExclusiveBitsByPredicate to FNetBitArray and FNetBitArrayView. This makes it easy to iterate over the same pair of bitarrays multiple times using different conditions for when to invoke the exclusive bit functors.

diff --git a/Engine/Plugins/Runtime/ReplicationSystemTestPlugin/Source/Private/Tests/Core/TestNetBitArray.cpp b/Engine/Plugins/Runtime/ReplicationSystemTestPlugin/Source/Private/Tests/Core/TestNetBitArray.cpp
--- a/Engine/Plugins/Runtime/ReplicationSystemTestPlugin/Source/Private/Tests/Core/TestNetBitArray.cpp
+++ b/Engine/Plugins/Runtime/ReplicationSystemTestPlugin/Source/Private/Tests/Core/TestNetBitArray.cpp
@@ -1698,6 +1698,128 @@ UE_NET_TEST_FIXTURE(FNetBitArrayViewFixture, TestForAllExclusiveBits)
 	}
 }
 
+UE_NET_TEST_FIXTURE(FNetBitArrayViewFixture, TestForAllExclusiveBitsByPredicate)
+{
+	// Test predicate that returns true which should behave identically to not having a predicate.
+	{
+		uint32 ExpectedOnlyInABits = 0x0000aaaa;
+		uint32 ExpectedOnlyInBBits = 0x00bb0000;
+		uint32 ExpectedBitsSetInAB = ~(ExpectedOnlyInABits ^ ExpectedOnlyInBBits) & 0x00ffffff;
+
+		FNetBitArrayView BitArrayA(&ExpectedOnlyInABits, 24);
+		FNetBitArrayView BitArrayB(&ExpectedOnlyInBBits, 24);
+		FNetBitArrayView BitArrayAB(&ExpectedBitsSetInAB, 24);
+
+		BitArrayA.Combine(BitArrayAB, FNetBitArrayView::OrOp);
+		BitArrayB.Combine(BitArrayAB, FNetBitArrayView::OrOp);
+
+		FCollectSetBitsFunctor ACollector, BCollector;
+		constexpr auto AllowAll = [](uint32)
+			{
+				return true;
+			};
+		FNetBitArrayView::ForAllExclusiveBitsByPredicate(BitArrayA, BitArrayB, ACollector, BCollector, AllowAll);
+
+		// Verify result
+		VerifyResult(ACollector, BitArrayA, BitArrayB);
+		VerifyResult(BCollector, BitArrayB, BitArrayA);
+	}
+
+	// Test predicate that returns false which should cause no functor to be called.
+	{
+		uint32 ExpectedOnlyInABits = 0x0000aaaa;
+		uint32 ExpectedOnlyInBBits = 0x00bb0000;
+		uint32 ExpectedBitsSetInAB = ~(ExpectedOnlyInABits ^ ExpectedOnlyInBBits) & 0x00ffffff;
+
+		FNetBitArrayView BitArrayA(&ExpectedOnlyInABits, 24);
+		FNetBitArrayView BitArrayB(&ExpectedOnlyInBBits, 24);
+		FNetBitArrayView BitArrayAB(&ExpectedBitsSetInAB, 24);
+
+		BitArrayA.Combine(BitArrayAB, FNetBitArrayView::OrOp);
+		BitArrayB.Combine(BitArrayAB, FNetBitArrayView::OrOp);
+
+		FCollectSetBitsFunctor ACollector, BCollector;
+		constexpr auto AllowNone = [](uint32)
+			{
+				return false;
+			};
+		FNetBitArrayView::ForAllExclusiveBitsByPredicate(BitArrayA, BitArrayB, ACollector, BCollector, AllowNone);
+
+		UE_NET_ASSERT_TRUE(ACollector.Invoked.IsEmpty());
+		UE_NET_ASSERT_TRUE(BCollector.Invoked.IsEmpty());
+	}
+
+	// Test predicate that only allows odd bit indices.
+	{
+		uint32 ExpectedOnlyInABits = 0x0000aaaa;
+		uint32 ExpectedOnlyInBBits = 0x00bb0000;
+		uint32 ExpectedBitsSetInAB = ~(ExpectedOnlyInABits ^ ExpectedOnlyInBBits) & 0x00ffffff;
+
+		FNetBitArrayView BitArrayA(&ExpectedOnlyInABits, 24);
+		FNetBitArrayView BitArrayB(&ExpectedOnlyInBBits, 24);
+		FNetBitArrayView BitArrayAB(&ExpectedBitsSetInAB, 24);
+
+		BitArrayA.Combine(BitArrayAB, FNetBitArrayView::OrOp);
+		BitArrayB.Combine(BitArrayAB, FNetBitArrayView::OrOp);
+
+		FCollectSetBitsFunctor ACollector, BCollector;
+		auto AllowOdd = [](uint32 BitIndex)
+			{
+				return (BitIndex & 1U) == 1U;
+			};
+		FNetBitArrayView::ForAllExclusiveBitsByPredicate(BitArrayA, BitArrayB, ACollector, BCollector, AllowOdd);
+
+		// Without odd bits being exclusive to one bit array this test is useless.
+		UE_NET_ASSERT_GT(ACollector.Invoked.Num(), 0);
+		UE_NET_ASSERT_GT(BCollector.Invoked.Num(), 0);
+
+		for (const uint32 BitIndex : ACollector.Invoked)
+		{
+			UE_NET_ASSERT_EQ(BitIndex & 1U, 1U);
+		}
+
+		for (const uint32 BitIndex : BCollector.Invoked)
+		{
+			UE_NET_ASSERT_EQ(BitIndex & 1U, 1U);
+		}
+	}
+
+	// Test predicate that only allows even bit indices.
+	{
+		uint32 ExpectedOnlyInABits = 0x0000aabb;
+		uint32 ExpectedOnlyInBBits = 0x00bb0004;
+		uint32 ExpectedBitsSetInAB = ~(ExpectedOnlyInABits ^ ExpectedOnlyInBBits) & 0x00ffffff;
+
+		FNetBitArrayView BitArrayA(&ExpectedOnlyInABits, 24);
+		FNetBitArrayView BitArrayB(&ExpectedOnlyInBBits, 24);
+		FNetBitArrayView BitArrayAB(&ExpectedBitsSetInAB, 24);
+
+		BitArrayA.Combine(BitArrayAB, FNetBitArrayView::OrOp);
+		BitArrayB.Combine(BitArrayAB, FNetBitArrayView::OrOp);
+
+		FCollectSetBitsFunctor ACollector, BCollector;
+		auto AllowEven = [](uint32 BitIndex)
+			{
+				return (BitIndex & 1U) == 0U;
+			};
+		FNetBitArrayView::ForAllExclusiveBitsByPredicate(BitArrayA, BitArrayB, ACollector, BCollector, AllowEven);
+
+		// Without even bits being exclusive to one bit array this test is useless.
+		UE_NET_ASSERT_GT(ACollector.Invoked.Num(), 0);
+		UE_NET_ASSERT_GT(BCollector.Invoked.Num(), 0);
+
+		for (const uint32 BitIndex : ACollector.Invoked)
+		{
+			UE_NET_ASSERT_EQ(BitIndex & 1U, 0U);
+		}
+
+		for (const uint32 BitIndex : BCollector.Invoked)
+		{
+			UE_NET_ASSERT_EQ(BitIndex & 1U, 0U);
+		}
+	}
+}
+
 class FNetBitArrayViewTestForAllSetBitsFixture : public FNetworkAutomationTestSuiteFixture
 {
 public:
diff --git a/Engine/Source/Runtime/Net/Core/Public/Net/Core/NetBitArray.h b/Engine/Source/Runtime/Net/Core/Public/Net/Core/NetBitArray.h
--- a/Engine/Source/Runtime/Net/Core/Public/Net/Core/NetBitArray.h
+++ b/Engine/Source/Runtime/Net/Core/Public/Net/Core/NetBitArray.h
@@ -287,6 +287,10 @@ public:
 	template<typename T, typename V>
 	static void ForAllExclusiveBits(const FNetBitArray& A, const FNetBitArray& B, T&& FunctorA, V&& FunctorB);
 
+	/** Compare two BitArrays bit by bit. Where they differ and Predicate(BitIndex) returns true FunctorA(BitIndex) is invoked for bits set only in array A and FunctorB(BitIndex) is invoked for bits set only in array B. */
+	template<typename T, typename V, typename Predicate>
+	static void ForAllExclusiveBitsByPredicate(const FNetBitArray& A, const FNetBitArray& B, T&& FunctorA, V&& FunctorB, Predicate&& Pred);
+
 	/**
 	 * DO NOT USE DIRECTLY
 	 * Range-based for loop support iterating over set bits and returning the index to the set bit.
@@ -501,6 +505,10 @@ public:
 	template <typename T, typename V>
 	static void ForAllExclusiveBits(const FNetBitArrayView& A, const FNetBitArrayView& B, T&& FunctorA, V&& FunctorB);
 
+	/** Compare two BitArrays bit by bit. Where they differ and Predicate(BitIndex) returns true FunctorA(BitIndex) is invoked for bits set only in array A and FunctorB(BitIndex) is invoked for bits set only in array B. */
+	template<typename T, typename V, typename Predicate>
+	static void ForAllExclusiveBitsByPredicate(const FNetBitArrayView& A, const FNetBitArrayView& B, T&& FunctorA, V&& FunctorB, Predicate&& Pred);
+
 	static constexpr inline uint32 CalculateRequiredWordCount(uint32 BitCount);
 
 	/**
@@ -865,7 +873,7 @@ private:
 		}
 	}
 
-	template<typename T, typename V> static void ForAllExclusiveBits(const StorageWordType* StorageA, const StorageWordType* StorageB, const uint32 WordCount, const uint32 BitCount, T&& FunctorA, V&& FunctorB)
+	template<typename T, typename V, typename Predicate> static void ForAllExclusiveBitsByPredicate(const StorageWordType* StorageA, const StorageWordType* StorageB, const uint32 WordCount, const uint32 BitCount, T&& FunctorA, V&& FunctorB, Predicate&& Pred)
 	{
 		const StorageWordType LastWordMask = ~StorageWordType(0) >> (uint32(-int32(BitCount)) & (WordBitCount - 1));
 		const uint32 LastWordIt = WordCount - 1;
@@ -882,13 +890,16 @@ private:
 			{
 				if (CurrentWordXOR & 0x1)
 				{
-					if (CurrentWordA & 0x1)
-					{
-						FunctorA(LocalBitIndex);
-					}
-					else
+					if (Pred(LocalBitIndex))
 					{
-						FunctorB(LocalBitIndex);
+						if (CurrentWordA & 0x1)
+						{
+							FunctorA(LocalBitIndex);
+						}
+						else
+						{
+							FunctorB(LocalBitIndex);
+						}
 					}
 				}
 
@@ -900,6 +911,16 @@ private:
 		}
 	}
 
+	template<typename T, typename V> static void ForAllExclusiveBits(const StorageWordType* StorageA, const StorageWordType* StorageB, const uint32 WordCount, const uint32 BitCount, T&& FunctorA, V&& FunctorB)
+	{
+		constexpr auto AlwaysProcessPredicate = [](uint32 Index)
+			{ 
+				return true; 
+			};
+		
+		return ForAllExclusiveBitsByPredicate(StorageA, StorageB, WordCount, BitCount, FunctorA, FunctorB, AlwaysProcessPredicate);
+	}
+
 };
 
 namespace Private
@@ -1188,6 +1209,12 @@ template <typename T, typename V> inline void FNetBitArray::ForAllExclusiveBits(
 	FNetBitArrayHelper::ForAllExclusiveBits(A.Storage.GetData(), B.Storage.GetData(), A.Storage.Num(), A.BitCount, FunctorA, FunctorB);
 }
 
+template <typename T, typename V, typename Predicate> inline void FNetBitArray::ForAllExclusiveBitsByPredicate(const FNetBitArray& A, const FNetBitArray& B, T&& FunctorA, V&& FunctorB, Predicate&& Pred)
+{
+	UE_NETBITARRAY_CHECK(A.GetNumBits() == B.GetNumBits());
+	FNetBitArrayHelper::ForAllExclusiveBitsByPredicate(A.Storage.GetData(), B.Storage.GetData(), A.Storage.Num(), A.BitCount, FunctorA, FunctorB, Pred);
+}
+
 inline void FNetBitArray::ClearPaddingBits()
 {
 	if (BitCount > 0)
@@ -1431,6 +1458,13 @@ void FNetBitArrayView::ForAllExclusiveBits(const FNetBitArrayView& A, const FNet
 	FNetBitArrayHelper::ForAllExclusiveBits(A.Storage, B.Storage, A.WordCount, A.BitCount, FunctorA, FunctorB);
 }
 
+template <typename T, typename V, typename Predicate>
+void FNetBitArrayView::ForAllExclusiveBitsByPredicate(const FNetBitArrayView& A, const FNetBitArrayView& B, T&& FunctorA, V&& FunctorB, Predicate&& Pred)
+{
+	UE_NETBITARRAY_CHECK(A.GetNumBits() == B.GetNumBits());
+	FNetBitArrayHelper::ForAllExclusiveBitsByPredicate(A.Storage, B.Storage, A.WordCount, A.BitCount, FunctorA, FunctorB, Pred);
+}
+
 inline Private::FNetBitArrayRangedForConstIterator FNetBitArrayView::begin() const
 {
 	return Private::FNetBitArrayRangedForConstIterator(Storage, BitCount, Private::FNetBitArrayRangedForConstIterator::ERangeStart::Begin);
