Iris - Fix ReplicationSystemLowLevelTests compile error.

#rnx

diff --git a/Engine/Plugins/Runtime/ReplicationSystemTestPlugin/Source/Private/Tests/Core/TestNetBitArray.cpp b/Engine/Plugins/Runtime/ReplicationSystemTestPlugin/Source/Private/Tests/Core/TestNetBitArray.cpp
--- a/Engine/Plugins/Runtime/ReplicationSystemTestPlugin/Source/Private/Tests/Core/TestNetBitArray.cpp
+++ b/Engine/Plugins/Runtime/ReplicationSystemTestPlugin/Source/Private/Tests/Core/TestNetBitArray.cpp
@@ -1775,12 +1775,12 @@ UE_NET_TEST_FIXTURE(FNetBitArrayViewFixture, TestForAllExclusiveBitsByPredicate)
 
 		for (const uint32 BitIndex : ACollector.Invoked)
 		{
-			UE_NET_ASSERT_EQ(BitIndex & 1U, 1U);
+			UE_NET_ASSERT_EQ((BitIndex & 1U), 1U);
 		}
 
 		for (const uint32 BitIndex : BCollector.Invoked)
 		{
-			UE_NET_ASSERT_EQ(BitIndex & 1U, 1U);
+			UE_NET_ASSERT_EQ((BitIndex & 1U), 1U);
 		}
 	}
 
@@ -1810,12 +1810,12 @@ UE_NET_TEST_FIXTURE(FNetBitArrayViewFixture, TestForAllExclusiveBitsByPredicate)
 
 		for (const uint32 BitIndex : ACollector.Invoked)
 		{
-			UE_NET_ASSERT_EQ(BitIndex & 1U, 0U);
+			UE_NET_ASSERT_EQ((BitIndex & 1U), 0U);
 		}
 
 		for (const uint32 BitIndex : BCollector.Invoked)
 		{
-			UE_NET_ASSERT_EQ(BitIndex & 1U, 0U);
+			UE_NET_ASSERT_EQ((BitIndex & 1U), 0U);
 		}
 	}
 }
