Iris - FReplicationWriter::UpdateScope will now process root objects in a first pass and subobjects in a second pass.

diff --git a/Engine/Plugins/Runtime/ReplicationSystemTestPlugin/Source/Private/Tests/ReplicationSystem/TestCancelPendingDestroy.cpp b/Engine/Plugins/Runtime/ReplicationSystemTestPlugin/Source/Private/Tests/ReplicationSystem/TestCancelPendingDestroy.cpp
--- a/Engine/Plugins/Runtime/ReplicationSystemTestPlugin/Source/Private/Tests/ReplicationSystem/TestCancelPendingDestroy.cpp
+++ b/Engine/Plugins/Runtime/ReplicationSystemTestPlugin/Source/Private/Tests/ReplicationSystem/TestCancelPendingDestroy.cpp
@@ -3,7 +3,7 @@
 #include "ReplicationSystemServerClientTestFixture.h"
 #include "Iris/ReplicationSystem/ReplicationSystem.h"
 #include "Iris/ReplicationSystem/Filtering/NetObjectFilter.h"
-
+#include "Tests/EnsureScope.h"
 
 namespace UE::Net::Private
 {
@@ -1021,6 +1021,52 @@ UE_NET_TEST_FIXTURE(FTestCancelPendingDestroyFixture, TestCancelRootObjectDestro
 	UE_NET_ASSERT_EQ(Client->GetReplicationBridge()->GetReplicatedObject(ServerSubObject->NetRefHandle), nullptr);
 }
 
+UE_NET_TEST_FIXTURE(FTestCancelPendingDestroyFixture, TestCancelPendingDestroyDuringWaitOnDestroyConfirmationDoesNotTriggerEnsure)
+{
+	// Add a client
+	FReplicationSystemTestClient* Client = CreateClient();
+
+	// Spawn temp object on server to reserver low internal index
+	UTestReplicatedIrisObject* ServerTempObject = Server->CreateObject(UTestReplicatedIrisObject::FComponents());
+	
+	// Spawn object on server
+	UTestReplicatedIrisObject* ServerObject = Server->CreateObject(UTestReplicatedIrisObject::FComponents());
+	
+	// Destroy temp object and perfom an update tick to allow the internal index to be reused
+	{
+		Server->DestroyObject(ServerTempObject);
+		ServerTempObject = nullptr;
+		Server->UpdateAndSend({});
+	}
+
+	// Write and send packet
+	Server->UpdateAndSend({Client});
+
+	// Filter out object and write packet to cause object to end up in WaitOnDestroyConfirmation
+	FNetObjectGroupHandle ExclusionGroupHandle = Server->GetReplicationSystem()->CreateGroup(NAME_None);
+	Server->ReplicationSystem->AddToGroup(ExclusionGroupHandle, ServerObject->NetRefHandle);
+	Server->ReplicationSystem->AddExclusionFilterGroup(ExclusionGroupHandle);
+	Server->ReplicationSystem->SetGroupFilterStatus(ExclusionGroupHandle, Client->ConnectionIdOnServer, ENetFilterStatus::Disallow);
+
+	Server->PreSendUpdate();
+	Server->SendUpdate(Client->ConnectionIdOnServer);
+	Server->PostSendUpdate();
+
+	// Create subobject on server
+	UTestReplicatedIrisObject* ServerSubObject = Server->CreateSubObject(ServerObject->NetRefHandle, UTestReplicatedIrisObject::FComponents());
+
+	// Allow replication again to have the object end up in CancelPendingDestroy. Assert that no ensure was triggered.
+	{
+		FEnsureScope EnsureScope;
+
+		Server->ReplicationSystem->SetGroupFilterStatus(ExclusionGroupHandle, Client->ConnectionIdOnServer, ENetFilterStatus::Allow);
+		Server->PreSendUpdate();
+		Server->PostSendUpdate();
+
+		UE_NET_ASSERT_EQ(EnsureScope.GetCount(), 0);
+	}
+}
+
 // See TestObjectSplitting.cpp for cancel pending destroy on huge object
 
 }
diff --git a/Engine/Source/Runtime/Experimental/Iris/Core/Private/Iris/ReplicationSystem/ReplicationWriter.cpp b/Engine/Source/Runtime/Experimental/Iris/Core/Private/Iris/ReplicationSystem/ReplicationWriter.cpp
--- a/Engine/Source/Runtime/Experimental/Iris/Core/Private/Iris/ReplicationSystem/ReplicationWriter.cpp
+++ b/Engine/Source/Runtime/Experimental/Iris/Core/Private/Iris/ReplicationSystem/ReplicationWriter.cpp
@@ -515,7 +515,7 @@ void FReplicationWriter::StartReplication(uint32 InternalIndex)
 		if (ensure(RootObjectInternalIndex != FNetRefHandleManager::InvalidInternalIndex))
 		{
 			FReplicationInfo& OwnerInfo = ReplicatedObjects[RootObjectInternalIndex];
-			if (OwnerInfo.GetState() != EReplicatedObjectState::Invalid && ensureMsgf(OwnerInfo.GetState() < EReplicatedObjectState::PendingDestroy, TEXT("Unsupported state %s"), LexToString(OwnerInfo.GetState())))
+			if (ensureMsgf(OwnerInfo.GetState() != EReplicatedObjectState::Invalid && OwnerInfo.GetState() < EReplicatedObjectState::PendingDestroy, TEXT("Unsupported state %s"), LexToString(OwnerInfo.GetState())))
 			{
 				ObjectsWithDirtyChanges.SetBit(RootObjectInternalIndex);
 				OwnerInfo.HasDirtySubObjects = 1U;
@@ -846,7 +846,22 @@ void FReplicationWriter::UpdateScope(const FNetBitArrayView& UpdatedScope)
 	};
 
 	FNetBitArrayView CurrentScope = MakeNetBitArrayView(ObjectsInScope);
-	FNetBitArrayView::ForAllExclusiveBits(UpdatedScope, CurrentScope, NewObjectFunctor, DestroyedObjectFunctor);
+	const FNetBitArrayView SubObjects = NetRefHandleManager->GetSubObjectInternalIndicesView();
+
+	// Process root objects first
+	FNetBitArrayView::ForAllExclusiveBitsByPredicate(UpdatedScope, CurrentScope, NewObjectFunctor, DestroyedObjectFunctor, [SubObjects](uint32 InternalIndex)
+		{
+			return !SubObjects.GetBit(InternalIndex);
+		}
+	);
+
+	// Process subobjects second
+	FNetBitArrayView::ForAllExclusiveBitsByPredicate(UpdatedScope, CurrentScope, NewObjectFunctor, DestroyedObjectFunctor, [SubObjects](uint32 InternalIndex)
+		{
+			return SubObjects.GetBit(InternalIndex);
+		}
+	);
+
 	CurrentScope.Copy(UpdatedScope);
 
 	// No objects marked for destroy can be in scope
