Add compile-time bounds check for FNetPacketNotify::MaxSequenceHistoryLength.
Add coverage of WriteHeader and ReadHeader to NetPacketNotifyTests.

diff --git a/Engine/Plugins/Runtime/ReplicationSystemTestPlugin/Source/Private/Tests/Core/TestNetPacketNotify.cpp b/Engine/Plugins/Runtime/ReplicationSystemTestPlugin/Source/Private/Tests/Core/TestNetPacketNotify.cpp
--- a/Engine/Plugins/Runtime/ReplicationSystemTestPlugin/Source/Private/Tests/Core/TestNetPacketNotify.cpp
+++ b/Engine/Plugins/Runtime/ReplicationSystemTestPlugin/Source/Private/Tests/Core/TestNetPacketNotify.cpp
@@ -6,6 +6,8 @@
 #include "Logging/LogScopedVerbosityOverride.h"
 #include "Net/Util/SequenceHistory.h"
 #include "Net/Util/SequenceNumber.h"
+#include "Serialization/BitWriter.h"
+#include "Serialization/BitReader.h"
 
 namespace UE::Net
 {
@@ -125,6 +127,8 @@ struct FNetPacketNotifyTestUtil : public FNetworkAutomationTestSuiteFixture
 		TArray<FNetPacketNotify::SequenceNumberT> AcceptedInPackets;
 		TArray<FNetPacketNotify::SequenceNumberT> AcknowledgedPackets;
 
+		FBitWriter Writer{FNetPacketNotify::MaxSequenceHistoryLength, true};
+
 		FTestNode(FNetPacketNotify::SequenceNumberT FirstSequence = 0)
 		{
 			PacketNotify.Init(FNetPacketNotify::SequenceNumberT(FirstSequence.Get() - 1), FirstSequence);
@@ -135,27 +139,22 @@ struct FNetPacketNotifyTestUtil : public FNetworkAutomationTestSuiteFixture
 			PacketNotify.Init(InSequnce, OutSequnce);
 		}
 
-		FNetPacketNotify::FNotificationHeader Send()
+		const FBitWriter& Send()
 		{
 			const SIZE_T CurrentHistoryWordCount = FMath::Clamp<SIZE_T>((PacketNotify.GetCurrentSequenceHistoryLength() + FNetPacketNotify::SequenceHistoryT::BitsPerWord - 1u) / FNetPacketNotify::SequenceHistoryT::BitsPerWord, 1u, FNetPacketNotify::SequenceHistoryT::WordCount);
-					
+
 			FNetPacketNotify::FNotificationHeader Data;
-			Data.Seq = PacketNotify.OutSeq;
-			Data.AckedSeq = PacketNotify.InAckSeq;
-			Data.HistoryWordCount = CurrentHistoryWordCount;
-			Data.History = PacketNotify.InSeqHistory;
 
-			PacketNotify.WrittenHistoryWordCount = CurrentHistoryWordCount;
-			PacketNotify.WrittenInAckSeq = Data.AckedSeq;
+			Writer.Reset();
+			PacketNotify.WriteHeader(Writer);
+
 			PacketNotify.CommitAndIncrementOutSeq();
 
-			return Data;
+			return Writer;
 		}
 
-		void Receive(const FNetPacketNotify::FNotificationHeader Data)
+		void Receive(const FBitWriter& IncomingData)
 		{
-			SIZE_T NotificationCount = 0;
-
 			auto HandleAck = [&](FNetPacketNotify::SequenceNumberT Seq, bool bDelivered)
 			{
 				if (bDelivered)
@@ -164,7 +163,12 @@ struct FNetPacketNotifyTestUtil : public FNetworkAutomationTestSuiteFixture
 				}
 			};
 
-			const uint32 SequenceDelta = PacketNotify.Update(Data, HandleAck);
+			FBitReader Reader(IncomingData.GetData(), IncomingData.GetNumBits());
+
+			FNetPacketNotify::FNotificationHeader Header;
+			PacketNotify.ReadHeader(Header, Reader);
+
+			const uint32 SequenceDelta = PacketNotify.Update(Header, HandleAck);
 			if (SequenceDelta > 0U && !PacketNotify.IsWaitingForSequenceHistoryFlush())
 			{
 				PacketNotify.AckSeq(PacketNotify.GetInSeq());
diff --git a/Engine/Source/Runtime/Engine/Private/Net/NetPacketNotify.cpp b/Engine/Source/Runtime/Engine/Private/Net/NetPacketNotify.cpp
--- a/Engine/Source/Runtime/Engine/Private/Net/NetPacketNotify.cpp
+++ b/Engine/Source/Runtime/Engine/Private/Net/NetPacketNotify.cpp
@@ -180,6 +180,8 @@ namespace
 		enum { AckSeqShift			= HistoryWordCountBits };
 		enum { SeqShift				= AckSeqShift + FNetPacketNotify::SequenceNumberBits };
 		
+		static_assert(FNetPacketNotify::SequenceHistoryT::WordCount <= (1 << HistoryWordCountBits), "Sequence history word count (minus one) must fit within HistoryWordCountBits. Check FNetPacketNotify::MaxSequenceHistoryLength.");
+
 		static uint32 Pack(SequenceNumberT Seq, SequenceNumberT AckedSeq, SIZE_T HistoryWordCount)
 		{
 			uint32 Packed = 0u;
diff --git a/Engine/Source/Runtime/Engine/Public/Net/NetPacketNotify.h b/Engine/Source/Runtime/Engine/Public/Net/NetPacketNotify.h
--- a/Engine/Source/Runtime/Engine/Public/Net/NetPacketNotify.h
+++ b/Engine/Source/Runtime/Engine/Public/Net/NetPacketNotify.h
@@ -225,7 +225,7 @@ void FNetPacketNotify::ProcessReceivedAcks(const FNotificationHeader& Notificati
 
 		// Make sure that we only look at the sequence history bit included in the notification data as the sequence history might have been reset, 
 		// in which case we might not receive the max size history even though the ack-count is bigger than the history
-		const SequenceNumberT::DifferenceT HistoryBits = NotificationData.HistoryWordCount * SequenceHistoryT::BitsPerWord;
+		const SequenceNumberT::DifferenceT HistoryBits = static_cast<SequenceNumberT::DifferenceT>(NotificationData.HistoryWordCount * SequenceHistoryT::BitsPerWord);
 
 		// Warn if the received sequence number is greater than our history buffer, since if that is the case we have to treat the data as lost
 		// Note: This should normally not be a problem since we try to flush the sequence history before allowing an overshoot of the sequence history window on the sending side.
