Iris - No longer group included objects should not cause scope hysteresis to kick in unless the object is dynamically filtered out.

#rnx

diff --git a/Engine/Plugins/Runtime/ReplicationSystemTestPlugin/Source/Private/Tests/ReplicationSystem/Filtering/TestFiltering.cpp b/Engine/Plugins/Runtime/ReplicationSystemTestPlugin/Source/Private/Tests/ReplicationSystem/Filtering/TestFiltering.cpp
--- a/Engine/Plugins/Runtime/ReplicationSystemTestPlugin/Source/Private/Tests/ReplicationSystem/Filtering/TestFiltering.cpp
+++ b/Engine/Plugins/Runtime/ReplicationSystemTestPlugin/Source/Private/Tests/ReplicationSystem/Filtering/TestFiltering.cpp
@@ -232,7 +232,6 @@ private:
 	bool bOriginalIsObjectScopeHysteresisEnabled = false;
 };
 
-
 class FTestFilteringWithConditionFixture : public FTestFilteringFixture
 {
 protected:
@@ -2865,7 +2864,7 @@ UE_NET_TEST_FIXTURE(FTestFilteringFixture, FilteringChangesDoesNotCauseHysteresi
 }
 
 // Test case where inclusion group added objects are filtered out yet should be replicated due to the inclusion group allowing replication. Inclusion group then disallows replication causing hysteresis to kick in.
-UE_NET_TEST_FIXTURE(FTestFilteringFixture, HysteresisKicksInForFormerlyInclusionGroupAllowedObject)
+UE_NET_TEST_FIXTURE(FTestFilteringFixture, HysteresisKicksInForFormerlyInclusionGroupAllowedObjectWhenFilterDisallowsReplication)
 {
 	constexpr uint32 DefaultHysteresisFrameCount = 2;
 	FScopedDefaultHysteresisFrameCount ScopedDefaultHysteresisFrameCount(DefaultHysteresisFrameCount);
@@ -2892,18 +2891,138 @@ UE_NET_TEST_FIXTURE(FTestFilteringFixture, HysteresisKicksInForFormerlyInclusion
 	// Disallow replication of inclusion group
 	Server->ReplicationSystem->SetGroupFilterStatus(GroupHandle, ENetFilterStatus::Disallow);
 
-	// Perform send update forthe hysteresis duration and make sure the object stays replicated.
+	// Perform send update for the hysteresis duration and make sure the object stays replicated.
+	for (uint32 It = 0; It != DefaultHysteresisFrameCount; ++It)
+	{
+		Server->UpdateAndSend({ Client });
+		UE_NET_ASSERT_NE(Client->GetReplicationBridge()->GetReplicatedObject(ServerObject->NetRefHandle), nullptr);
+	}
+
+	// Now the object should be filtered out
+	Server->UpdateAndSend({ Client });
+	UE_NET_ASSERT_EQ(Client->GetReplicationBridge()->GetReplicatedObject(ServerObject->NetRefHandle), nullptr);
+}
+
+// Variant of above test. Test case where inclusion group added objects are filtered out yet should be replicated due to the inclusion group allowing replication. Object is then removed from inclusion group causing hysteresis to kick in.
+UE_NET_TEST_FIXTURE(FTestFilteringFixture, HysteresisKicksInForFormerInclusionGroupMemberWhenFilterDisallowsReplication)
+{
+	constexpr uint32 DefaultHysteresisFrameCount = 2;
+	FScopedDefaultHysteresisFrameCount ScopedDefaultHysteresisFrameCount(DefaultHysteresisFrameCount);
+
+	// Add client
+	FReplicationSystemTestClient* Client = CreateClient();
+
+	// Add object to filter which disallows replication but inclusion group that allows it.
+	SetDynamicFilterStatus(ENetFilterStatus::Disallow);
+
+	UReplicatedTestObject* ServerObject = Server->CreateObject({ .IrisComponentCount = 0 });
+	Server->ReplicationSystem->SetFilter(ServerObject->NetRefHandle, MockFilterHandle);
+
+	FNetObjectGroupHandle GroupHandle = Server->ReplicationSystem->CreateGroup(NAME_None);
+	Server->ReplicationSystem->AddToGroup(GroupHandle, ServerObject->NetRefHandle);
+	Server->ReplicationSystem->AddInclusionFilterGroup(GroupHandle);
+	Server->ReplicationSystem->SetGroupFilterStatus(GroupHandle, ENetFilterStatus::Allow);
+
+	Server->UpdateAndSend({ Client });
+
+	// Make sure the object has been created 
+	UE_NET_ASSERT_NE(Client->GetReplicationBridge()->GetReplicatedObject(ServerObject->NetRefHandle), nullptr);
+
+	// Remove object from inclusion group
+	Server->ReplicationSystem->RemoveFromGroup(GroupHandle, ServerObject->NetRefHandle);
+
+	// Perform send update for the hysteresis duration and make sure the object stays replicated.
 	for (uint32 It = 0; It != DefaultHysteresisFrameCount; ++It)
 	{
 		Server->UpdateAndSend({ Client });
 		UE_NET_ASSERT_NE(Client->GetReplicationBridge()->GetReplicatedObject(ServerObject->NetRefHandle), nullptr);
 	}
 
-	// Now the object should now be filtered out
+	// Now the object should be filtered out
 	Server->UpdateAndSend({ Client });
 	UE_NET_ASSERT_EQ(Client->GetReplicationBridge()->GetReplicatedObject(ServerObject->NetRefHandle), nullptr);
 }
 
+// Test case where inclusion group added objects are allowed to be replicated by dyanmic filter too. Inclusion group then disallows replication which should not cause hysteresis to kick in.
+UE_NET_TEST_FIXTURE(FTestFilteringFixture, HysteresisDoesNotKickInForFormerlyInclusionGroupAllowedObjectWhenFilterAllowsReplication)
+{
+	constexpr uint32 DefaultHysteresisFrameCount = 2;
+	FScopedDefaultHysteresisFrameCount ScopedDefaultHysteresisFrameCount(DefaultHysteresisFrameCount);
+
+	// Add client
+	FReplicationSystemTestClient* Client = CreateClient();
+
+	// Add object to filter which allows replication.
+	SetDynamicFilterStatus(ENetFilterStatus::Allow);
+
+	UReplicatedTestObject* ServerObject = Server->CreateObject({ .IrisComponentCount = 0 });
+	Server->ReplicationSystem->SetFilter(ServerObject->NetRefHandle, MockFilterHandle);
+
+	FNetObjectGroupHandle GroupHandle = Server->ReplicationSystem->CreateGroup(NAME_None);
+	Server->ReplicationSystem->AddToGroup(GroupHandle, ServerObject->NetRefHandle);
+	Server->ReplicationSystem->AddInclusionFilterGroup(GroupHandle);
+	Server->ReplicationSystem->SetGroupFilterStatus(GroupHandle, ENetFilterStatus::Allow);
+
+	Server->UpdateAndSend({ Client });
+
+	// Make sure the object has been created 
+	UE_NET_ASSERT_NE(Client->GetReplicationBridge()->GetReplicatedObject(ServerObject->NetRefHandle), nullptr);
+
+	// Disallow replication of inclusion group
+	Server->ReplicationSystem->SetGroupFilterStatus(GroupHandle, ENetFilterStatus::Disallow);
+
+	// Perform send update for the hysteresis duration and make sure the object stays replicated.
+	for (uint32 It = 0; It != DefaultHysteresisFrameCount; ++It)
+	{
+		Server->UpdateAndSend({ Client });
+		UE_NET_ASSERT_NE(Client->GetReplicationBridge()->GetReplicatedObject(ServerObject->NetRefHandle), nullptr);
+	}
+
+	// The object should stay replicated
+	Server->UpdateAndSend({ Client });
+	UE_NET_ASSERT_NE(Client->GetReplicationBridge()->GetReplicatedObject(ServerObject->NetRefHandle), nullptr);
+}
+
+// Variant of above test. Test case where inclusion group added objects are allowed to be replicated by dynamic filter too. Object is then removed from inclusion group which should not cause hysteresis to kick in.
+UE_NET_TEST_FIXTURE(FTestFilteringFixture, HysteresisDoesNotKickInForFormerInclusionGroupMemberWhenFilterAllowsReplication)
+{
+	constexpr uint32 DefaultHysteresisFrameCount = 2;
+	FScopedDefaultHysteresisFrameCount ScopedDefaultHysteresisFrameCount(DefaultHysteresisFrameCount);
+
+	// Add client
+	FReplicationSystemTestClient* Client = CreateClient();
+
+	// Add object to filter which allows replication.
+	SetDynamicFilterStatus(ENetFilterStatus::Allow);
+
+	UReplicatedTestObject* ServerObject = Server->CreateObject({ .IrisComponentCount = 0 });
+	Server->ReplicationSystem->SetFilter(ServerObject->NetRefHandle, MockFilterHandle);
+
+	FNetObjectGroupHandle GroupHandle = Server->ReplicationSystem->CreateGroup(NAME_None);
+	Server->ReplicationSystem->AddToGroup(GroupHandle, ServerObject->NetRefHandle);
+	Server->ReplicationSystem->AddInclusionFilterGroup(GroupHandle);
+	Server->ReplicationSystem->SetGroupFilterStatus(GroupHandle, ENetFilterStatus::Allow);
+
+	Server->UpdateAndSend({ Client });
+
+	// Make sure the object has been created 
+	UE_NET_ASSERT_NE(Client->GetReplicationBridge()->GetReplicatedObject(ServerObject->NetRefHandle), nullptr);
+
+	// Remove object from inclusion group
+	Server->ReplicationSystem->RemoveFromGroup(GroupHandle, ServerObject->NetRefHandle);
+
+	// Perform send update for the hysteresis duration and make sure the object stays replicated.
+	for (uint32 It = 0; It != DefaultHysteresisFrameCount; ++It)
+	{
+		Server->UpdateAndSend({ Client });
+		UE_NET_ASSERT_NE(Client->GetReplicationBridge()->GetReplicatedObject(ServerObject->NetRefHandle), nullptr);
+	}
+
+	// The object should stay replicated
+	Server->UpdateAndSend({ Client });
+	UE_NET_ASSERT_NE(Client->GetReplicationBridge()->GetReplicatedObject(ServerObject->NetRefHandle), nullptr);
+}
+
 // Make sure that connection throttling does not cause objects to be filtered out too soon. Also verify throttling occurs.
 UE_NET_TEST_FIXTURE(FTestFilteringFixture, HysteresisConnectionThrottlingWorksAsExpected)
 {
diff --git a/Engine/Source/Runtime/Experimental/Iris/Core/Private/Iris/ReplicationSystem/Filtering/ReplicationFiltering.cpp b/Engine/Source/Runtime/Experimental/Iris/Core/Private/Iris/ReplicationSystem/Filtering/ReplicationFiltering.cpp
--- a/Engine/Source/Runtime/Experimental/Iris/Core/Private/Iris/ReplicationSystem/Filtering/ReplicationFiltering.cpp
+++ b/Engine/Source/Runtime/Experimental/Iris/Core/Private/Iris/ReplicationSystem/Filtering/ReplicationFiltering.cpp
@@ -2275,22 +2275,17 @@ bool FReplicationFiltering::ClearGroupInclusionFilterEffectsForObject(uint32 Obj
 
 	if (!IsIncludedByAnyGroup(ObjectIndex, ConnectionId))
 	{
-		FNetBitArray& GroupIncludedObjects = ConnectionInfo.GroupIncludedObjects;
-		
-		GroupIncludedObjects.ClearBit(ObjectIndex);
-
 		// Dynamically filtered objects are subject to hysteresis.
-		if (HysteresisState.Mode == EHysteresisProcessingMode::Enabled)
+		if (HysteresisState.Mode == EHysteresisProcessingMode::Enabled && DynamicFilterEnabledObjects.GetBit(ObjectIndex))
 		{
-			const uint8 HysteresisFrameCount = DynamicFilterEnabledObjects.GetBit(ObjectIndex) ? ObjectScopeHysteresisFrameCounts[ObjectIndex] : uint8(0);
-			if (HysteresisFrameCount)
-			{
-				const uint32 ConnIdMod = ConnectionId % HysteresisState.ConnectionIdStride;
-				const uint32 AdjustHysteresisForUpdateThrottling = (HysteresisState.ConnectionStartId + HysteresisState.ConnectionIdStride - ConnIdMod) % HysteresisState.ConnectionIdStride;
-				ConnectionInfo.HysteresisUpdater.SetHysteresisFrameCount(ObjectIndex, static_cast<uint16>(HysteresisFrameCount + AdjustHysteresisForUpdateThrottling));
-			}
+			// Force processing of dynamically filtered object. If the object is still in scope and dynamically filtered out it should cause hysteresis to kick in.
+			ObjectsRequiringDynamicFilterUpdate.SetBit(ObjectIndex);
 		}
 
+		FNetBitArray& GroupIncludedObjects = ConnectionInfo.GroupIncludedObjects;
+		
+		GroupIncludedObjects.ClearBit(ObjectIndex);
+
 		for (const FInternalNetRefIndex SubObjectIndex : NetRefHandleManager->GetSubObjects(ObjectIndex))
 		{
 			GroupIncludedObjects.ClearBit(SubObjectIndex);
