FReplicationSystemServerClientTestFixture improvements.
Destroy clients properly when tearing down FReplicationSystemServerClientTestFixture.
Added IsResolvableNetRefHandle to see whether a FNetRefHandle can be resolved to an object pointer.
Added isValidNetRefHandle to see whether a FNetRefHandle is known to the ReplicationSystem.

#rnx

diff --git a/Engine/Plugins/Runtime/ReplicationSystemTestPlugin/Source/Private/Tests/ReplicationSystem/ReplicationSystemServerClientTestFixture.cpp b/Engine/Plugins/Runtime/ReplicationSystemTestPlugin/Source/Private/Tests/ReplicationSystem/ReplicationSystemServerClientTestFixture.cpp
--- a/Engine/Plugins/Runtime/ReplicationSystemTestPlugin/Source/Private/Tests/ReplicationSystem/ReplicationSystemServerClientTestFixture.cpp
+++ b/Engine/Plugins/Runtime/ReplicationSystemTestPlugin/Source/Private/Tests/ReplicationSystem/ReplicationSystemServerClientTestFixture.cpp
@@ -3,6 +3,7 @@
 #include "ReplicationSystemServerClientTestFixture.h"
 #include "Logging/LogScopedVerbosityOverride.h"
 #include "Iris/Core/IrisLog.h"
+#include "Iris/ReplicationSystem/NetRefHandleManager.h"
 #include "Iris/ReplicationSystem/ReplicationSystem.h"
 #include "Iris/ReplicationSystem/ReplicationSystemInternal.h"
 #include "Iris/ReplicationSystem/ReplicationConnections.h"
@@ -227,6 +228,16 @@ void FReplicationSystemTestNode::DestroyObject(UReplicatedTestObject* Object, EE
 	Object->MarkAsGarbage();
 }
 
+bool FReplicationSystemTestNode::IsResolvableNetRefHandle(FNetRefHandle RefHandle)
+{
+	return ReplicationBridge->GetReplicatedObject(RefHandle) != nullptr;
+}
+
+bool FReplicationSystemTestNode::IsValidNetRefHandle(FNetRefHandle RefHandle)
+{
+	return ReplicationSystem->GetReplicationSystemInternal()->GetNetRefHandleManager().IsValidNetRefHandle(RefHandle);
+}
+
 uint32 FReplicationSystemTestNode::AddConnection()
 {
 	check(ReplicationSystem);
@@ -522,12 +533,11 @@ void FReplicationSystemServerClientTestFixture::SetUp()
 
 void FReplicationSystemServerClientTestFixture::TearDown()
 {
-	delete Server;
-	for (FReplicationSystemTestClient* Client : Clients)
+	for (TArray<FReplicationSystemTestClient*> TempClients = Clients; FReplicationSystemTestClient* Client : TempClients)
 	{
 		delete Client;
 	}
-	Clients.Empty();
+	delete Server;
 	DataStreamUtil.TearDown();
 
 	FNetworkAutomationTestSuiteFixture::TearDown();
diff --git a/Engine/Plugins/Runtime/ReplicationSystemTestPlugin/Source/Private/Tests/ReplicationSystem/ReplicationSystemServerClientTestFixture.h b/Engine/Plugins/Runtime/ReplicationSystemTestPlugin/Source/Private/Tests/ReplicationSystem/ReplicationSystemServerClientTestFixture.h
--- a/Engine/Plugins/Runtime/ReplicationSystemTestPlugin/Source/Private/Tests/ReplicationSystem/ReplicationSystemServerClientTestFixture.h
+++ b/Engine/Plugins/Runtime/ReplicationSystemTestPlugin/Source/Private/Tests/ReplicationSystem/ReplicationSystemServerClientTestFixture.h
@@ -137,6 +137,12 @@ public:
 
 	void DestroyObject(UReplicatedTestObject*, EEndReplicationFlags EndReplicationFlags = EEndReplicationFlags::Destroy);
 
+	// Returns true if the NetRefHandle can still be resolved to an object.
+	bool IsResolvableNetRefHandle(FNetRefHandle RefHandle);
+
+	// Returns true if the NetRefHandleManager considers the NetRefHandle valid.
+	bool IsValidNetRefHandle(FNetRefHandle RefHandle);
+
 	// Connection
 	uint32 AddConnection();
 	void RemoveConnection(uint32 ConnectionId);
@@ -168,8 +174,8 @@ public:
 
 public:
 	TUniquePtr<UE::Net::FNetTokenStore> NetTokenStore;
-	UReplicationSystem* ReplicationSystem;
-	UReplicatedTestObjectBridge* ReplicationBridge;
+	UReplicationSystem* ReplicationSystem = nullptr;
+	UReplicatedTestObjectBridge* ReplicationBridge = nullptr;
 	TArray<TStrongObjectPtr<UObject>> CreatedObjects;
 	EReplicationSystemSendPass CurrentSendPass = EReplicationSystemSendPass::Invalid;
 
