Added unit tests for NetBitArray::ForAllUnsetBits

#rnx

diff --git a/Engine/Plugins/Runtime/ReplicationSystemTestPlugin/Source/Private/Tests/Core/TestNetBitArray.cpp b/Engine/Plugins/Runtime/ReplicationSystemTestPlugin/Source/Private/Tests/Core/TestNetBitArray.cpp
--- a/Engine/Plugins/Runtime/ReplicationSystemTestPlugin/Source/Private/Tests/Core/TestNetBitArray.cpp
+++ b/Engine/Plugins/Runtime/ReplicationSystemTestPlugin/Source/Private/Tests/Core/TestNetBitArray.cpp
@@ -744,6 +744,7 @@ UE_NET_TEST(FNetBitArrayView, TestForEachSet)
 	};
 
 
+	// Test with first two bits set
 	{
 		uint32 SingleWordBuffer = 0xffffff00;
 		FNetBitArrayView BitArray(&SingleWordBuffer, 8, FNetBitArrayView::NoResetNoValidate);
@@ -751,10 +752,14 @@ UE_NET_TEST(FNetBitArrayView, TestForEachSet)
 		FCollectSetBitsFunctor Functor;
 
 		BitArray.ForAllSetBits(Functor);
-
 		UE_NET_ASSERT_EQ(0, Functor.Invoked.Num());
+
+		Functor.Invoked.Reset();
+		BitArray.ForAllUnsetBits(Functor);
+		UE_NET_ASSERT_EQ(8, Functor.Invoked.Num());
 	}
 
+	// Test with two last bits set
 	{
 		uint32 SingleWordBuffer = 0xa0000000;
 		FNetBitArrayView BitArray(&SingleWordBuffer, 32);
@@ -762,10 +767,14 @@ UE_NET_TEST(FNetBitArrayView, TestForEachSet)
 		FCollectSetBitsFunctor Functor;
 
 		BitArray.ForAllSetBits(Functor);
-
 		UE_NET_ASSERT_EQ(2, Functor.Invoked.Num());
+
+		Functor.Invoked.Reset();
+		BitArray.ForAllUnsetBits(Functor);
+		UE_NET_ASSERT_EQ(30, Functor.Invoked.Num());
 	}
 
+	// Test ForAllSetBits with first and last bits set
 	{
 		uint32 SingleWordBuffer = 0x80000001;
 		FNetBitArrayView BitArray(&SingleWordBuffer, 32);
@@ -773,12 +782,25 @@ UE_NET_TEST(FNetBitArrayView, TestForEachSet)
 		FCollectSetBitsFunctor Functor;
 
 		BitArray.ForAllSetBits(Functor);
+		UE_NET_ASSERT_EQ(2, Functor.Invoked.Num());
+		UE_NET_ASSERT_EQ(0u, Functor.Invoked[0]);
+		UE_NET_ASSERT_EQ(31u, Functor.Invoked[1]);
+	}
 
+	// Test ForAllUnsetBits with first and last bits not set
+	{
+		uint32 SingleWordBuffer = 0x7FFFFFFE;
+		FNetBitArrayView BitArray(&SingleWordBuffer, 32);
+		
+		FCollectSetBitsFunctor Functor;
+		
+		BitArray.ForAllUnsetBits(Functor);
 		UE_NET_ASSERT_EQ(2, Functor.Invoked.Num());
 		UE_NET_ASSERT_EQ(0u, Functor.Invoked[0]);
 		UE_NET_ASSERT_EQ(31u, Functor.Invoked[1]);
 	}
 
+	// Test ForAllSetBits with ALL bits set
 	{
 		uint32 SingleWordBuffer = 0xffffffff;
 		FNetBitArrayView BitArray(&SingleWordBuffer, 32);
@@ -789,6 +811,18 @@ UE_NET_TEST(FNetBitArrayView, TestForEachSet)
 		UE_NET_ASSERT_EQ(32, Functor.Invoked.Num());
 	}
 
+	// Test ForAllUnsetBits with ZERO bits set
+	{
+		uint32 SingleWordBuffer = 0x00000000;
+		FNetBitArrayView BitArray(&SingleWordBuffer, 32);
+		
+		FCollectSetBitsFunctor Functor;
+		BitArray.ForAllUnsetBits(Functor);
+
+		UE_NET_ASSERT_EQ(32, Functor.Invoked.Num());
+	}
+
+	// Test ForAllSetBit with random bits
 	{
 		const uint32 IndicesToTest[] = { 1, 31, 32, 34, 35, 36, 37, 64, 65, 66, 126 };
 		const uint32 ExpectedNumBitsSet = UE_ARRAY_COUNT(IndicesToTest);
@@ -812,6 +846,32 @@ UE_NET_TEST(FNetBitArrayView, TestForEachSet)
 			UE_NET_ASSERT_EQ(IndicesToTest[It], Functor.Invoked[It]);
 		}
 	}
+
+	// Test ForAllUnsetBits with random bits
+	{
+		const uint32 IndicesToTest[] = { 1, 31, 32, 34, 35, 36, 37, 64, 65, 66, 126 };
+		const uint32 ExpectedNumBitsSet = UE_ARRAY_COUNT(IndicesToTest);
+
+		uint32 WordBuffer[32] = { 0 };
+		
+		FNetBitArrayView BitArray(&WordBuffer[0], 127);
+		BitArray.SetAllBits();
+
+		for (uint32 It = 0; It < ExpectedNumBitsSet; ++It)
+		{
+			BitArray.SetBitValue(IndicesToTest[It], false);
+		}
+
+		FCollectSetBitsFunctor Functor;
+		BitArray.ForAllUnsetBits(Functor);
+
+		UE_NET_ASSERT_EQ(ExpectedNumBitsSet, (uint32)Functor.Invoked.Num());
+
+		for (uint32 It = 0; It < ExpectedNumBitsSet; ++It)
+		{
+			UE_NET_ASSERT_EQ(IndicesToTest[It], Functor.Invoked[It]);
+		}
+	}
 }
 
 UE_NET_TEST(FNetBitArrayView, TestGetSetBitIndices)
