Iris - Handle object hierarchy destruction atomically. The new behavior is controlled by the cvar net.Iris.DestroyRootAndSubObjectsIndividually which defaults to false to support the atomic destruction.

Bumped the max replication record count from 32768 to 65535 and made it easier to change the max replication record count to an arbitrary number, though further bumping requires changing the index type. Introduced SubObjectRecord which enables the atomic destruction of arbitrarily large object hirerarchies. The latter is only used when atomic object hierarchy destruction is enabled.

Detect when we're starting to run out of replication records instead of crashing. When running out of records the ReplicationWriter will only send OOB attachments and huge objects, if at all possible. Control the threshold for the replication record starvationwith cvar net.Iris.ReplicationWriterReplicationRecordStarvationThreshold which defaults to 1000.

Limit the number of objects and destruction infos sent per frame using the new cvar net.Iris.ReplicationWriterMaxDestroyObjectsPerPacket which defaults to 128. This can be used to reduce replication record consumption and limit CPU usage for object destruction on clients.

diff --git a/Engine/Plugins/Runtime/ReplicationSystemTestPlugin/Source/Private/Tests/ReplicationSystem/Filtering/TestFiltering.cpp b/Engine/Plugins/Runtime/ReplicationSystemTestPlugin/Source/Private/Tests/ReplicationSystem/Filtering/TestFiltering.cpp
--- a/Engine/Plugins/Runtime/ReplicationSystemTestPlugin/Source/Private/Tests/ReplicationSystem/Filtering/TestFiltering.cpp
+++ b/Engine/Plugins/Runtime/ReplicationSystemTestPlugin/Source/Private/Tests/ReplicationSystem/Filtering/TestFiltering.cpp
@@ -3120,13 +3120,13 @@ UE_NET_TEST_FIXTURE(FTestFilteringFixture, LotsOfObjectsCanBeFilteredOutViaHyste
 	for (uint32 It = 0; It < HighObjectCount; ++It)
 	{
 		Server->UpdateAndSend({ Client });
-		if (Client->GetReplicationBridge()->GetReplicatedObject(ServerObjects[HighObjectCount - 1]->NetRefHandle) != nullptr)
+		if (Client->IsResolvableNetRefHandle(ServerObjects[HighObjectCount - 1]->NetRefHandle))
 		{
 			bAllObjectsCreated = true;
 
 			for (const UReplicatedTestObject* ServerObject : ServerObjects)
 			{
-				if (Client->GetReplicationBridge()->GetReplicatedObject(ServerObject->NetRefHandle) == nullptr)
+				if (!Client->IsResolvableNetRefHandle(ServerObject->NetRefHandle))
 				{
 					bAllObjectsCreated = false;
 					break;
@@ -3144,14 +3144,20 @@ UE_NET_TEST_FIXTURE(FTestFilteringFixture, LotsOfObjectsCanBeFilteredOutViaHyste
 
 	SetDynamicFilterStatus(ENetFilterStatus::Disallow);
 
-	// Send and deliver packet. Need to update twice to have the objects filtered out.
-	Server->UpdateAndSend({ Client });
+	// Send and deliver packet. Need to update a couple of times to have the objects filtered out.
 	Server->UpdateAndSend({ Client });
+	for (uint32 It = 0; It < HighObjectCount; ++It)
+	{
+		if (!Server->UpdateAndSend({ Client }))
+		{
+			break;
+		}
+	}
 
 	// All client objects should be destroyed
 	for (const UReplicatedTestObject* ServerObject : ServerObjects)
 	{
-		UE_NET_ASSERT_EQ(Client->GetReplicationBridge()->GetReplicatedObject(ServerObject->NetRefHandle), nullptr);
+		UE_NET_ASSERT_FALSE(Client->IsValidNetRefHandle(ServerObject->NetRefHandle));
 	}
 }
 
diff --git a/Engine/Plugins/Runtime/ReplicationSystemTestPlugin/Source/Private/Tests/ReplicationSystem/TestCancelPendingDestroy.cpp b/Engine/Plugins/Runtime/ReplicationSystemTestPlugin/Source/Private/Tests/ReplicationSystem/TestCancelPendingDestroy.cpp
--- a/Engine/Plugins/Runtime/ReplicationSystemTestPlugin/Source/Private/Tests/ReplicationSystem/TestCancelPendingDestroy.cpp
+++ b/Engine/Plugins/Runtime/ReplicationSystemTestPlugin/Source/Private/Tests/ReplicationSystem/TestCancelPendingDestroy.cpp
@@ -39,7 +39,7 @@ UE_NET_TEST_FIXTURE(FTestCancelPendingDestroyFixture, TestCancelPendingDestroyDu
 	Server->DeliverTo(Client, DeliverPacket);
 
 	// Verify that the object now exists on client
-	UE_NET_ASSERT_NE(Client->GetReplicationBridge()->GetReplicatedObject(ServerObject->NetRefHandle), nullptr);
+	UE_NET_ASSERT_TRUE(Client->IsResolvableNetRefHandle(ServerObject->NetRefHandle));
 
 	// Modify a property on the object and make sure it's replicated as the object should now be confirmed created
 	ServerObject->IntA ^= 1;
@@ -84,7 +84,7 @@ UE_NET_TEST_FIXTURE(FTestCancelPendingDestroyFixture, TestCancelPendingDestroyDu
 	Server->SendAndDeliverTo(Client, DeliverPacket);
 	Server->PostSendUpdate();
 
-	UE_NET_ASSERT_NE(Client->GetReplicationBridge()->GetReplicatedObject(ServerObject->NetRefHandle), nullptr);
+	UE_NET_ASSERT_TRUE(Client->IsResolvableNetRefHandle(ServerObject->NetRefHandle));
 }
 
 UE_NET_TEST_FIXTURE(FTestCancelPendingDestroyFixture, TestCancelPendingDestroyDuringWaitOnCreateConfirmationWithPacketLossNotifyBeforeCancel)
@@ -118,7 +118,7 @@ UE_NET_TEST_FIXTURE(FTestCancelPendingDestroyFixture, TestCancelPendingDestroyDu
 	Server->SendAndDeliverTo(Client, DeliverPacket);
 	Server->PostSendUpdate();
 
-	UE_NET_ASSERT_NE(Client->GetReplicationBridge()->GetReplicatedObject(ServerObject->NetRefHandle), nullptr);
+	UE_NET_ASSERT_TRUE(Client->IsResolvableNetRefHandle(ServerObject->NetRefHandle));
 }
 
 UE_NET_TEST_FIXTURE(FTestCancelPendingDestroyFixture, TestCancelPendingDestroyDuringCreated)
@@ -135,7 +135,7 @@ UE_NET_TEST_FIXTURE(FTestCancelPendingDestroyFixture, TestCancelPendingDestroyDu
 	Server->PostSendUpdate();
 
 	// Verify that the object now also exists on the client
-	UE_NET_ASSERT_NE(Client->GetReplicationBridge()->GetReplicatedObject(ServerObject->NetRefHandle), nullptr);
+	UE_NET_ASSERT_TRUE(Client->IsResolvableNetRefHandle(ServerObject->NetRefHandle));
 
 	// Filter out object to cause a PendingDestroy
 	Server->GetReplicationSystem()->AddToGroup(Server->GetReplicationSystem()->GetNotReplicatedNetObjectGroup(), ServerObject->NetRefHandle);
@@ -187,7 +187,7 @@ UE_NET_TEST_FIXTURE(FTestCancelPendingDestroyFixture, TestCancelPendingDestroyDu
 	Server->DeliverTo(Client, DoNotDeliverPacket);
 
 	// Verify that object still exists on client
-	UE_NET_ASSERT_NE(Client->GetReplicationBridge()->GetReplicatedObject(ServerObject->NetRefHandle), nullptr);
+	UE_NET_ASSERT_TRUE(Client->IsResolvableNetRefHandle(ServerObject->NetRefHandle));
 
 	// Modify a property on the object and make sure it's replicated as the object should still be created
 	ServerObject->IntA ^= 1;
@@ -229,7 +229,7 @@ UE_NET_TEST_FIXTURE(FTestCancelPendingDestroyFixture, TestCancelPendingDestroyDu
 	Server->PostSendUpdate();
 
 	// Verify that object still exists on client
-	UE_NET_ASSERT_NE(Client->GetReplicationBridge()->GetReplicatedObject(ServerObject->NetRefHandle), nullptr);
+	UE_NET_ASSERT_TRUE(Client->IsResolvableNetRefHandle(ServerObject->NetRefHandle));
 
 	// Modify a property on the object and make sure it's replicated as the object should still be created
 	ServerObject->IntA ^= 1;
@@ -325,7 +325,7 @@ UE_NET_TEST_FIXTURE(FTestCancelPendingDestroyFixture, TestCancelPendingDestroyDu
 	Server->DeliverTo(Client, DeliverPacket);
 
 	// The destroy was delivered so the client object should not exist
-	UE_NET_ASSERT_EQ(Client->GetReplicationBridge()->GetReplicatedObject(ServerObject->NetRefHandle), nullptr);
+	UE_NET_ASSERT_FALSE(Client->IsValidNetRefHandle(ServerObject->NetRefHandle));
 
 	// Write and send packet
 	Server->PreSendUpdate();
@@ -592,6 +592,49 @@ UE_NET_TEST_FIXTURE(FTestCancelPendingDestroyFixture, TestCancelPendingSubObject
 	UE_NET_ASSERT_EQ(ClientSubObject->IntA, ServerSubObject->IntA);
 }
 
+UE_NET_TEST_FIXTURE(FTestCancelPendingDestroyFixture, TestDestroySubObjectWhenRootGoesOutOfAndInToScope)
+{
+	// Add a client
+	FReplicationSystemTestClient* Client = CreateClient();
+
+	// Spawn object with subobject on server
+	UTestReplicatedIrisObject* ServerObject = Server->CreateObject(UTestReplicatedIrisObject::FComponents());
+	const FNetRefHandle ObjectRefHandle = ServerObject->NetRefHandle;
+	UTestReplicatedIrisObject* ServerSubObject = Server->CreateSubObject(ServerObject->NetRefHandle, UTestReplicatedIrisObject::FComponents());
+	const FNetRefHandle SubObjectRefHandle = ServerSubObject->NetRefHandle;
+
+	// Write and send packet
+	Server->UpdateAndSend({Client});
+
+	// Verify that the subobject exists on the client
+	UE_NET_ASSERT_TRUE(Client->IsResolvableNetRefHandle(SubObjectRefHandle));
+
+	// Destroy subobject to cause a SubObjectPendingDestroy
+	Server->DestroyObject(ServerSubObject);
+	Server->PreSendUpdate();
+	Server->PostSendUpdate();
+
+	// Filter out root object and write packet. This causes subobject to go from SubObjectPendingDestroy to PendingDestroy.
+	Server->GetReplicationSystem()->AddToGroup(Server->GetReplicationSystem()->GetNotReplicatedNetObjectGroup(), ServerObject->NetRefHandle);
+	Server->PreSendUpdate();
+	Server->SendUpdate(Client->ConnectionIdOnServer);
+	Server->PostSendUpdate();
+
+	// Lose packet
+	Server->DeliverTo(Client, DoNotDeliverPacket);
+
+	// Remove root object from NotReplicated group to have it end up in CancelPendingDestroy
+	Server->GetReplicationSystem()->RemoveFromGroup(Server->GetReplicationSystem()->GetNotReplicatedNetObjectGroup(), ServerObject->NetRefHandle);
+
+	// Update and deliver packet
+	Server->UpdateAndSend({Client});
+
+	// The subobject should not exist
+	UE_NET_ASSERT_FALSE(Client->IsValidNetRefHandle(SubObjectRefHandle));
+	// The root object should still exist
+	UE_NET_ASSERT_TRUE(Client->IsResolvableNetRefHandle(ObjectRefHandle));
+}
+
 UE_NET_TEST_FIXTURE(FTestCancelPendingDestroyFixture, TestCancelPendingSubObjectDestroyDuringWaitOnDestroyConfirmationWithPacketLoss)
 {
 	// Add a client
@@ -921,65 +964,6 @@ UE_NET_TEST_FIXTURE(FTestCancelPendingDestroyFixture, TestCancelingCancelPending
 	UE_NET_ASSERT_EQ(Client->GetReplicationBridge()->GetReplicatedObject(ServerSubObject->NetRefHandle), nullptr);
 }
 
-UE_NET_TEST_FIXTURE(FTestCancelPendingDestroyFixture, TestCancelingSubObjectDestroyAfterParentIsMarkedForDestroy)
-{
-	// Add a client
-	FReplicationSystemTestClient* Client = CreateClient();
-
-	// Spawn object on server
-	UTestReplicatedIrisObject* ServerObject = Server->CreateObject(UTestReplicatedIrisObject::FComponents());
-	UTestReplicatedIrisObject* ServerSubObject1 = Server->CreateSubObject(ServerObject->NetRefHandle, UTestReplicatedIrisObject::FComponents());
-	UTestReplicatedIrisObject* ServerSubObject2 = Server->CreateSubObject(ServerObject->NetRefHandle, UTestReplicatedIrisObject::FComponents());
-
-	// Write and send packet
-	Server->PreSendUpdate();
-	Server->SendAndDeliverTo(Client, DeliverPacket);
-	Server->PostSendUpdate();
-
-	// Filter out subobject and write packet to cause subobject to end up in WaitOnDestroyConfirmation
-	Server->GetReplicationSystem()->AddToGroup(Server->GetReplicationSystem()->GetNotReplicatedNetObjectGroup(), ServerSubObject2->NetRefHandle);
-	Server->PreSendUpdate();
-	Server->SendUpdate(Client->ConnectionIdOnServer);
-	Server->PostSendUpdate();
-
-	// Filter out parent and write packet
-	Server->GetReplicationSystem()->AddToGroup(Server->GetReplicationSystem()->GetNotReplicatedNetObjectGroup(), ServerObject->NetRefHandle);
-	Server->PreSendUpdate();
-	Server->SendUpdate(Client->ConnectionIdOnServer);
-	Server->PostSendUpdate();
-
-	// Remove subobject from filter to cause subobject to end up in CancelPendingDestroy
-	Server->GetReplicationSystem()->RemoveFromGroup(Server->GetReplicationSystem()->GetNotReplicatedNetObjectGroup(), ServerSubObject2->NetRefHandle);
-	Server->PreSendUpdate();
-	Server->PostSendUpdate();
-
-	// Cause subobject destroy packet to be lost
-	Server->DeliverTo(Client, DoNotDeliverPacket);
-
-	// Verify all objects and subobjects exist on the client
-	UE_NET_ASSERT_NE(Client->GetReplicationBridge()->GetReplicatedObject(ServerObject->NetRefHandle), nullptr);
-	UE_NET_ASSERT_NE(Client->GetReplicationBridge()->GetReplicatedObject(ServerSubObject1->NetRefHandle), nullptr);
-	UE_NET_ASSERT_NE(Client->GetReplicationBridge()->GetReplicatedObject(ServerSubObject2->NetRefHandle), nullptr);
-
-	// Deliver owner destroy packet
-	Server->DeliverTo(Client, DeliverPacket);
-
-	// Verify that no objects exist on the client
-	UE_NET_ASSERT_EQ(Client->GetReplicationBridge()->GetReplicatedObject(ServerObject->NetRefHandle), nullptr);
-	UE_NET_ASSERT_EQ(Client->GetReplicationBridge()->GetReplicatedObject(ServerSubObject1->NetRefHandle), nullptr);
-	UE_NET_ASSERT_EQ(Client->GetReplicationBridge()->GetReplicatedObject(ServerSubObject2->NetRefHandle), nullptr);
-
-	// Write and send a packet
-	Server->PreSendUpdate();
-	Server->SendAndDeliverTo(Client, DeliverPacket);
-	Server->PostSendUpdate();
-
-	// Verify that all objects are still gone on the client
-	UE_NET_ASSERT_EQ(Client->GetReplicationBridge()->GetReplicatedObject(ServerObject->NetRefHandle), nullptr);
-	UE_NET_ASSERT_EQ(Client->GetReplicationBridge()->GetReplicatedObject(ServerSubObject1->NetRefHandle), nullptr);
-	UE_NET_ASSERT_EQ(Client->GetReplicationBridge()->GetReplicatedObject(ServerSubObject2->NetRefHandle), nullptr);
-}
-
 UE_NET_TEST_FIXTURE(FTestCancelPendingDestroyFixture, TestCancelRootObjectDestroyAfterSubObjectDestroy)
 {
 	// Add a client
diff --git a/Engine/Plugins/Runtime/ReplicationSystemTestPlugin/Source/Private/Tests/ReplicationSystem/TestObjectSplitting.cpp b/Engine/Plugins/Runtime/ReplicationSystemTestPlugin/Source/Private/Tests/ReplicationSystem/TestObjectSplitting.cpp
--- a/Engine/Plugins/Runtime/ReplicationSystemTestPlugin/Source/Private/Tests/ReplicationSystem/TestObjectSplitting.cpp
+++ b/Engine/Plugins/Runtime/ReplicationSystemTestPlugin/Source/Private/Tests/ReplicationSystem/TestObjectSplitting.cpp
@@ -1,6 +1,9 @@
 // Copyright Epic Games, Inc. All Rights Reserved.
 
+#include "HAL/IConsoleManager.h"
 #include "Iris/ReplicationSystem/NetBlob/NetObjectBlobHandler.h"
+#include "Iris/ReplicationSystem/ReplicationRecord.h"
+#include "Misc/ScopeExit.h"
 #include "NetBlob/NetBlobTestFixture.h"
 #include "NetBlob/MockNetBlob.h"
 #include "NetBlob/MockNetObjectAttachment.h"
@@ -1147,6 +1150,62 @@ UE_NET_TEST_FIXTURE(FSplitObjectTestFixture, TestDependentObjectCannotBeDestroye
 	UE_NET_ASSERT_EQ(ClientDependentObject, nullptr);
 }
 
+UE_NET_TEST_FIXTURE(FSplitObjectTestFixture, TestHugeObjectCanBeSentDuringReplicationRecordStarvation)
+{
+	IConsoleVariable* CVarReplicationRecordStarvationThreshold = IConsoleManager::Get().FindConsoleVariable(TEXT("net.Iris.ReplicationWriterReplicationRecordStarvationThreshold"));
+	UE_NET_ASSERT_NE(CVarReplicationRecordStarvationThreshold, nullptr);
+	UE_NET_ASSERT_TRUE(CVarReplicationRecordStarvationThreshold->IsVariableInt());
+	const int32 PrevReplicationRecordStarvationThreshold = CVarReplicationRecordStarvationThreshold->GetInt();
+	ON_SCOPE_EXIT
+	{
+		CVarReplicationRecordStarvationThreshold->Set(PrevReplicationRecordStarvationThreshold, ECVF_SetByCode);
+	};
+	
+	// Add a client
+	FReplicationSystemTestClient* Client = CreateClient();
+
+	// Set starvation threshold to highest possible
+	CVarReplicationRecordStarvationThreshold->Set(FReplicationRecord::MaxReplicationRecordCount, ECVF_SetByCode);
+
+	// Consume a few ReplicationRecords to enter starvation
+	for (int It : {0, 1, 2})
+	{
+		Server->CreateObject(UTestReplicatedIrisObject::FComponents{});
+	}
+
+	// Create huge object so that the ReplicationWriter will use the huge object path.
+	UTestReplicatedIrisObject* ServerHugeObject = CreateHugeObject(Server);
+
+	// Write packet
+	Server->PreSendUpdate();
+	Server->SendUpdate(Client->ConnectionIdOnServer);
+	Server->PostSendUpdate();
+
+	for (uint32 RetryIt = 0; RetryIt != HugeObjectMaxNetTickCountToArrive; ++RetryIt)
+	{
+		Server->PreSendUpdate();
+		const bool bPacketWasWritten = Server->SendUpdate(Client->ConnectionIdOnServer, TEXT("Create HugeObject"));
+		Server->PostSendUpdate();
+		if (!bPacketWasWritten)
+		{
+			break;
+		}
+	}
+
+	// Deliver object creation packets
+	{
+		const auto& ConnectionInfo = Server->GetConnectionInfo(Client->ConnectionIdOnServer);
+		const SIZE_T PacketCount = ConnectionInfo.WrittenPackets.Count();
+		for (SIZE_T PacketIt = 0; PacketIt != PacketCount; ++PacketIt)
+		{
+			Server->DeliverTo(Client, DeliverPacket);
+		}
+	}
+
+	// The huge object should have been delivered despite the ReplicationRecord starvation
+	UE_NET_ASSERT_TRUE(Client->IsResolvableNetRefHandle(ServerHugeObject->NetRefHandle));
+}
+
 // Below test will fail as we only have a special path for reliable attachments for objects that stopped replicating, not for being filtered out.
 #if 0
 UE_NET_TEST_FIXTURE(FSplitObjectTestFixture, TestReliableAttachmentAddedAfterSplittingHugeObjectIsDeliveredBeforeObjectIsFilteredOut)
diff --git a/Engine/Plugins/Runtime/ReplicationSystemTestPlugin/Source/Private/Tests/ReplicationSystem/TestReplicationRecord.cpp b/Engine/Plugins/Runtime/ReplicationSystemTestPlugin/Source/Private/Tests/ReplicationSystem/TestReplicationRecord.cpp
--- a/Engine/Plugins/Runtime/ReplicationSystemTestPlugin/Source/Private/Tests/ReplicationSystem/TestReplicationRecord.cpp
+++ b/Engine/Plugins/Runtime/ReplicationSystemTestPlugin/Source/Private/Tests/ReplicationSystem/TestReplicationRecord.cpp
@@ -8,7 +8,7 @@ namespace UE::Net::Private
 {
 
 static constexpr FReplicationRecord::ReplicationRecordIndex InvalidReplicationRecordIndex = FReplicationRecord::InvalidReplicationRecordIndex;
-static constexpr FReplicationRecord::ReplicationRecordIndex ReplicationRecordIndexMask = FReplicationRecord::ReplicationRecordIndexMask;
+static constexpr FReplicationRecord::ReplicationRecordIndex MaxReplicationRecordCount = FReplicationRecord::MaxReplicationRecordCount;
 
 UE_NET_TEST(ReplicationRecord, ResetRecordList)
 {
@@ -40,15 +40,15 @@ UE_NET_TEST(ReplicationRecord, PushInfoAndAddToList)
 	UE_NET_ASSERT_EQ(0U, (uint32)List.FirstRecordIndex);
 
 	// Iterate over entries from start to end
-	uint16 CurrentIndex = List.FirstRecordIndex;
-	uint16 CurrentExpectedIndex = 0;
+	FReplicationRecord::ReplicationRecordIndex CurrentIndex = List.FirstRecordIndex;
+	FReplicationRecord::ReplicationRecordIndex CurrentExpectedIndex = 0;
 	
 	while (CurrentIndex != InvalidReplicationRecordIndex)
 	{
 		UE_NET_ASSERT_EQ(CurrentExpectedIndex, CurrentIndex);
 
 		const FReplicationRecord::FRecordInfo* CurrentRecordInfo = ReplicationRecord.GetInfoForIndex(CurrentIndex);
-		UE_NET_ASSERT_TRUE(CurrentRecordInfo != nullptr);
+		UE_NET_ASSERT_NE(CurrentRecordInfo, nullptr);
 
 		CurrentIndex = CurrentRecordInfo->NextIndex;
 		++CurrentExpectedIndex;
@@ -100,7 +100,7 @@ UE_NET_TEST(ReplicationRecord, TestMaxCapacityAndWraparound)
 	ReplicationRecord.ResetList(List);
 
 	// Insert entries in list to max capacity
-	const uint32 InfoCount = ReplicationRecordIndexMask + 1;
+	const uint32 InfoCount = MaxReplicationRecordCount;
 	for (uint32 It = 0; It < InfoCount; ++It)
 	{
 		FReplicationRecord::FRecordInfo Info;
@@ -108,13 +108,13 @@ UE_NET_TEST(ReplicationRecord, TestMaxCapacityAndWraparound)
 		ReplicationRecord.PushInfoAndAddToList(List, Info);
 	}
 
-	UE_NET_ASSERT_EQ(ReplicationRecordIndexMask, List.LastRecordIndex);
-	UE_NET_ASSERT_EQ(0U, (uint32)List.FirstRecordIndex);
+	UE_NET_ASSERT_EQ((uint32)List.LastRecordIndex, MaxReplicationRecordCount - 1U);
+	UE_NET_ASSERT_EQ((uint32)List.FirstRecordIndex, 0U);
 
 	// Pop one so we can test wraparound
 	ReplicationRecord.PopInfoAndRemoveFromList(List);
 	
-	UE_NET_ASSERT_EQ(1U, (uint32)ReplicationRecord.GetFrontIndex());
+	UE_NET_ASSERT_EQ((uint32)ReplicationRecord.GetFrontIndex(), 1U);
 
 	// Push one more
 	FReplicationRecord::FRecordInfo Info;
@@ -125,7 +125,7 @@ UE_NET_TEST(ReplicationRecord, TestMaxCapacityAndWraparound)
 
 	// Access Last pushed entry and verify that it points to what we expect
 	FReplicationRecord::FRecordInfo* LastInfo = ReplicationRecord.GetInfoForIndex(List.LastRecordIndex);
-	UE_NET_ASSERT_TRUE(LastInfo != nullptr);
+	UE_NET_ASSERT_NE(LastInfo, nullptr);
 	UE_NET_ASSERT_EQ(InvalidReplicationRecordIndex, LastInfo->NextIndex);
 }
 
diff --git a/Engine/Plugins/Runtime/ReplicationSystemTestPlugin/Source/Private/Tests/ReplicationSystem/TestReplicationSystem.cpp b/Engine/Plugins/Runtime/ReplicationSystemTestPlugin/Source/Private/Tests/ReplicationSystem/TestReplicationSystem.cpp
--- a/Engine/Plugins/Runtime/ReplicationSystemTestPlugin/Source/Private/Tests/ReplicationSystem/TestReplicationSystem.cpp
+++ b/Engine/Plugins/Runtime/ReplicationSystemTestPlugin/Source/Private/Tests/ReplicationSystem/TestReplicationSystem.cpp
@@ -1,15 +1,17 @@
 // Copyright Epic Games, Inc. All Rights Reserved.
 
 #include "ReplicationSystemServerClientTestFixture.h"
-#include "Iris/Serialization/NetBitStreamReader.h"
-#include "Iris/Serialization/NetBitStreamWriter.h"
-#include "Iris/ReplicationSystem/ReplicationSystem.h"
-#include "Iris/ReplicationSystem/ReplicationSystemInternal.h"
-#include "Net/Core/NetToken/NetToken.h"
+#include "HAL/IConsoleManager.h"
 #include "Iris/Core/IrisLog.h"
 #include "Iris/Metrics/NetMetrics.h"
-#include "Misc/ScopeExit.h"
+#include "Iris/ReplicationSystem/ReplicationRecord.h"
+#include "Iris/ReplicationSystem/ReplicationSystem.h"
+#include "Iris/ReplicationSystem/ReplicationSystemInternal.h"
+#include "Iris/Serialization/NetBitStreamReader.h"
+#include "Iris/Serialization/NetBitStreamWriter.h"
 #include "Logging/LogScopedVerbosityOverride.h"
+#include "Misc/ScopeExit.h"
+#include "Net/Core/NetToken/NetToken.h"
 
 namespace UE::Net::Private
 {
@@ -23,23 +25,19 @@
 	UReplicatedTestObject* ServerObject = Server->CreateObject(0,0);
 
 	// Send and deliver packet
-	Server->PreSendUpdate();
-	Server->SendAndDeliverTo(Client, true);
-	Server->PostSendUpdate();
+	Server->UpdateAndSend({Client});
 
 	// Verify that created server handle now also exists on client
-	UE_NET_ASSERT_TRUE(Client->GetReplicationBridge()->GetReplicatedObject(ServerObject->NetRefHandle) != nullptr);
+	UE_NET_ASSERT_TRUE(Client->IsResolvableNetRefHandle(ServerObject->NetRefHandle));
 
 	// Destroy the spawned object on server
 	Server->DestroyObject(ServerObject);
 
 	// Send and deliver packet
-	Server->PreSendUpdate();
-	Server->SendAndDeliverTo(Client, true);
-	Server->PostSendUpdate();
+	Server->UpdateAndSend({Client});
 
 	// Verify that object now is destroyed on client as well
-	UE_NET_ASSERT_FALSE(Client->GetReplicationBridge()->GetReplicatedObject(ServerObject->NetRefHandle) != nullptr);
+	UE_NET_ASSERT_FALSE(Client->IsValidNetRefHandle(ServerObject->NetRefHandle));
 }
 
 UE_NET_TEST_FIXTURE(FReplicationSystemServerClientTestFixture, TestDroppedDestroyed)
@@ -51,28 +49,22 @@
 	UReplicatedTestObject* ServerObject = Server->CreateObject(0,0);
 
 	// Send and deliver packet
-	Server->PreSendUpdate();
-	Server->SendAndDeliverTo(Client, true);
-	Server->PostSendUpdate();
+	Server->UpdateAndSend({Client});
 
 	// Verify that created server handle now also exists on client
-	UE_NET_ASSERT_TRUE(Client->GetReplicationBridge()->GetReplicatedObject(ServerObject->NetRefHandle) != nullptr);
+	UE_NET_ASSERT_TRUE(Client->IsResolvableNetRefHandle(ServerObject->NetRefHandle));
 
 	// Destroy the spawned object on server
 	Server->DestroyObject(ServerObject);
 
 	// Send and drop packet
-	Server->PreSendUpdate();
-	Server->SendAndDeliverTo(Client, false);
-	Server->PostSendUpdate();
+	Server->UpdateAndSend({Client}, DoNotDeliverPacket);
 
 	// Send and deliver packet
-	Server->PreSendUpdate();
-	Server->SendAndDeliverTo(Client, true);
-	Server->PostSendUpdate();
+	Server->UpdateAndSend({Client});
 
 	// Verify that object now is destroyed on client
-	UE_NET_ASSERT_FALSE(Client->GetReplicationBridge()->GetReplicatedObject(ServerObject->NetRefHandle) != nullptr);
+	UE_NET_ASSERT_FALSE(Client->IsValidNetRefHandle(ServerObject->NetRefHandle));
 }
 
 UE_NET_TEST_FIXTURE(FReplicationSystemServerClientTestFixture, TestDestroyWhilePendingCreate)
@@ -82,31 +74,25 @@
 
 	// Spawn object on server
 	UReplicatedTestObject* ServerObject = Server->CreateObject(0,0);
+	const FNetRefHandle ServerObjectRefHandle = ServerObject->NetRefHandle;
 
-	// Send and deliver packet
-	Server->PreSendUpdate();
-	Server->SendAndDeliverTo(Client, false);
+	// Send and drop packet
+	Server->UpdateAndSend({Client}, DoNotDeliverPacket);
 
-	Server->PostSendUpdate();
-
 	// Destroy
 	Server->DestroyObject(ServerObject);
 
-	// Verify that created server handle does not exists on client
-	UE_NET_ASSERT_FALSE(Client->GetReplicationBridge()->GetReplicatedObject(ServerObject->NetRefHandle) != nullptr);
+	// Verify that the object does not exist on the client
+	UE_NET_ASSERT_FALSE(Client->IsValidNetRefHandle(ServerObjectRefHandle));
 
 	// Send and drop packet
-	Server->PreSendUpdate();
-	Server->SendAndDeliverTo(Client, false);
-	Server->PostSendUpdate();
+	Server->UpdateAndSend({Client}, DoNotDeliverPacket);
 
 	// Send and deliver packet
-	Server->PreSendUpdate();
-	Server->SendAndDeliverTo(Client, true);
-	Server->PostSendUpdate();
+	Server->UpdateAndSend({Client}, DeliverPacket);
 
-	// Verify that object now is destroyed on client
-	UE_NET_ASSERT_FALSE(Client->GetReplicationBridge()->GetReplicatedObject(ServerObject->NetRefHandle) != nullptr);
+	// Verify the object still doesn't exist on the client
+	UE_NET_ASSERT_FALSE(Client->IsValidNetRefHandle(ServerObjectRefHandle));
 }
 
 UE_NET_TEST_FIXTURE(FReplicationSystemServerClientTestFixture, TestDestroyWhileWaitingOnCreate)
@@ -116,8 +102,9 @@
 
 	// Spawn object on server
 	UReplicatedTestObject* ServerObject = Server->CreateObject(0,0);
+	const FNetRefHandle ServerObjectRefHandle = ServerObject->NetRefHandle;
 
-	// Send packet with create
+	// Write packet with create
 	Server->PreSendUpdate();
 	Server->SendTo(Client);
 	Server->PostSendUpdate();
@@ -125,16 +112,22 @@
 	// Destroy while we are waiting for confirmation
 	Server->DestroyObject(ServerObject);
 
-	// Send packet with destroy
+	// Write packet with destroy
 	Server->PreSendUpdate();
 	Server->SendTo(Client);
 	Server->PostSendUpdate();
 
-	// Drop and report packet with create as lost
-	Server->DeliverTo(Client, false);
+	// Drop packet with create
+	Server->DeliverTo(Client, DoNotDeliverPacket);
 
-	// deliver packet with destroy
-	Server->DeliverTo(Client, true);
+	// Verify that the object does not exists on the client
+	UE_NET_ASSERT_FALSE(Client->IsValidNetRefHandle(ServerObjectRefHandle));
+
+	// Deliver packet with destroy
+	Server->DeliverTo(Client, DeliverPacket);
+
+	// Verify the object still doesn't exist on the client
+	UE_NET_ASSERT_FALSE(Client->IsValidNetRefHandle(ServerObjectRefHandle));
 }
 
 UE_NET_TEST_FIXTURE(FReplicationSystemServerClientTestFixture, TestDestroyWithDataInFlight)
@@ -144,34 +137,42 @@
 
 	// Spawn object on server
 	UTestReplicatedIrisObject* ServerObject = Server->CreateObject(0,0);
-	Server->GetReplicationSystem()->SetStaticPriority(ServerObject->NetRefHandle, 1.f);
+	const FNetRefHandle ServerObjectRefHandle = ServerObject->NetRefHandle;
+	Server->GetReplicationSystem()->SetStaticPriority(ServerObjectRefHandle, 1.f);
 
 	// Send packet with create
-	Server->PreSendUpdate();
-	Server->SendAndDeliverTo(Client, true);
-	Server->PostSendUpdate();
+	Server->UpdateAndSend({Client});
 
+	// Verify that the object exists on the client
+	UE_NET_ASSERT_TRUE(Client->IsResolvableNetRefHandle(ServerObjectRefHandle));
+
 	// Modify some data to mark object dirty
 	ServerObject->IntA = 13;
 
-	// Trigger send 
+	// Write a packet with updated data
 	Server->PreSendUpdate();
 	UE_NET_ASSERT_TRUE(Server->SendTo(Client));
 	Server->PostSendUpdate();
 
-	// Destroy while we are waiting for confirmation
+	// Destroy while we are waiting for ack on update
 	Server->DestroyObject(ServerObject);
 
-	// Send packet with destroy
+	// Write packet with destroy
 	Server->PreSendUpdate();
 	UE_NET_ASSERT_TRUE(Server->SendTo(Client));
 	Server->PostSendUpdate();
 
-	// Drop and report packet with create as lost
-	Server->DeliverTo(Client, false);
+	// Drop and report packet with update as lost
+	Server->DeliverTo(Client, DoNotDeliverPacket);
 
-	// deliver packet with destroy
-	Server->DeliverTo(Client, true);
+	// Verify that the object still exists on the client
+	UE_NET_ASSERT_TRUE(Client->IsResolvableNetRefHandle(ServerObjectRefHandle));
+	
+	// Deliver packet with destroy
+	Server->DeliverTo(Client, DeliverPacket);
+
+	// Verify that the object is destroyed on the client
+	UE_NET_ASSERT_FALSE(Client->IsValidNetRefHandle(ServerObjectRefHandle));
 }
 
 UE_NET_TEST_FIXTURE(FReplicationSystemServerClientTestFixture, ReplicateAndDestroySubObject)
@@ -181,29 +182,27 @@
 
 	// Spawn object on server
 	UReplicatedTestObject* ServerObject = Server->CreateObject(0,0);
-
+	const FNetRefHandle ServerObjectRefHandle = ServerObject->NetRefHandle;
 	// Spawn second object on server as a subobject
 	UReplicatedTestObject* ServerSubObject = Server->CreateSubObject(ServerObject->NetRefHandle, 0, 0);
+	const FNetRefHandle ServerSubObjectRefHandle = ServerSubObject->NetRefHandle;
 
 	// Send and deliver packet
-	Server->PreSendUpdate();
-	Server->SendAndDeliverTo(Client, true);
-	Server->PostSendUpdate();
+	Server->UpdateAndSend({Client});
 
-	// Verify that created server handles now also exists on client
-	UE_NET_ASSERT_TRUE(Client->GetReplicationBridge()->GetReplicatedObject(ServerObject->NetRefHandle) != nullptr);
-	UE_NET_ASSERT_TRUE(Client->GetReplicationBridge()->GetReplicatedObject(ServerSubObject->NetRefHandle) != nullptr);
+	// Verify that created server objects now also exist on the client
+	UE_NET_ASSERT_TRUE(Client->IsResolvableNetRefHandle(ServerObjectRefHandle));
+	UE_NET_ASSERT_TRUE(Client->IsResolvableNetRefHandle(ServerSubObjectRefHandle));
 
 	// Destroy the spawned subobject on server
 	Server->DestroyObject(ServerSubObject);
 
 	// Send and deliver packet
-	Server->PreSendUpdate();
-	Server->SendAndDeliverTo(Client, true);
-	Server->PostSendUpdate();
+	Server->UpdateAndSend({Client});
 
-	// Verify that object now is destroyed on client as well
-	UE_NET_ASSERT_FALSE(Client->GetReplicationBridge()->GetReplicatedObject(ServerSubObject->NetRefHandle) != nullptr);
+	// Verify that only the subobject is destroyed on the client
+	UE_NET_ASSERT_FALSE(Client->IsValidNetRefHandle(ServerSubObjectRefHandle));
+	UE_NET_ASSERT_TRUE(Client->IsResolvableNetRefHandle(ServerObjectRefHandle));
 }
 
 UE_NET_TEST_FIXTURE(FReplicationSystemServerClientTestFixture, ReplicateAndDestroyMultipleSubObjects)
@@ -213,51 +212,45 @@
 
 	// Spawn object on server
 	UReplicatedTestObject* ServerObject = Server->CreateObject(0, 0);
-
-	// Spawn second object on server as a subobject
+	const FNetRefHandle ServerObjectRefHandle = ServerObject->NetRefHandle;
+	// Spawn subobject on server
 	UReplicatedTestObject* ServerSubObject = Server->CreateSubObject(ServerObject->NetRefHandle, 0, 0);
+	FNetRefHandle ServerSubObjectRefHandle = ServerSubObject->NetRefHandle;
 
 	// Send and deliver packet
-	Server->PreSendUpdate();
-	Server->SendAndDeliverTo(Client, true);
-	Server->PostSendUpdate();
+	Server->UpdateAndSend({Client});
 
-	// Verify that created server handles now also exists on client
-	UE_NET_ASSERT_NE(Client->GetReplicationBridge()->GetReplicatedObject(ServerObject->NetRefHandle), nullptr);
-	UE_NET_ASSERT_NE(Client->GetReplicationBridge()->GetReplicatedObject(ServerSubObject->NetRefHandle), nullptr);
+	// Verify that created server objects now also exist on the client
+	UE_NET_ASSERT_TRUE(Client->IsResolvableNetRefHandle(ServerObjectRefHandle));
+	UE_NET_ASSERT_TRUE(Client->IsResolvableNetRefHandle(ServerSubObjectRefHandle));
 
 	// Destroy the spawned subobject on server
 	Server->DestroyObject(ServerSubObject);
 
 	// Send and deliver packet
-	Server->PreSendUpdate();
-	Server->SendAndDeliverTo(Client, true);
-	Server->PostSendUpdate();
+	Server->UpdateAndSend({Client});
 
-	// Verify that subobject now is destroyed on client as well
-	UE_NET_ASSERT_EQ(Client->GetReplicationBridge()->GetReplicatedObject(ServerSubObject->NetRefHandle), nullptr);
+	// Verify that the subobject is destroyed on the client
+	UE_NET_ASSERT_FALSE(Client->IsValidNetRefHandle(ServerSubObjectRefHandle));
 
 	// Spawn second object on server as a subobject
 	ServerSubObject = Server->CreateSubObject(ServerObject->NetRefHandle, 0, 0);
+	ServerSubObjectRefHandle = ServerSubObject->NetRefHandle;
 
 	// Send and deliver packet
-	Server->PreSendUpdate();
-	Server->SendAndDeliverTo(Client, true);
-	Server->PostSendUpdate();
+	Server->UpdateAndSend({Client});
 
 	// Verify that second subobject replicated properly to server
-	UE_NET_ASSERT_NE(Client->GetReplicationBridge()->GetReplicatedObject(ServerSubObject->NetRefHandle), nullptr);
+	UE_NET_ASSERT_TRUE(Client->IsResolvableNetRefHandle(ServerSubObjectRefHandle));
 
 	// Destroy the spawned object on server
 	Server->DestroyObject(ServerSubObject);
 
 	// Send and deliver packet
-	Server->PreSendUpdate();
-	Server->SendAndDeliverTo(Client, true);
-	Server->PostSendUpdate();
+	Server->UpdateAndSend({Client});
 
-	// Verify that object now is destroyed on client as well
-	UE_NET_ASSERT_EQ(Client->GetReplicationBridge()->GetReplicatedObject(ServerSubObject->NetRefHandle), nullptr);
+	// Verify that the second subobjects object is destroyed on the client
+	UE_NET_ASSERT_FALSE(Client->IsValidNetRefHandle(ServerSubObjectRefHandle));
 }
 
 UE_NET_TEST_FIXTURE(FReplicationSystemServerClientTestFixture, ReplicateAndDestroySubObjectAndDestroyOwner)
@@ -267,18 +260,17 @@
 
 	// Spawn object on server
 	UReplicatedTestObject* ServerObject = Server->CreateObject(0, 0);
-
-	// Spawn second object on server as a subobject
+	const FNetRefHandle ServerObjectRefHandle = ServerObject->NetRefHandle;
+	// Spawn subobject on server
 	UReplicatedTestObject* ServerSubObject = Server->CreateSubObject(ServerObject->NetRefHandle, 0, 0);
+	const FNetRefHandle ServerSubObjectRefHandle = ServerSubObject->NetRefHandle;
 
 	// Send and deliver packet
-	Server->PreSendUpdate();
-	Server->SendAndDeliverTo(Client, true);
-	Server->PostSendUpdate();
+	Server->UpdateAndSend({Client});
 
-	// Verify that created server handles now also exists on client
-	UE_NET_ASSERT_NE(Client->GetReplicationBridge()->GetReplicatedObject(ServerObject->NetRefHandle), nullptr);
-	UE_NET_ASSERT_NE(Client->GetReplicationBridge()->GetReplicatedObject(ServerSubObject->NetRefHandle), nullptr);
+	// Verify that created server handles now also exist on the client
+	UE_NET_ASSERT_TRUE(Client->IsResolvableNetRefHandle(ServerObjectRefHandle));
+	UE_NET_ASSERT_TRUE(Client->IsResolvableNetRefHandle(ServerSubObjectRefHandle));
 
 	// Destroy the spawned subobject on server
 	Server->DestroyObject(ServerSubObject);
@@ -287,12 +279,13 @@
 	Server->DestroyObject(ServerObject);
 
 	// Send and deliver packet
-	Server->PreSendUpdate();
-	Server->SendAndDeliverTo(Client, true);
-	Server->PostSendUpdate();
+	Server->UpdateAndSend({Client});
 
-	// Verify that subobject now is destroyed on client as well
-	UE_NET_ASSERT_EQ(Client->GetReplicationBridge()->GetReplicatedObject(ServerSubObject->NetRefHandle), nullptr);
+	// Verify that the subobject is destroyed on the client
+	UE_NET_ASSERT_FALSE(Client->IsValidNetRefHandle(ServerSubObjectRefHandle));
+
+	// Verify that the root object is destroyed on the client
+	UE_NET_ASSERT_FALSE(Client->IsValidNetRefHandle(ServerObjectRefHandle));
 }
 
 UE_NET_TEST_FIXTURE(FReplicationSystemServerClientTestFixture, ReplicateAndDestroySubObjectAndDestroyOwnerWithDataInFlight)
@@ -302,37 +295,32 @@
 
 	// Spawn object on server
 	UReplicatedTestObject* ServerObject = Server->CreateObject(0, 0);
-
-	// Spawn second object on server as a subobject
+	const FNetRefHandle ServerObjectRefHandle = ServerObject->NetRefHandle;
+	// Spawn subobject on server
 	UReplicatedTestObject* ServerSubObject = Server->CreateSubObject(ServerObject->NetRefHandle, 0, 0);
+	const FNetRefHandle ServerSubObjectRefHandle = ServerSubObject->NetRefHandle;
 
 	// Send and deliver packet
-	Server->PreSendUpdate();
-	Server->SendAndDeliverTo(Client, true);
-	Server->PostSendUpdate();
+	Server->UpdateAndSend({Client});
 
-	// Verify that created server handles now also exists on client
-	UE_NET_ASSERT_NE(Client->GetReplicationBridge()->GetReplicatedObject(ServerObject->NetRefHandle), nullptr);
-	UE_NET_ASSERT_NE(Client->GetReplicationBridge()->GetReplicatedObject(ServerSubObject->NetRefHandle), nullptr);
+	// Verify that created server handles now also exist on the client
+	UE_NET_ASSERT_TRUE(Client->IsResolvableNetRefHandle(ServerObjectRefHandle));
+	UE_NET_ASSERT_TRUE(Client->IsResolvableNetRefHandle(ServerSubObjectRefHandle));
 
 	// Destroy the spawned subobject on server
 	Server->DestroyObject(ServerSubObject);
 
 	// Send and drop packet
-	Server->PreSendUpdate();
-	Server->SendAndDeliverTo(Client, false);
-	Server->PostSendUpdate();
+	Server->UpdateAndSend({Client}, DoNotDeliverPacket);
 
 	// Destroy owner after we spawned subobject on server
 	Server->DestroyObject(ServerObject);
 
 	// Send and deliver packet
-	Server->PreSendUpdate();
-	Server->SendAndDeliverTo(Client, true);
-	Server->PostSendUpdate();
+	Server->UpdateAndSend({Client}, DeliverPacket);
 
-	// Verify that subobject now is destroyed on client as well
-	UE_NET_ASSERT_EQ(Client->GetReplicationBridge()->GetReplicatedObject(ServerSubObject->NetRefHandle), nullptr);
+	// Verify that the subobject is destroyed on the client
+	UE_NET_ASSERT_FALSE(Client->IsValidNetRefHandle(ServerSubObjectRefHandle));
 }
 
 UE_NET_TEST_FIXTURE(FReplicationSystemServerClientTestFixture, ReplicateAndDestroySubObjectWithLostData)
@@ -342,34 +330,29 @@
 
 	// Spawn object on server
 	UReplicatedTestObject* ServerObject = Server->CreateObject(0, 0);
-
-	// Spawn second object on server as a subobject
+	const FNetRefHandle ServerObjectRefHandle = ServerObject->NetRefHandle;
+	// Spawn subobject on server
 	UReplicatedTestObject* ServerSubObject = Server->CreateSubObject(ServerObject->NetRefHandle, 0, 0);
+	const FNetRefHandle ServerSubObjectRefHandle = ServerSubObject->NetRefHandle;
 
 	// Send and deliver packet
-	Server->PreSendUpdate();
-	Server->SendAndDeliverTo(Client, true);
-	Server->PostSendUpdate();
+	Server->UpdateAndSend({Client});
 
-	// Verify that created server handles now also exists on client
-	UE_NET_ASSERT_NE(Client->GetReplicationBridge()->GetReplicatedObject(ServerObject->NetRefHandle), nullptr);
-	UE_NET_ASSERT_NE(Client->GetReplicationBridge()->GetReplicatedObject(ServerSubObject->NetRefHandle), nullptr);
+	// Verify that created server objects exist on the client
+	UE_NET_ASSERT_TRUE(Client->IsResolvableNetRefHandle(ServerObjectRefHandle));
+	UE_NET_ASSERT_TRUE(Client->IsResolvableNetRefHandle(ServerSubObjectRefHandle));
 
 	// Destroy the spawned subobject on server
 	Server->DestroyObject(ServerSubObject);
 
 	// Send and drop packet
-	Server->PreSendUpdate();
-	Server->SendAndDeliverTo(Client, false);
-	Server->PostSendUpdate();
+	Server->UpdateAndSend({Client}, DoNotDeliverPacket);
 
 	// Send and deliver packet
-	Server->PreSendUpdate();
-	Server->SendAndDeliverTo(Client, true);
-	Server->PostSendUpdate();
+	Server->UpdateAndSend({Client}, DeliverPacket);
 
-	// Verify that subobject now is destroyed on client as well
-	UE_NET_ASSERT_EQ(Client->GetReplicationBridge()->GetReplicatedObject(ServerSubObject->NetRefHandle), nullptr);
+	// Verify that subobject is destroyed on the client
+	UE_NET_ASSERT_FALSE(Client->IsValidNetRefHandle(ServerSubObjectRefHandle));
 }
 
 UE_NET_TEST_FIXTURE(FReplicationSystemServerClientTestFixture, ReplicateAndDestroySubObjectPendingCreateConfirmation)
@@ -379,12 +362,13 @@
 
 	// Spawn object on server
 	UReplicatedTestObject* ServerObject = Server->CreateObject(0, 0);
-
-	// Spawn second object on server as a subobject
+	const FNetRefHandle ServerObjectRefHandle = ServerObject->NetRefHandle;
+	// Spawn subobject on server
 	UReplicatedTestObject* ServerSubObject = Server->CreateSubObject(ServerObject->NetRefHandle, 0, 0);
+	const FNetRefHandle ServerSubObjectRefHandle = ServerSubObject->NetRefHandle;
 
 	uint32 NumUnAcknowledgedPackets = 0;
-	// Send and deliver packet
+	// Write a packet
 	Server->PreSendUpdate();
 	Server->SendTo(Client);
 	Server->PostSendUpdate();
@@ -397,21 +381,143 @@
 	UE_NET_ASSERT_FALSE(Server->SendTo(Client));
 	Server->PostSendUpdate();
 
-	Server->DeliverTo(Client, false);
-	// As we did not send any data we do not have anything to deliver
+	// Drop creation packet
+	Server->DeliverTo(Client, DoNotDeliverPacket);
+
+	// As the second update did not send any data we do not have anything to deliver
 	//Server->DeliverTo(Client, true);
 
-	Server->PreSendUpdate();
-	Server->SendAndDeliverTo(Client, true);
-	Server->PostSendUpdate();
-
 	// Send and deliver packet
 	Server->PreSendUpdate();
 	Server->SendAndDeliverTo(Client, true);
 	Server->PostSendUpdate();
 
-	// Verify that subobject now is destroyed on client as well
-	UE_NET_ASSERT_EQ(Client->GetReplicationBridge()->GetReplicatedObject(ServerSubObject->NetRefHandle), nullptr);
+	// Verify that the subobject does not exist on the client
+	UE_NET_ASSERT_FALSE(Client->IsValidNetRefHandle(ServerSubObjectRefHandle));
+
+	// The root object should exist on the client
+	UE_NET_ASSERT_TRUE(Client->IsResolvableNetRefHandle(ServerObjectRefHandle));
 }
 
+// In this test we're going to create a subobject after the root has been created on the client. Then create it with a bit of latency and destroy root prior to subobject being created on the client.
+UE_NET_TEST_FIXTURE(FReplicationSystemServerClientTestFixture, LateCreatedSubObjectIsDestroyedWithRoot)
+{
+	// Add a client
+	FReplicationSystemTestClient* Client = CreateClient();
+
+	// Spawn object on server
+	UReplicatedTestObject* ServerObject = Server->CreateObject(UTestReplicatedIrisObject::FComponents{});
+	const FNetRefHandle ServerObjectRefHandle = ServerObject->NetRefHandle;
+
+	// Send and deliver packet
+	Server->UpdateAndSend({Client});
+
+	UE_NET_ASSERT_TRUE(Client->IsResolvableNetRefHandle(ServerObjectRefHandle));
+
+	// Spawn subobject
+	UReplicatedTestObject* ServerSubObject = Server->CreateSubObject(ServerObjectRefHandle, UTestReplicatedIrisObject::FComponents{});
+	const FNetRefHandle ServerSubObjectRefHandle = ServerSubObject->NetRefHandle;
+
+	// Put subobject in WaitOnCreateConfirmation state
+	Server->PreSendUpdate();
+	Server->SendUpdate(Client->ConnectionIdOnServer);
+	Server->PostSendUpdate();
+
+	// Destroy the root object
+	Server->DestroyObject(ServerObject);
+
+	// Write a packet
+	Server->PreSendUpdate();
+	Server->SendUpdate(Client->ConnectionIdOnServer);
+	Server->PostSendUpdate();
+
+	// Deliver first packet
+	Server->DeliverTo(Client, DeliverPacket);
+
+	// Verify the root and subobject are created/still created on the client
+	UE_NET_ASSERT_TRUE(Client->IsResolvableNetRefHandle(ServerObjectRefHandle));
+	UE_NET_ASSERT_TRUE(Client->IsResolvableNetRefHandle(ServerSubObjectRefHandle));
+
+	// Deliver second packet
+	Server->DeliverTo(Client, DeliverPacket);
+
+	// Verify the root and subobject are created/still created on the client
+	UE_NET_ASSERT_TRUE(Client->IsResolvableNetRefHandle(ServerObjectRefHandle));
+	UE_NET_ASSERT_TRUE(Client->IsResolvableNetRefHandle(ServerSubObjectRefHandle));
+
+	// Update and send
+	Server->UpdateAndSend({Client});
+
+	// Verify the root and subobject are fully destroyed
+	UE_NET_ASSERT_FALSE(Client->IsValidNetRefHandle(ServerObjectRefHandle));
+	UE_NET_ASSERT_FALSE(Client->IsValidNetRefHandle(ServerSubObjectRefHandle));
+}
+
+// In this test we're going to try destroying an object with thousands of subobjects atomically.
+UE_NET_TEST_FIXTURE(FReplicationSystemServerClientTestFixture, CanDestroyObjectHierarchyAtomically)
+{
+	// Add a client
+	FReplicationSystemTestClient* Client = CreateClient();
+
+	// Spawn object on server
+	UReplicatedTestObject* ServerObject = Server->CreateObject(UTestReplicatedIrisObject::FComponents{});
+	const FNetRefHandle ServerObjectRefHandle = ServerObject->NetRefHandle;
+
+	constexpr unsigned SubObjectCount = 2001;
+	TArray<FNetRefHandle> ServerSubObjectRefHandles;
+	ServerSubObjectRefHandles.Reserve(SubObjectCount);
+
+	// Spawn thousands of subobjects over several frames to avoid huge object path.
+	for (;ServerSubObjectRefHandles.Num() < SubObjectCount;)
+	{
+		constexpr uint32 MaxSubObjectCreationCountPerFrame = 15;
+		for (unsigned It = 0; It < MaxSubObjectCreationCountPerFrame && ServerSubObjectRefHandles.Num() < SubObjectCount; ++It)
+		{
+			UReplicatedTestObject* ServerSubObject = Server->CreateSubObject<UReplicatedTestObject>(ServerObjectRefHandle);
+			ServerSubObjectRefHandles.Add(ServerSubObject->NetRefHandle);
+		}
+
+		Server->UpdateAndSend({Client});
+	}
+
+	// We expect to be done creating the object hierarchy by now. Make sure of it.
+	UE_NET_ASSERT_FALSE(Server->UpdateAndSend({ Client }));
+	UE_NET_ASSERT_TRUE(Client->IsResolvableNetRefHandle(ServerObjectRefHandle));
+	for (const FNetRefHandle SubObjectRefHandle : ServerSubObjectRefHandles)
+	{
+		UE_NET_ASSERT_TRUE(Client->IsResolvableNetRefHandle(SubObjectRefHandle));
+	}
+
+	// Destroy the root object
+	Server->DestroyObject(ServerObject);
+
+	for (unsigned It = 0, MaxTryCount = SubObjectCount/50; It < MaxTryCount; ++It)
+	{
+		const bool bDidSendSomething = Server->UpdateAndSend({ Client });
+
+		// Verify the object hierarchy is destroyed as a whole or not at all. Once we've stopped sending data we should have destroyed the object as a whole on the client.
+		const bool bRootIsResolvable = Client->IsResolvableNetRefHandle(ServerObjectRefHandle);
+		if (bRootIsResolvable && bDidSendSomething)
+		{
+			for (const FNetRefHandle SubObjectRefHandle : ServerSubObjectRefHandles)
+			{
+				UE_NET_ASSERT_TRUE(Client->IsResolvableNetRefHandle(SubObjectRefHandle));
+			}
+		}
+		else
+		{
+			UE_NET_ASSERT_FALSE(Client->IsValidNetRefHandle(ServerObjectRefHandle));
+			for (const FNetRefHandle SubObjectRefHandle : ReverseIterate(ServerSubObjectRefHandles))
+			{
+				UE_NET_ASSERT_FALSE(Client->IsResolvableNetRefHandle(SubObjectRefHandle));
+			}
+		}
+
+		if (!bDidSendSomething)
+		{
+			break;
+		}
+	}
+}
+
 UE_NET_TEST_FIXTURE(FReplicationSystemServerClientTestFixture, TestSubObjectDefaultReplicationOrder)
@@ -2068,6 +2174,67 @@
 	}
 }
 
+UE_NET_TEST_FIXTURE(FReplicationSystemServerClientTestFixture, TestReplicationRecordStarvation)
+{
+	IConsoleVariable* CVarReplicationRecordStarvationThreshold = IConsoleManager::Get().FindConsoleVariable(TEXT("net.Iris.ReplicationWriterReplicationRecordStarvationThreshold"));
+	UE_NET_ASSERT_NE(CVarReplicationRecordStarvationThreshold, nullptr);
+	UE_NET_ASSERT_TRUE(CVarReplicationRecordStarvationThreshold->IsVariableInt());
+	const int32 PrevReplicationRecordStarvationThreshold = CVarReplicationRecordStarvationThreshold->GetInt();
+	ON_SCOPE_EXIT
+	{
+		CVarReplicationRecordStarvationThreshold->Set(PrevReplicationRecordStarvationThreshold, ECVF_SetByCode);
+	};
+	
+	// Add a client
+	FReplicationSystemTestClient* Client = CreateClient();
+
+	// Set starvation threshold to highest possible
+	CVarReplicationRecordStarvationThreshold->Set(FReplicationRecord::MaxReplicationRecordCount, ECVF_SetByCode);
+
+	// Consume one ReplicationRecord to enter starvation
+	UReplicatedTestObject* FirstObject = Server->CreateObject(UTestReplicatedIrisObject::FComponents{});
+	const FNetRefHandle FirstObjectRefHandle = FirstObject->NetRefHandle;
+
+	Server->PreSendUpdate();
+	Server->SendUpdate(Client->ConnectionIdOnServer);
+	Server->PostSendUpdate();
+
+	// Try creating a second object. This should not succeed but we won't be able to test until we've delivered packets. 
+	UReplicatedTestObject* SecondObject = Server->CreateObject(UTestReplicatedIrisObject::FComponents{});
+	const FNetRefHandle SecondObjectRefHandle = SecondObject->NetRefHandle;
+
+	Server->PreSendUpdate();
+	Server->SendUpdate(Client->ConnectionIdOnServer);
+	Server->PostSendUpdate();
+
+	Server->DeliverTo(Client, DeliverPacket);
+	UE_NET_ASSERT_TRUE(Client->IsResolvableNetRefHandle(FirstObjectRefHandle));
+
+	// The second packet, if any, should not allow object replication due to starvation.
+	Server->DeliverTo(Client, DeliverPacket);
+	UE_NET_ASSERT_FALSE(Client->IsResolvableNetRefHandle(SecondObjectRefHandle));
+
+	// Now we should be able replicate objects again. Retry sending the second object.
+	Server->PreSendUpdate();
+	Server->SendUpdate(Client->ConnectionIdOnServer);
+	Server->PostSendUpdate();
+
+	// Try destroying the first object. This should not succeed. 
+	Server->DestroyObject(FirstObject);
+
+	// Deliver the attempt to create the second object and verify it exists on the client.
+	Server->DeliverTo(Client, DeliverPacket);
+	UE_NET_ASSERT_TRUE(Client->IsResolvableNetRefHandle(SecondObjectRefHandle));
+
+	// The second packet, if any, should not allow object destruction due to starvation.
+	Server->DeliverTo(Client, DeliverPacket);
+	UE_NET_ASSERT_TRUE(Client->IsResolvableNetRefHandle(FirstObjectRefHandle));
+
+	// Now we should be able to destroy objects again. Retry destroying the first object
+	Server->UpdateAndSend({Client});
+	UE_NET_ASSERT_FALSE(Client->IsValidNetRefHandle(FirstObjectRefHandle));
+}
+
 } // end namespace UE::Net::Private
 
 
diff --git a/Engine/Source/Runtime/Experimental/Iris/Core/Private/Iris/ReplicationSystem/ReplicationReader.cpp b/Engine/Source/Runtime/Experimental/Iris/Core/Private/Iris/ReplicationSystem/ReplicationReader.cpp
--- a/Engine/Source/Runtime/Experimental/Iris/Core/Private/Iris/ReplicationSystem/ReplicationReader.cpp
+++ b/Engine/Source/Runtime/Experimental/Iris/Core/Private/Iris/ReplicationSystem/ReplicationReader.cpp
@@ -25,6 +25,7 @@
 #include "HAL/IConsoleManager.h"
 #include "UObject/Class.h"
 #include "Algo/RemoveIf.h"
+#include "Misc/ScopeExit.h"
 #include "ProfilingDebugging/CsvProfiler.h"
 
 #if UE_NET_ENABLE_REPLICATIONREADER_LOG
@@ -166,6 +167,7 @@ struct FReplicationReader::FDispatchObjectInfo
 	uint32 bDestroy : 1 = false;
 	uint32 bTearOff : 1 = false;
 	uint32 bDeferredEndReplication : 1 = false;
+	uint32 bProcessEndReplicationForHierarchy : 1 = false;
 	uint32 bShouldCallSubObjectCreatedFromReplication : 1 = false;
 	uint32 bDynamicObjectCreated : 1 = false;
 };
@@ -377,73 +379,178 @@ FNetRefHandle FReplicationReader::ReadNetRefHandleId(FNetSerializationContext& C
 uint32 FReplicationReader::ReadObjectsPendingDestroy(FNetSerializationContext& Context)
 {
 	FNetBitStreamReader& Reader = *Context.GetBitStreamReader();
+	const bool bReadbjectsAndSubObjectsIndividually = Reader.ReadBool();
+	if (bReadbjectsAndSubObjectsIndividually)
+	{
+		return ReadObjectsAndSubObjectsPendingDestroy(Context);
+	}
+	else
+	{
+		return ReadRootObjectsPendingDestroy(Context);
+	}
+}
 
-	UE_NET_TRACE_SCOPE(ObjectsPendingDestroy, Reader, Context.GetTraceCollector(), ENetTraceVerbosity::Trace);
+uint32 FReplicationReader::ReadObjectsAndSubObjectsPendingDestroy(FNetSerializationContext& Context)
+{
+	FNetBitStreamReader& Reader = *Context.GetBitStreamReader();
+
+	UE_NET_TRACE_SCOPE(ObjectsAndSubObjectsPendingDestroy, Reader, Context.GetTraceCollector(), ENetTraceVerbosity::Trace);
 
 	// Read how many destroyed objects we have
-	const uint32 ObjectsToRead = Reader.ReadBits(16);
+	constexpr uint32 DestroyObjectBitCount = 16U;
+	constexpr uint32 MaxDestroyObjectCount = (1U << DestroyObjectBitCount) - 1U;
+	const uint32 ObjectsToRead = Reader.ReadBits(DestroyObjectBitCount);
 	
-	if (!Context.HasErrorOrOverflow())
+	if (Context.HasErrorOrOverflow())
 	{
-		const bool bHasPendingBatches = PendingBatches.GetHasPendingBatches();
+		return 0;
+	}
+
+	const bool bHasPendingBatches = PendingBatches.GetHasPendingBatches();
 	
-		for (uint32 It = 0; It < ObjectsToRead; ++It)
+	for (uint32 It = 0; It < ObjectsToRead; ++It)
+	{
+		UE_NET_TRACE_NAMED_OBJECT_SCOPE(DestroyedObjectScope, FNetRefHandle::GetInvalid(), Reader, Context.GetTraceCollector(), ENetTraceVerbosity::Trace);
+
+		FNetRefHandle IncompleteHandle = ReadNetRefHandleId(Context, Reader);
+		FNetRefHandle SubObjectRootOrHandle = IncompleteHandle;
+		if (Reader.ReadBool())
+		{
+			SubObjectRootOrHandle = ReadNetRefHandleId(Context, Reader);
+		}
+		const bool bShouldDestroyInstance = Reader.ReadBool();
+		if (Context.HasErrorOrOverflow())
+		{
+			break;
+		}
+
+		if (FPendingBatchData* PendingBatchData = bHasPendingBatches ? PendingBatches.Find(SubObjectRootOrHandle) : nullptr)
 		{
-			UE_NET_TRACE_NAMED_OBJECT_SCOPE(DestroyedObjectScope, FNetRefHandle::GetInvalid(), Reader, Context.GetTraceCollector(), ENetTraceVerbosity::Trace);
+			// We expect end replication requests for subobjects too so we should only handle one object at a time.
+			constexpr bool bShouldProcessHierarchy = false;
+			EnqueueEndReplication(PendingBatchData, bShouldDestroyInstance, bShouldProcessHierarchy, IncompleteHandle);
+			continue;
+		}
 
-			FNetRefHandle IncompleteHandle = ReadNetRefHandleId(Context, Reader);
-			FNetRefHandle SubObjectRootOrHandle = IncompleteHandle;
-			if (Reader.ReadBool())
+		{
+			// Resolve handle and destroy using bridge
+			const uint32 InternalIndex = NetRefHandleManager->GetInternalIndex(IncompleteHandle);
+			if (InternalIndex != FNetRefHandleManager::InvalidInternalIndex)
 			{
-				SubObjectRootOrHandle = ReadNetRefHandleId(Context, Reader);
+				UE_NET_TRACE_SET_SCOPE_OBJECTID(DestroyedObjectScope, IncompleteHandle);
+
+				// Defer EndReplication until after applying state data
+				if (bDeferEndReplication)
+				{
+					FDispatchObjectInfo& Info = ObjectsToDispatchArray->AddPendingDispatchObjectInfo(TempLinearAllocator);
+
+					Info.bDestroy = bShouldDestroyInstance;
+					Info.bTearOff = false;
+					Info.bDeferredEndReplication = true;
+					Info.InternalIndex = InternalIndex;
+					Info.bIsInitialState = 0U;
+					Info.bHasState = 0U;
+					Info.bHasAttachments = 0U;
+					Info.bShouldCallSubObjectCreatedFromReplication = 0U;
+
+					// Mark for dispatch
+					ObjectsToDispatchArray->CommitPendingDispatchObjectInfo();
+				}
+				else
+				{
+					EndReplication(InternalIndex, false, bShouldDestroyInstance);
+				}
 			}
-			const bool bShouldDestroyInstance = Reader.ReadBool();
-			if (Context.HasErrorOrOverflow())
+			else
 			{
-				break;
+				// If we did not find the object or associated bridge, the packet that would have created the object may have been lost.
+				UE_LOG(LogIris, Verbose, TEXT("FReplicationReader::Read Tried to destroy object %s. This can occur if the server sends destroy for an object that has not yet been confirmed as created."), *IncompleteHandle.ToString());
 			}
+		}
+	}
 
-			if (FPendingBatchData* PendingBatchData = bHasPendingBatches ? PendingBatches.Find(SubObjectRootOrHandle) : nullptr)
-			{
-				EnqueueEndReplication(PendingBatchData, bShouldDestroyInstance, IncompleteHandle);
-				continue;
-			}
+	return ObjectsToRead;
+}
+
+uint32 FReplicationReader::ReadRootObjectsPendingDestroy(FNetSerializationContext& Context)
+{
+	FNetBitStreamReader& Reader = *Context.GetBitStreamReader();
+
+	UE_NET_TRACE_SCOPE(RootObjectsPendingDestroy, Reader, Context.GetTraceCollector(), ENetTraceVerbosity::Trace);
+
+	// Read how many destroyed objects we have
+	constexpr uint32 DestroyObjectBitCount = 16U;
+	constexpr uint32 MaxDestroyObjectCount = (1U << DestroyObjectBitCount) - 1U;
+	const uint32 ObjectsToRead = Reader.ReadBits(DestroyObjectBitCount);
+	
+	if (Context.HasErrorOrOverflow())
+	{
+		return 0;
+	}
+
+	const bool bHasPendingBatches = PendingBatches.GetHasPendingBatches();
+	
+	for (uint32 It = 0; It < ObjectsToRead; ++It)
+	{
+		UE_NET_TRACE_NAMED_OBJECT_SCOPE(DestroyedObjectScope, FNetRefHandle::GetInvalid(), Reader, Context.GetTraceCollector(), ENetTraceVerbosity::Trace);
+
+		FNetRefHandle IncompleteHandle = ReadNetRefHandleId(Context, Reader);
+		const bool bShouldDestroyInstance = Reader.ReadBool();
+		if (Context.HasErrorOrOverflow())
+		{
+			break;
+		}
+
+		if (FPendingBatchData* PendingBatchData = bHasPendingBatches ? PendingBatches.Find(IncompleteHandle) : nullptr)
+		{
+			constexpr bool bShouldProcessHierarchy = true;
+			EnqueueEndReplication(PendingBatchData, bShouldDestroyInstance, bShouldProcessHierarchy, IncompleteHandle);
+			continue;
+		}
 
+		{
+			// Resolve handle and destroy using bridge
+			const FInternalNetRefIndex InternalIndex = NetRefHandleManager->GetInternalIndex(IncompleteHandle);
+			if (InternalIndex != FNetRefHandleManager::InvalidInternalIndex)
 			{
-				// Resolve handle and destroy using bridge
-				const uint32 InternalIndex = NetRefHandleManager->GetInternalIndex(IncompleteHandle);
-				if (InternalIndex != FNetRefHandleManager::InvalidInternalIndex)
-				{
-					UE_NET_TRACE_SET_SCOPE_OBJECTID(DestroyedObjectScope, IncompleteHandle);
+				UE_NET_TRACE_SET_SCOPE_OBJECTID(DestroyedObjectScope, IncompleteHandle);
 
-					// Defer EndReplication until after applying state data
-					if (bDeferEndReplication)
-					{
-						FDispatchObjectInfo& Info = ObjectsToDispatchArray->AddPendingDispatchObjectInfo(TempLinearAllocator);
-
-						Info.bDestroy = bShouldDestroyInstance;
-						Info.bTearOff = false;
-						Info.bDeferredEndReplication = true;
-						Info.InternalIndex = InternalIndex;
-						Info.bIsInitialState = 0U;
-						Info.bHasState = 0U;
-						Info.bHasAttachments = 0U;
-						Info.bShouldCallSubObjectCreatedFromReplication = 0U;
-
-						// Mark for dispatch
-						ObjectsToDispatchArray->CommitPendingDispatchObjectInfo();
-					}
-					else
-					{
-						EndReplication(InternalIndex, false, bShouldDestroyInstance);
-					}
+				// Defer EndReplication until after applying state data
+				if (bDeferEndReplication)
+				{
+					FDispatchObjectInfo& Info = ObjectsToDispatchArray->AddPendingDispatchObjectInfo(TempLinearAllocator);
+
+					Info.bDestroy = bShouldDestroyInstance;
+					Info.bTearOff = false;
+					Info.bDeferredEndReplication = true;
+					Info.bProcessEndReplicationForHierarchy = true;
+					Info.InternalIndex = InternalIndex;
+					Info.bIsInitialState = 0U;
+					Info.bHasState = 0U;
+					Info.bHasAttachments = 0U;
+					Info.bShouldCallSubObjectCreatedFromReplication = 0U;
+
+					// Mark for dispatch
+					ObjectsToDispatchArray->CommitPendingDispatchObjectInfo();
 				}
 				else
 				{
-					// If we did not find the object or associated bridge, the packet that would have created the object may have been lost.
-					UE_LOG(LogIris, Verbose, TEXT("FReplicationReader::Read Tried to destroy object with %s (This can occur if the server sends destroy for an object that has not yet been confirmed as created)"), *IncompleteHandle.ToString());
+					constexpr bool bTearOff = false;
+					EndReplication(InternalIndex, bTearOff, bShouldDestroyInstance);
+					for (FInternalNetRefIndex SubObjectIndex : NetRefHandleManager->GetSubObjects(InternalIndex))
+					{
+						// End replication for all subobjects and destroy dynamic ones
+						const FNetRefHandle SubObjectRefHandle = NetRefHandleManager->GetNetRefHandleFromInternalIndex(SubObjectIndex);
+						const bool bShouldDestroySubObject = SubObjectRefHandle.IsDynamic();
+						EndReplication(SubObjectIndex, bTearOff, bShouldDestroySubObject);
+					}
 				}
 			}
+			else
+			{
+				// If we did not find the object or associated bridge, the packet that would have created the object may have been lost.
+				UE_LOG(LogIris, Verbose, TEXT("FReplicationReader::Read Tried to destroy object %s. This can occur if the server sends destroy for an object that has not yet been confirmed as created."), *IncompleteHandle.ToString());
+			}
 		}
 	}
 
@@ -996,7 +1103,7 @@ void FReplicationReader::ReadObjectInBatch(FNetSerializationContext& Context, FN
 	else
 	{
 		bHasErrors = Context.HasErrorOrOverflow();
-		UE_CLOG(bHasErrors, LogIris, Error, TEXT("FReplicationReader::ReadObject ErrorOrOverFlow after reading object header %s"), *Context.PrintReadJournal())
+		UE_CLOG(bHasErrors, LogIris, Error, TEXT("FReplicationReader::ReadObject ErrorOrOverFlow after reading object header. Overflow: %c Error: %s Bit stream bits left: %u position: %u %s"), TEXT("YN")[Context.HasError()], ToCStr(Context.GetError().ToString()), Reader.GetBitsLeft(), Reader.GetPosBits(), *Context.PrintReadJournal())
 
 		if (bHasErrors || !IncompleteHandle.IsValid())
 		{
@@ -2192,13 +2299,22 @@ void FReplicationReader::DispatchEndReplication(FNetSerializationContext& Contex
 	{
 		if (Info.bDeferredEndReplication)
 		{
-			// Detach destroy object
+			// Detach and destroy object
 			EndReplication(Info.InternalIndex, Info.bTearOff, Info.bDestroy);
+			if (Info.bProcessEndReplicationForHierarchy)
+			{
+				for (FInternalNetRefIndex SubObjectIndex : NetRefHandleManager->GetSubObjects(Info.InternalIndex))
+				{
+					// End replication for all subobjects and destroy dynamic ones
+					const FNetRefHandle SubObjectRefHandle = NetRefHandleManager->GetNetRefHandleFromInternalIndex(SubObjectIndex);
+					const bool bShouldDestroySubObject = SubObjectRefHandle.IsDynamic();
+					EndReplication(SubObjectIndex, Info.bTearOff, bShouldDestroySubObject);
+				}
+			}
 		}
 	}
 }
 
-
 void FReplicationReader::ReadObjects(FNetSerializationContext& Context, uint32 ObjectBatchCountToRead, uint32 ReadObjectFlags)
 {
 	IRIS_PROFILER_SCOPE(ReplicationReader_ReadObjects);
@@ -2240,6 +2356,21 @@ void FReplicationReader::ProcessHugeObjectAttachment(FNetSerializationContext& C
 
 	UE_NET_TRACE_NAMED_SCOPE(HugeObjectTraceScope, HugeObjectState, HugeObjectReader, HugeObjectTraceCollector, ENetTraceVerbosity::Trace);
 
+	ON_SCOPE_EXIT
+	{
+#if UE_NET_TRACE_ENABLED
+		UE_NET_TRACE_EXIT_NAMED_SCOPE(HugeObjectTraceScope);
+
+		// Append huge object state at end of stream.
+		if (FNetTraceCollector* TraceCollector = Context.GetTraceCollector())
+		{
+			FNetBitStreamReader& Reader = *Context.GetBitStreamReader();
+			// Inject after all other trace events
+			FNetTrace::FoldTraceCollector(TraceCollector, HugeObjectTraceCollector, GetBitStreamPositionForNetTrace(Reader));
+		}
+#endif
+	};
+
 	// Find out how many objects to read so we can reserve object dispatch infos.
 	FNetObjectBlob::FHeader HugeObjectHeader = {};
 	FNetObjectBlob::DeserializeHeader(HugeObjectSerializationContext, HugeObjectHeader);
@@ -2262,25 +2393,13 @@ void FReplicationReader::ProcessHugeObjectAttachment(FNetSerializationContext& C
 		Context.SetError(GNetError_BitStreamError);
 		return;
 	}
-
-#if UE_NET_TRACE_ENABLED
-	UE_NET_TRACE_EXIT_NAMED_SCOPE(HugeObjectTraceScope);
-
-	// Append huge object state at end of stream.
-	if (FNetTraceCollector* TraceCollector = Context.GetTraceCollector())
-	{
-		FNetBitStreamReader& Reader = *Context.GetBitStreamReader();
-		// Inject after all other trace events
-		FNetTrace::FoldTraceCollector(TraceCollector, HugeObjectTraceCollector, GetBitStreamPositionForNetTrace(Reader));
-	}
-#endif
 }
 
-bool FReplicationReader::EnqueueEndReplication(FPendingBatchData* PendingBatchData, bool bShouldDestroyInstance, FNetRefHandle NetRefHandleToEndReplication)
+bool FReplicationReader::EnqueueEndReplication(FPendingBatchData* PendingBatchData, bool bShouldDestroyInstance, bool bShouldProcessHierarchy, FNetRefHandle NetRefHandleToEndReplication)
 {
 	UE_LOG(LogIris, Verbose, TEXT("FReplicationReader::EnqueueEndReplication for %s since %s has queued batches"), *NetRefHandleToEndReplication.ToString(), *PendingBatchData->Handle.ToString());
 		
-	const uint32 MaxNumDataBits = 65U;
+	const uint32 MaxNumDataBits = 66U;
 	const uint32 NumDataWords = (MaxNumDataBits + 31U) / 32U;
 
 	// Enqueue BatchData
@@ -2297,10 +2416,15 @@ bool FReplicationReader::EnqueueEndReplication(FPendingBatchData* PendingBatchDa
 	FNetBitStreamWriter Writer;
 	Writer.InitBytes(PendingBatchData->DataChunkStorage.GetData() + DataChunk.StorageOffset, NumDataWords*sizeof(uint32));
 
-	// Write data
+	// Write data to be parsed by ProcessQueuedBatches
 	WriteUint64(&Writer, NetRefHandleToEndReplication.GetId());
 	Writer.WriteBool(bShouldDestroyInstance);
+	Writer.WriteBool(bShouldProcessHierarchy);
 	Writer.CommitWrites();
+	ensure(Writer.GetPosBits() <= MaxNumDataBits);
+	
+	// Set chunk's actual bit count
+	DataChunk.NumBits = Writer.GetPosBits();
 
 	if (Writer.IsOverflown())
 	{
@@ -2361,6 +2485,11 @@ void FReplicationReader::ProcessQueuedBatches()
 		{
 			UE_LOG(LogIris, Verbose, TEXT("ProcessQueuedBatches processing %d queued batches for Handle %s "), PendingBatchData.QueuedDataChunks.Num(), *PendingBatchData.Handle.ToString());
 
+			// Reset for batch
+			Context.ResetErrorContext();
+			UE_RESET_READ_JOURNAL(Context);
+			UE_ADD_READ_JOURNAL_ENTRY(Context, ToCStr(FString::Printf(TEXT("ProcessQueuedBatches for %s"), ToCStr(PendingBatchData.Handle.ToString()))));
+
 			// Process batched data & dispatch data
 			for (const FQueuedDataChunk& CurrentChunk : PendingBatchData.QueuedDataChunks)
 			{
@@ -2373,6 +2502,7 @@ void FReplicationReader::ProcessQueuedBatches()
 					// this can be the batch root or a subobject owned by the batched root.
 					const uint64 NetRefHandleIdToEndReplication = ReadUint64(&Reader);
 					const bool bShouldDestroyInstance = Reader.ReadBool();
+					const bool bShouldProcessHierarchy = Reader.ReadBool();
 					
 					const FNetRefHandle NetRefHandleToEndReplication = FNetRefHandleManager::MakeNetRefHandleFromId(NetRefHandleIdToEndReplication);
 					
@@ -2380,7 +2510,22 @@ void FReplicationReader::ProcessQueuedBatches()
 					const uint32 InternalIndex = NetRefHandleManager->GetInternalIndex(NetRefHandleToEndReplication);
 					if (InternalIndex != FNetRefHandleManager::InvalidInternalIndex)
 					{
-						EndReplication(InternalIndex, false, bShouldDestroyInstance);
+						constexpr bool bTearOff = false;
+						if (bShouldProcessHierarchy)
+						{
+							EndReplication(InternalIndex, bTearOff, bShouldDestroyInstance);
+							for (const FInternalNetRefIndex SubObjectIndex : NetRefHandleManager->GetSubObjects(InternalIndex))
+							{
+								// End replication for all subobjects and destroy dynamic ones
+								const FNetRefHandle SubObjectRefHandle = NetRefHandleManager->GetNetRefHandleFromInternalIndex(SubObjectIndex);
+								const bool bShouldDestroySubObject = SubObjectRefHandle.IsDynamic();
+								EndReplication(SubObjectIndex, bTearOff, bShouldDestroySubObject);
+							}
+						}
+						else
+						{
+							EndReplication(InternalIndex, bTearOff, bShouldDestroyInstance);
+						}
 					}
 
 					// Remove from broken list
@@ -2414,14 +2559,15 @@ void FReplicationReader::ProcessQueuedBatches()
 
 				if (Context.HasErrorOrOverflow())
 				{
+						// Log error  if get more incoming data for an object in the broken state we will skip it.
+					UE_LOG(LogIris, Error, TEXT("FReplicationReader::ProcessQueuedBatches Failed to process object batch handle: %s skipping batch data. bIsEndReplication: %d"), ToCStr(PendingBatchData.Handle.ToString()), int(CurrentChunk.bIsEndReplicationChunk));
+
+					BrokenObjects.AddUnique(PendingBatchData.Handle);
+
 					if (Context.GetError() == GNetError_BrokenNetHandle)
 					{
 						// $TODO: Report this to the server so it knows that the state of data in the batch is unknown
 
-						// Log error and try to recover, if get more incoming data for an object in the broken state we will skip it.
-						UE_LOG(LogIris, Error, TEXT("FReplicationReader::ProcessQueuedBatches Failed to process object batch handle: %s skipping batch data"), ToCStr(PendingBatchData.Handle.ToString()));
-					
-						BrokenObjects.AddUnique(PendingBatchData.Handle);
 
 						Context.ResetErrorContext();
 					}
diff --git a/Engine/Source/Runtime/Experimental/Iris/Core/Private/Iris/ReplicationSystem/ReplicationReader.h b/Engine/Source/Runtime/Experimental/Iris/Core/Private/Iris/ReplicationSystem/ReplicationReader.h
--- a/Engine/Source/Runtime/Experimental/Iris/Core/Private/Iris/ReplicationSystem/ReplicationReader.h
+++ b/Engine/Source/Runtime/Experimental/Iris/Core/Private/Iris/ReplicationSystem/ReplicationReader.h
@@ -142,8 +142,12 @@ private:
 
 	uint32 ReadObjectsInBatch(FNetSerializationContext& Context, FNetRefHandle InCompleteHandle, bool bHasBatchOwnerData, uint32 BatchEndBitPosition);
 
-	// Read all objects pending destroy
+	// Read objects pending destroy. This will call ReadObjectsAndSubObjectsPendingDestroy or ReadRootObjectsPendingDestroy depending on sending side settings.
 	uint32 ReadObjectsPendingDestroy(FNetSerializationContext& Context);
+	// Receive both objects and subobject destruction. Legacy Iris behavior.
+	uint32 ReadObjectsAndSubObjectsPendingDestroy(FNetSerializationContext& Context);
+	// Receive only root object destruction and destroy all existing subobjects with them atomically.
+	uint32 ReadRootObjectsPendingDestroy(FNetSerializationContext& Context);
 
 	// Read state data for all incoming objects
 	void ReadObjects(FNetSerializationContext& Context, uint32 ObjectCountToRead, uint32 ReadObjectFlags);
@@ -207,7 +211,7 @@ private:
 
 	// If we are queuing data for a batch we must also defer calls to EndReplication
 	// This method writes this method in the form of a QueuedChunk
-	bool EnqueueEndReplication(FPendingBatchData* PendingBatchData, bool bShouldDestroyInstance, FNetRefHandle NetRefHandleToEndReplication);
+	bool EnqueueEndReplication(FPendingBatchData* PendingBatchData, bool bShouldDestroyInstance, bool bShouldProcessHierarchy, FNetRefHandle NetRefHandleToEndReplication);
 
 	// Remove a handle from the hot and cold unresolved caches used by ResolveAndDispatchUnresolvedReferences(). If this handle is marked as unresolved
 	// again, it will be added to the hot cache.
diff --git a/Engine/Source/Runtime/Experimental/Iris/Core/Private/Iris/ReplicationSystem/ReplicationRecord.h b/Engine/Source/Runtime/Experimental/Iris/Core/Private/Iris/ReplicationSystem/ReplicationRecord.h
--- a/Engine/Source/Runtime/Experimental/Iris/Core/Private/Iris/ReplicationSystem/ReplicationRecord.h
+++ b/Engine/Source/Runtime/Experimental/Iris/Core/Private/Iris/ReplicationSystem/ReplicationRecord.h
@@ -2,8 +2,9 @@
 
 #pragma once
 
-#include "Net/Core/Misc/ResizableCircularQueue.h"
+#include "Containers/Array.h"
 #include "Iris/ReplicationSystem/ChangeMaskUtil.h"
+#include "Net/Core/Misc/ResizableCircularQueue.h"
 
 namespace UE::Net::Private
 {
@@ -11,11 +12,14 @@ namespace UE::Net::Private
 class FReplicationRecord
 {
 public:
-	// We should get away with 32k records in flight, given a packet window of 256 and average of 128 replicated objects per packet
-	// We reserve the top bit for marking an index as invalid as it would otherwise require some special logic to inject invalid indices in the array	
+	// We should get away with 64k records in flight given a packet window of 256 and an average of 256 replicated object destroys, objects and subobjects per packet
 	typedef uint16 ReplicationRecordIndex;
-	static constexpr ReplicationRecordIndex InvalidReplicationRecordIndex = 1u << ((sizeof(ReplicationRecordIndex) * 8u) - 1);
-	static constexpr ReplicationRecordIndex ReplicationRecordIndexMask =  InvalidReplicationRecordIndex - 1u;
+	static constexpr ReplicationRecordIndex InvalidReplicationRecordIndex = 65535U;
+	// The MaxReplicationRecordCount should account for at least 256 packets with N records per replicated object or destroyed object. With lots of destroyed objects you could end up with maybe 300-400 records for a single packet.
+	static constexpr ReplicationRecordIndex MaxReplicationRecordCount = InvalidReplicationRecordIndex;
+
+	static constexpr uint32 ObjectIndexBitCount = 20U;
+	static constexpr uint32 ReplicatedObjectStateBitCount = 5U;
 
 	// Circular queue, indexing into record infos
 	struct FRecord
@@ -23,6 +27,17 @@ public:
 		uint16 RecordCount;
 	};
 
+	struct FSubObjectRecord
+	{
+		struct FSubObjectInfo
+		{
+			uint32 Index : ObjectIndexBitCount;
+			uint32 ReplicatedObjectState : ReplicatedObjectStateBitCount;
+		};
+
+		TArray<FSubObjectInfo> SubObjectInfos;
+	};
+
 	// $IRIS: implement some some sort of low overhead chunked fifo array allocator for changemasks used for the replication record
 	// They tend to be relatively short lived and are always allocated and freed in the same order.
 	// for smaller changemasks we use inlined storage.
@@ -31,19 +46,20 @@ public:
 	struct FRecordInfo
 	{
 		FChangeMaskStorageOrPointer ChangeMaskOrPtr;	// used for ChangeMask storage or a pointer to the changemask (could be repurposed to always be an index to save space)
-		uint32 Index : 20;								// Index in the ReplicationInfo array, effectively identifying the object
-		uint32 ReplicatedObjectState : 5;				// Encode EReplicatedObjectState using as few bits as we can
+		uint32 Index : ObjectIndexBitCount;				// Index in the ReplicationInfo array, effectively identifying the object
+		uint32 ReplicatedObjectState : ReplicatedObjectStateBitCount; // Encode EReplicatedObjectState using as few bits as we can
 		uint32 HasChangeMask : 1;						// Do we have a changeMask? 
 		uint32 HasAttachments : 1;						// If this flag is set there's an associated AttachmentRecord
 		uint32 WroteTearOff : 1;						// If this flag is set, we wrote TearOff
 		uint32 WroteDestroySubObject : 1;				// If this flag is set, we wrote DestroySubObject
 		uint32 NewBaselineIndex : 2;					// This is a new Baseline pending ack
-		uint32 UnusedFlag : 1;
+		uint32 HasSubObjectRecord : 1;					// If this flag is set there's an associated SubObjectRecord
 		ReplicationRecordIndex NextIndex;				// Points to next older record index
 		uint16 Padding : 16;
 	};
 
 	static_assert(sizeof(FRecordInfo) == 16, "Expected sizeof FRecordInfo to be 16 bytes");
+	static_assert(sizeof(ReplicationRecordIndex) == 2, "Need to remove or adjust Padding field in FRecordInfo");
 
 public:
 
@@ -62,6 +78,9 @@ public:
 	// Push RecordInfo to record and add it to the provided RecordInfoList
 	inline void PushInfoAndAddToList(FRecordInfoList& RecordList, const FRecordInfo& RecordInfo, uint64 AttachmentRecord = 0);
 
+	// Push RecordInfo to record and add it to the provided RecordInfoList. Call this when there's an associated SubObjectRecord.
+	inline void PushInfoAndAddToList(FRecordInfoList& RecordList, const FRecordInfo& RecordInfo, const FSubObjectRecord& SubObjectRecord);
+
 	// Reset a RecordInfoList
 	inline void ResetList(FRecordInfoList& RecordList);
 
@@ -72,14 +91,18 @@ public:
 	inline FRecordInfo* GetInfoForIndex(ReplicationRecordIndex Index);
 	inline const FRecordInfo* GetInfoForIndex(ReplicationRecordIndex Index) const;
 
-	// PeekInfo . If the info indicates there's an attachment record one should call DequeueAttachmentRecord() as well.
+	// PeekInfo. If the info indicates there's an attachment record one must call DequeueAttachmentRecord(). If the info indicates there's a SubObjectRecord one must call DequeueSubObjectRecord().
 	inline const FRecordInfo& PeekInfo() const { return RecordInfos.Peek(); }
 	inline const FRecordInfo& PeekInfoAtOffset(uint32 Offset) const { return RecordInfos.PeekAtOffset(Offset); }
 	inline const uint32 GetInfoCount() const { return static_cast<uint32>(RecordInfos.Count()); }
+	inline const uint32 GetUnusedInfoCount() const { return static_cast<uint32>(MaxReplicationRecordCount) - static_cast<uint32>(RecordInfos.Count()); }
 
-	// If the info from PeekInfo() indicates there's an attachment record one need to call this function as well.
+	// If the info from PeekInfo() indicates there's an attachment record one needs to call this function as well.
 	uint64 DequeueAttachmentRecord();
 
+	// If the info from PeekInfo() indicates there's a subobject record one needs to call this function as well.
+	FSubObjectRecord DequeueSubObjectRecord();
+
 	// FrontIndex
 	inline ReplicationRecordIndex GetFrontIndex() const { return FrontIndex; }
 
@@ -112,6 +135,9 @@ private:
 	// Storage for attachment records
 	TResizableCircularQueue<uint64> AttachmentRecords;
 
+	// Storage for minimalistic subobject records, used by for example object destruction.
+	TResizableCircularQueue<FSubObjectRecord> SubObjectRecords;
+
 	// Current Index at the oldest Record in the queue, this is used to do relative indexing into the queue when linking pushed records
 	ReplicationRecordIndex FrontIndex;
 };
@@ -120,28 +146,46 @@ FReplicationRecord::FReplicationRecord()
 : RecordInfos(1024)
 , Record(256)
 , AttachmentRecords(64)
+, SubObjectRecords(128)
 , FrontIndex(0u)
 {
 }
 
 FReplicationRecord::FRecordInfo* FReplicationRecord::GetInfoForIndex(ReplicationRecordIndex Index)
 {
-	return (Index != InvalidReplicationRecordIndex) ? &RecordInfos.PokeAtOffsetNoCheck((Index - FrontIndex) & ReplicationRecordIndexMask) : nullptr;
+	if (Index != InvalidReplicationRecordIndex)
+	{
+		const uint32 Offset = uint32(MaxReplicationRecordCount + Index - FrontIndex) % uint32(MaxReplicationRecordCount);
+		return &RecordInfos.PokeAtOffsetNoCheck(Offset);
+	}
+
+	return nullptr;
 }
 
 const FReplicationRecord::FRecordInfo* FReplicationRecord::GetInfoForIndex(ReplicationRecordIndex Index) const
 {
-	return (Index != InvalidReplicationRecordIndex) ? &RecordInfos.PeekAtOffsetNoCheck((Index - FrontIndex) & ReplicationRecordIndexMask) : nullptr;
+	if (Index != InvalidReplicationRecordIndex)
+	{
+		const uint32 Offset = uint32(MaxReplicationRecordCount + Index - FrontIndex) % uint32(MaxReplicationRecordCount);
+		return &RecordInfos.PeekAtOffsetNoCheck(Offset);
+	}
+
+	return nullptr;
 }
 
-// $IRIS TODO. To simplify (?) we could have one public facing struct and one or more internal structs.
-// That would avoid having to call multiple functions.
 inline uint64 FReplicationRecord::DequeueAttachmentRecord()
 {
 	const uint64 AttachmentRecord = AttachmentRecords.Peek();
 	AttachmentRecords.Pop();
 	return AttachmentRecord;
-} 
+}
+
+inline FReplicationRecord::FSubObjectRecord FReplicationRecord::DequeueSubObjectRecord()
+{
+	FSubObjectRecord SubObjectRecord = MoveTemp(SubObjectRecords.Poke());
+	SubObjectRecords.Pop();
+	return SubObjectRecord;
+}
 
 void FReplicationRecord::PushRecord(uint16 InfoCount)
 {
@@ -162,18 +206,18 @@ FReplicationRecord::ReplicationRecordIndex FReplicationRecord::PushInfo(const FR
 { 
 	const SIZE_T CurrentInfoCount = RecordInfos.Count();
 
-	check(CurrentInfoCount <= ReplicationRecordIndexMask);
+	check(CurrentInfoCount < MaxReplicationRecordCount);
 
 	FRecordInfo& NewInfo = RecordInfos.EnqueueDefaulted_GetRef();
 	NewInfo = Info;
 	NewInfo.NextIndex = InvalidReplicationRecordIndex;
 
-	return (FrontIndex + CurrentInfoCount) & ReplicationRecordIndexMask;
+	return (FrontIndex + CurrentInfoCount) % MaxReplicationRecordCount;
 }
 
 void FReplicationRecord::PopInfo()
 { 
-	FrontIndex = (FrontIndex + 1u) & ReplicationRecordIndexMask;
+	FrontIndex = (FrontIndex + 1U) % MaxReplicationRecordCount;
 	RecordInfos.Pop();
 };
 
@@ -212,6 +256,28 @@ void FReplicationRecord::PushInfoAndAddToList(FRecordInfoList& RecordList, const
 	}
 }
 
+void FReplicationRecord::PushInfoAndAddToList(FRecordInfoList& RecordList, const FRecordInfo& RecordInfo, const FSubObjectRecord& SubObjectRecord)
+{
+	FReplicationRecord::ReplicationRecordIndex NewIndex = PushInfo(RecordInfo);
+	if (RecordInfo.HasSubObjectRecord)
+	{
+		SubObjectRecords.Enqueue(SubObjectRecord);
+	}
+
+	if (FRecordInfo* LastRecord = GetInfoForIndex(RecordList.LastRecordIndex))
+	{
+		// Link already in flight record to this newIndex
+		LastRecord->NextIndex = NewIndex;
+		RecordList.LastRecordIndex = NewIndex;
+	}
+	else
+	{
+		// If this is the FirstRecord we update it as well.
+		RecordList.FirstRecordIndex = NewIndex;
+		RecordList.LastRecordIndex = NewIndex;
+	}
+}
+
 void FReplicationRecord::ResetList(FRecordInfoList& RecordList)
 {
 	RecordList.FirstRecordIndex = InvalidReplicationRecordIndex;
diff --git a/Engine/Source/Runtime/Experimental/Iris/Core/Private/Iris/ReplicationSystem/ReplicationWriter.cpp b/Engine/Source/Runtime/Experimental/Iris/Core/Private/Iris/ReplicationSystem/ReplicationWriter.cpp
--- a/Engine/Source/Runtime/Experimental/Iris/Core/Private/Iris/ReplicationSystem/ReplicationWriter.cpp
+++ b/Engine/Source/Runtime/Experimental/Iris/Core/Private/Iris/ReplicationSystem/ReplicationWriter.cpp
@@ -65,9 +65,30 @@ static int32 GReplicationWriterMaxHugeObjectsInTransit = 16;
 static FAutoConsoleVariableRef CVarReplicationWriterMaxHugeObjectsInTransit(TEXT("net.Iris.ReplicationWriterMaxHugeObjectsInTransit"), GReplicationWriterMaxHugeObjectsInTransit,
 	TEXT("How many very large objects, one whose payload doesn't fit in a single packet, is allowed to be scheduled for send. Needs to be at least 1."));
 
+/*
+ * net.Iris.ReplicationWriterMaxDestroyObjectsPerPacket
+ * Limit the number of root objects to be destroyed per packet to relieve pressure on the ReplicationRecord count and potentially for client performance reasons. The number includes both regular object destruction as well as destruction infos.
+ * Without limiting several hundred objects could fit a packet which would exhaust the replication records long before filling the packet window.
+ */
+static int32 GReplicationWriterMaxDestroyObjectsPerPacket = 128;
+static FAutoConsoleVariableRef CVarReplicationWriterMaxDestroyObjectsPerPacket(TEXT("net.Iris.ReplicationWriterMaxDestroyObjectsPerPacket"), GReplicationWriterMaxDestroyObjectsPerPacket,
+	TEXT("How many objects can be replicated for destroy per packet. The number is for regular destruction and destruction infos combined. A value less or equal to zero means unlimited."));
+
+/*
+ * net.Iris.ReplicationWriterReplicationRecordStarvationThreshold
+ * When the number of ReplicationRecords left subceeds this number the ReplicationWriter will limit what is replicated in a packet to a minimum, effectively limiting replication to OOB attachments and huge objects.
+ */
+static int32 GReplicationWriterReplicationRecordStarvationThreshold = 1000;
+static FAutoConsoleVariableRef CVarGReplicationWriterReplicationRecordStarvationThreshold(TEXT("net.Iris.ReplicationWriterReplicationRecordStarvationThreshold"), GReplicationWriterReplicationRecordStarvationThreshold,
+	TEXT("How many ReplicationRecords need to be left in order to proceed with replication as normal. Below this threshold replication will be limited to OOB attachments and huge objects until there are more ReplicationRecords available."));
+
 static bool bValidateObjectsWithDirtyChanges = true;
 static FAutoConsoleVariableRef CvarValidateObjectsWithDirtyChanges(TEXT("net.Iris.ReplicationWriter.ValidateObjectsWithDirtyChanges"), bValidateObjectsWithDirtyChanges, TEXT("Ensure that we don't try to mark invalid objects as dirty when they shouldn't."));
 
+/* NOTE: This cvar is just temporary for testing purposes. Eventually only root objects will replicate destroy and the cvar will be removed. */
+static bool bDestroyRootAndSubObjectsIndividually = false;
+static FAutoConsoleVariableRef CvarDestroyRootAndSubObjectsIndividually(TEXT("net.Iris.DestroyRootAndSubObjectsIndividually"), bDestroyRootAndSubObjectsIndividually, TEXT("TEMPORARY. Whether to send destroy on both root and subobjects (Iris legacy behavior) or just root objects and have the receiving end destroy the subobjects with it. Default is false, only sending destroy for root objects."));
+
 static const FName NetError_ObjectStateTooLarge("Object state is too large to be split.");
 
 const TCHAR* FReplicationWriter::LexToString(const EReplicatedObjectState State)
@@ -194,7 +215,7 @@ struct TReplicationRecordHelper
 	{
 		for (uint32 It = 0; It < RecordInfoCount; ++It)
 		{
- 			const FReplicationRecord::FRecordInfo& RecordInfo = ReplicationRecord->PeekInfo();
+			const FReplicationRecord::FRecordInfo& RecordInfo = ReplicationRecord->PeekInfo();
 			FNetObjectAttachmentsWriter::FReliableReplicationRecord AttachmentRecord(RecordInfo.HasAttachments ? ReplicationRecord->DequeueAttachmentRecord() : uint64(0));
 			FReplicationInfo& Info = ReplicationInfos[RecordInfo.Index];
 			FReplicationRecord::FRecordInfoList& RecordInfoList = ReplicationInfosRecordInfoLists[RecordInfo.Index];
@@ -211,6 +232,24 @@ struct TReplicationRecordHelper
 				FChangeMaskStorageOrPointer::Free(RecordInfo.ChangeMaskOrPtr, ChangeMaskBitCount, s_DefaultChangeMaskAllocator);
 			}
 
+			// Construct a RecordInfo from each subobject info
+			if (RecordInfo.HasSubObjectRecord)
+			{
+				FNetObjectAttachmentsWriter::FReliableReplicationRecord SubObjectAttachmentRecord(0);
+				FReplicationRecord::FRecordInfo SubObjectRecordInfo = {.NewBaselineIndex = FDeltaCompressionBaselineManager::InvalidBaselineIndex};
+				FReplicationRecord::FSubObjectRecord SubObjectRecord = ReplicationRecord->DequeueSubObjectRecord();
+				for (FReplicationRecord::FSubObjectRecord::FSubObjectInfo SubObjectInfo : SubObjectRecord.SubObjectInfos)
+				{
+					FReplicationInfo& SubObjectReplicationInfo = ReplicationInfos[SubObjectInfo.Index];
+
+					SubObjectRecordInfo.Index = SubObjectInfo.Index;
+					SubObjectRecordInfo.ReplicatedObjectState = SubObjectInfo.ReplicatedObjectState;
+
+					// Invoke function
+					Functor(SubObjectRecordInfo, SubObjectReplicationInfo, SubObjectAttachmentRecord);
+				}
+			}
+
 			// We must remove the record and unlink it
 			// It is safe to call even if we have stopped replicating the object		
 			ReplicationRecord->PopInfoAndRemoveFromList(RecordInfoList);
@@ -1490,6 +1529,8 @@ void FReplicationWriter::HandleDroppedRecord<FReplicationWriter::EReplicatedObje
 
 		if (RecordInfo.WroteDestroySubObject && Info.SubObjectPendingDestroy)
 		{
+			// 2024-10-23. Look into enabling this ensure at a proper time. 
+			//ensureMsgf(false, TEXT("Subobjects destroyed individually should not be canceled. Object: %s"), ToCStr(NetRefHandleManager->PrintObjectFromIndex(InternalIndex)));
 			// If the subobject owner still is replicated and valid
 			FNetRefHandleManager::FReplicatedObjectData& ObjectData = NetRefHandleManager->GetReplicatedObjectDataNoCheck(InternalIndex);
 			check(ObjectData.IsSubObject());
@@ -1735,7 +1776,7 @@ void FReplicationWriter::CreateObjectRecord(const FNetBitArrayView* ChangeMask,
 	RecordInfo.HasAttachments = (OutRecord.AttachmentRecord.IsValid() ? 1U : 0U);
 	RecordInfo.WroteTearOff = ObjectInfo.bSentTearOff;
 	RecordInfo.WroteDestroySubObject = Info.SubObjectPendingDestroy;
-	
+	RecordInfo.HasSubObjectRecord = 0U;
 	// If we wrote a new baseline we need to store it in the record
 	RecordInfo.NewBaselineIndex = ObjectInfo.bSentState ? ObjectInfo.NewBaselineIndex : FDeltaCompressionBaselineManager::InvalidBaselineIndex;
 
@@ -1760,6 +1801,11 @@ void FReplicationWriter::CommitObjectRecord(uint32 InternalObjectIndex, const FO
 	ReplicationRecord.PushInfoAndAddToList(ReplicatedObjectsRecordInfoLists[InternalObjectIndex], ObjectRecord.Record, ObjectRecord.AttachmentRecord.ToUint64());
 }
 
+void FReplicationWriter::CommitObjectDestroyRecord(uint32 InternalObjectIndex, const FObjectRecord& ObjectRecord, const FReplicationRecord::FSubObjectRecord& SubObjectRecord)
+{
+	ReplicationRecord.PushInfoAndAddToList(ReplicatedObjectsRecordInfoLists[InternalObjectIndex], ObjectRecord.Record, SubObjectRecord);
+}
+
 void FReplicationWriter::CommitBatchRecord(const FBatchRecord& BatchRecord)
 {
 	for (const FObjectRecord& ObjectRecord : BatchRecord.ObjectReplicationRecords)
@@ -1769,26 +1815,49 @@ void FReplicationWriter::CommitBatchRecord(const FBatchRecord& BatchRecord)
 }
 
 uint32 FReplicationWriter::WriteObjectsPendingDestroy(FNetSerializationContext& Context)
+{
+	FNetBitStreamWriter& Writer = *Context.GetBitStreamWriter();
+	const bool bWriteObjectsAndSubObjectsIndividually = Writer.WriteBool(bDestroyRootAndSubObjectsIndividually);
+	if (Writer.IsOverflown())
+	{
+		return 0;
+	}
+
+	if (bWriteObjectsAndSubObjectsIndividually)
+	{
+		return WriteObjectsAndSubObjectsPendingDestroy(Context);
+	}
+	else
+	{
+		return WriteRootObjectsPendingDestroy(Context);
+	}
+}
+
+uint32 FReplicationWriter::WriteObjectsAndSubObjectsPendingDestroy(FNetSerializationContext& Context)
 {
 	FNetBitStreamWriter& Writer = *Context.GetBitStreamWriter();
 
-	UE_NET_TRACE_SCOPE(ObjectsPendingDestroy, Writer, Context.GetTraceCollector(), ENetTraceVerbosity::Trace);
+	UE_NET_TRACE_SCOPE(ObjectsAndSubObjectsPendingDestroy, Writer, Context.GetTraceCollector(), ENetTraceVerbosity::Trace);
 
-	// Write bit indicating that we did write, we also need one trailing bit to indicate the end
-	uint16 WrittenCount = 0;
+	uint32 WrittenCount = 0;
 
 	// Write how many destroyed objects we have
 	const uint32 HeaderPos = Writer.GetPosBits();
-	Writer.WriteBits(WrittenCount, 16);
+	constexpr uint32 DestroyObjectBitCount = 16U;
+	constexpr uint32 MaxDestroyObjectCount = (1U << DestroyObjectBitCount) - 1U;
 
-	if (Writer.IsOverflown())
+	Writer.WriteBits(WrittenCount, DestroyObjectBitCount);
+
+	// Can't write if bitstream is overflown. Shouldn't write if we're in replication record starvation mode.
+	if (Writer.IsOverflown() || WriteContext.bIsInReplicationRecordStarvation)
 	{
 		return 0U;
 	}
 
 	bool bWroteAllDestroyedObjects = true;
 
-	for (uint32 InternalIndex = 0U; (InternalIndex = ObjectsPendingDestroy.FindFirstOne(InternalIndex + 1U)) != FNetBitArray::InvalidIndex; )
+	// Note: This code is not honoring net.Iris.MaxDestroyObjectsPerPacket because it's bug prone to split root and subobjects in different packets, which is already a possibility with this code. Set net.Iris.DestroyRootAndSubObjectsIndividually 0 to get proper destruction throttling.
+	for (uint32 InternalIndex = 0U; ((InternalIndex = ObjectsPendingDestroy.FindFirstOne(InternalIndex + 1U)) != FNetBitArray::InvalidIndex) && (WrittenCount < MaxDestroyObjectCount); )
 	{
 		FReplicationInfo& Info = GetReplicationInfo(InternalIndex);
 		const FNetRefHandleManager::FReplicatedObjectData& ObjectData = NetRefHandleManager->GetReplicatedObjectDataNoCheck(InternalIndex);
@@ -1895,8 +1964,9 @@ uint32 FReplicationWriter::WriteObjectsPendingDestroy(FNetSerializationContext&
 
 			// update transmission record
 			FBatchObjectInfo ObjectInfo = {};
-			ObjectInfo.InternalIndex = InternalIndex;
 			FObjectRecord ObjectRecord;
+
+			ObjectInfo.InternalIndex = InternalIndex;
 			CreateObjectRecord(nullptr, Info, ObjectInfo, ObjectRecord);
 			CommitObjectRecord(InternalIndex, ObjectRecord);
 	
@@ -1904,7 +1974,6 @@ uint32 FReplicationWriter::WriteObjectsPendingDestroy(FNetSerializationContext&
 		}
 		else
 		{
-			bWroteAllDestroyedObjects = false;
 			break;
 		}
 	}
@@ -1912,14 +1981,196 @@ uint32 FReplicationWriter::WriteObjectsPendingDestroy(FNetSerializationContext&
 	// Write Header
 	{
 		FNetBitStreamWriteScope WriteScope(Writer, HeaderPos);
-		Writer.WriteBits(WrittenCount, 16);
+		Writer.WriteBits(WrittenCount, DestroyObjectBitCount);
 	}
 
+	bWroteAllDestroyedObjects = bWroteAllDestroyedObjects && !Writer.IsOverflown() && (WrittenCount < MaxDestroyObjectCount);
 	WriteContext.bHasDestroyedObjectsToSend = !bWroteAllDestroyedObjects;
 
 	return WrittenCount;
 }
 
+uint32 FReplicationWriter::WriteRootObjectsPendingDestroy(FNetSerializationContext& Context)
+{
+	FNetBitStreamWriter& Writer = *Context.GetBitStreamWriter();
+
+	UE_NET_TRACE_SCOPE(RootObjectsPendingDestroy, Writer, Context.GetTraceCollector(), ENetTraceVerbosity::Trace);
+
+	uint32 WrittenCount = 0;
+
+	// Write how many destroyed objects we have
+	const uint32 HeaderPos = Writer.GetPosBits();
+	constexpr uint32 DestroyObjectBitCount = 16U;
+	constexpr uint32 MaxDestroyObjectCount = (1U << DestroyObjectBitCount) - 1U;
+
+	Writer.WriteBits(WrittenCount, DestroyObjectBitCount);
+
+	// Can't write if bitstream is overflown. Shouldn't write if we're in replication record starvation mode.
+	if (Writer.IsOverflown() || WriteContext.bIsInReplicationRecordStarvation)
+	{
+		return 0U;
+	}
+
+	auto AreSubObjectsReadyToBeDestroyed = [this](uint32 InternalIndex)
+	{
+		for (uint32 SubObjectIndex : NetRefHandleManager->GetSubObjects(InternalIndex))
+		{
+			const FReplicationInfo& SubObjectInfo = GetReplicationInfo(SubObjectIndex);
+			const EReplicatedObjectState SubObjectState = SubObjectInfo.GetState();
+			// A subobject may have stopped replicating immediately when going out of scope if it was never replicated to begin with.
+			if (SubObjectState == EReplicatedObjectState::Invalid)
+			{
+				continue;
+			}
+
+			if (!(SubObjectState == EReplicatedObjectState::PendingDestroy || SubObjectState == EReplicatedObjectState::SubObjectPendingDestroy))
+			{
+				return false;
+			}
+
+			if (!SubObjectInfo.IsCreationConfirmed)
+			{
+				return false;
+			}
+		}
+
+		return true;
+	};
+
+	bool bWroteAllDestroyedObjects = true;
+
+	const uint32 EffectiveMaxDestroyObjectCount = FPlatformMath::Min((GReplicationWriterMaxDestroyObjectsPerPacket > 0 ? static_cast<uint32>(GReplicationWriterMaxDestroyObjectsPerPacket) : MaxDestroyObjectCount), MaxDestroyObjectCount);
+	for (uint32 InternalIndex = 0U; ((InternalIndex = ObjectsPendingDestroy.FindFirstOne(InternalIndex + 1U)) != FNetBitArray::InvalidIndex) && (WrittenCount < EffectiveMaxDestroyObjectCount); )
+	{
+		FReplicationInfo& Info = GetReplicationInfo(InternalIndex);
+		const FNetRefHandleManager::FReplicatedObjectData& ObjectData = NetRefHandleManager->GetReplicatedObjectDataNoCheck(InternalIndex);
+
+		if (ObjectData.IsSubObject())
+		{
+			// Have the root object query status of subobjects
+			continue;
+		}
+
+		// Don't send destroy until object creation has been acked.
+		if (!Info.IsCreationConfirmed)
+		{
+			continue;
+		}
+
+		// Already waiting on destroy confirmation
+		if (Info.GetState() == EReplicatedObjectState::WaitOnDestroyConfirmation)
+		{
+			continue;
+		}
+
+		if (Info.GetState() == EReplicatedObjectState::WaitOnFlush)
+		{
+			if (GetFlushStatus(InternalIndex, Info, Info.FlushFlags) != EFlushFlags::FlushFlags_None)
+			{
+				continue;
+			}
+			
+			// Object *and* subobjects are now flushed and can be destroyed
+			SetPendingDestroyOrSubObjectPendingDestroyState(InternalIndex, Info);
+
+			// Change state for subobjects waiting on flush
+			for (uint32 SubObjectIndex : NetRefHandleManager->GetSubObjects(InternalIndex))
+			{
+				FReplicationInfo& SubObjectInfo = GetReplicationInfo(SubObjectIndex);
+				if (SubObjectInfo.GetState() == EReplicatedObjectState::WaitOnFlush)
+				{
+					SetPendingDestroyOrSubObjectPendingDestroyState(SubObjectIndex, SubObjectInfo);
+				}
+			}
+		}
+
+		if (!AreSubObjectsReadyToBeDestroyed(InternalIndex))
+		{
+			bWroteAllDestroyedObjects = false;
+			continue;
+		}
+		
+		// Unexpected. Get more info.
+		if (Info.GetState() != EReplicatedObjectState::PendingDestroy)
+		{
+			ensureMsgf(Info.GetState() == EReplicatedObjectState::PendingDestroy, TEXT("Skipping writing destroy for object %s which is in unexpected state %s. IsSubObject: %u IsDestructionInfo: %u"), ToCStr(NetRefHandleManager->PrintObjectFromIndex(InternalIndex)), LexToString(Info.GetState()), Info.IsSubObject, Info.IsDestructionInfo);
+			continue;
+		}
+
+		// We do not support destroying an object that is currently being sent as a huge object.
+		if (IsObjectPartOfActiveHugeObject(InternalIndex))
+		{
+			UE_LOG(LogIris, Verbose, TEXT("Skipping writing destroy for object ( InternalIndex: %u ) which is part of active huge object."), InternalIndex);
+			bWroteAllDestroyedObjects = false;
+			continue;
+		}
+
+		UE_NET_TRACE_OBJECT_SCOPE(ObjectData.RefHandle,  Writer, Context.GetTraceCollector(), ENetTraceVerbosity::Trace);
+
+		FNetBitStreamRollbackScope RollbackScope(Writer);
+
+		// Write handle with the needed bitCount
+		WriteNetRefHandleId(Context, ObjectData.RefHandle);
+
+		// Write bit indicating if the static instance should be destroyed or not (could skip the bit for dynamic objects)
+		const bool bShouldDestroyInstance = ObjectData.RefHandle.IsDynamic() || NetRefHandleManager->GetIsDestroyedStartupObject(InternalIndex);
+		Writer.WriteBool(bShouldDestroyInstance);
+
+		if (!Writer.IsOverflown())
+		{
+			// Must update state before pushing record
+			SetState(InternalIndex, EReplicatedObjectState::WaitOnDestroyConfirmation);
+			Info.HasDirtyChangeMask = 0U;
+
+			// update transmission record
+			FBatchObjectInfo ObjectInfo = {};
+			FObjectRecord ObjectRecord;
+			FReplicationRecord::FSubObjectRecord SubObjectRecord;
+
+			ObjectInfo.InternalIndex = InternalIndex;
+			CreateObjectRecord(nullptr, Info, ObjectInfo, ObjectRecord);
+
+			// Fill in subobject record
+			for (const uint32 SubObjectIndex : NetRefHandleManager->GetSubObjects(InternalIndex))
+			{
+				const FReplicationInfo& SubObjectInfo = GetReplicationInfo(SubObjectIndex);
+				if (SubObjectInfo.GetState() == EReplicatedObjectState::Invalid)
+				{
+					continue;
+				}
+
+				SetState(SubObjectIndex, EReplicatedObjectState::WaitOnDestroyConfirmation);
+				Info.HasDirtyChangeMask = 0U;
+				Info.SubObjectPendingDestroy = 0U;
+
+				FReplicationRecord::FSubObjectRecord::FSubObjectInfo& SubObjectRecordInfo = SubObjectRecord.SubObjectInfos.AddDefaulted_GetRef();
+				SubObjectRecordInfo.Index = SubObjectIndex;
+				SubObjectRecordInfo.ReplicatedObjectState = static_cast<uint32>(Info.GetState());
+			}
+
+			ObjectRecord.Record.HasSubObjectRecord = !SubObjectRecord.SubObjectInfos.IsEmpty();
+			CommitObjectDestroyRecord(InternalIndex, ObjectRecord, SubObjectRecord);
+
+			++WrittenCount;
+		}
+		else
+		{
+			break;
+		}
+	}
+
+	// Write Header
+	{
+		FNetBitStreamWriteScope WriteScope(Writer, HeaderPos);
+		Writer.WriteBits(WrittenCount, DestroyObjectBitCount);
+	}
+
+	bWroteAllDestroyedObjects = bWroteAllDestroyedObjects && !Writer.IsOverflown() && (WrittenCount < MaxDestroyObjectCount);
+	WriteContext.bHasDestroyedObjectsToSend = !bWroteAllDestroyedObjects;
+	WriteContext.WrittenDestroyObjectCount += WrittenCount;
+	return WrittenCount;
+}
+
 bool FReplicationWriter::CanSendObject(uint32 InternalIndex) const
 {
 	const FReplicationInfo& Info = GetReplicationInfo(InternalIndex);
@@ -2886,6 +3137,11 @@ int FReplicationWriter::WriteDestructionInfo(FNetSerializationContext& Context,
 {
 	IRIS_PROFILER_SCOPE(FReplicationWriter_WriteDestructionInfo);
 
+	if ((GReplicationWriterMaxDestroyObjectsPerPacket > 0) && (WriteContext.WrittenDestroyObjectCount >= static_cast<uint32>(GReplicationWriterMaxDestroyObjectsPerPacket)))
+	{
+		return 0;
+	}
+
 	FNetBitStreamWriter& Writer = *Context.GetBitStreamWriter();
 	UE_NET_TRACE_SCOPE(Batch, Writer, Context.GetTraceCollector(), ENetTraceVerbosity::Trace);
 
@@ -2946,6 +3202,9 @@ int FReplicationWriter::WriteDestructionInfo(FNetSerializationContext& Context,
 
 		// We count this as an object batch
 		++WriteContext.WrittenBatchCount;
+
+		// We also count it as an object destroy.
+		++WriteContext.WrittenDestroyObjectCount;
 	}
 
 	return Writer.IsOverflown() ? -1 :  1;
@@ -2957,7 +3216,7 @@ uint32 FReplicationWriter::WriteOOBAttachments(FNetSerializationContext& Context
 
 	if (WriteContext.WriteMode == EDataStreamWriteMode::PostTickDispatch)
 	{
-		if (WriteContext.bHasOOBAttachmentsToSend && CanSendObject(ObjectIndexForOOBAttachment))
+		if (WriteContext.bHasOOBAttachmentsToSend && CanSendObject(ObjectIndexForOOBAttachment) && (ReplicationRecord.GetUnusedInfoCount() > 0U))
 		{
 			IRIS_PROFILER_SCOPE(FReplicationWriter_WriteOOBAttachments);
 			const int32 Result = WriteObjectBatch(Context, ObjectIndexForOOBAttachment, WriteObjectFlag_Attachments);
@@ -2972,7 +3231,7 @@ uint32 FReplicationWriter::WriteOOBAttachments(FNetSerializationContext& Context
 	}
 	else
 	{
-		if (WriteContext.bHasHugeObjectToSend)
+		if (WriteContext.bHasHugeObjectToSend && (ReplicationRecord.GetUnusedInfoCount() > 0U))
 		{
 			IRIS_PROFILER_SCOPE(FReplicationWriter_WriteHugeObjectAttachments);
 			const int32 Result = WriteObjectBatch(Context, ObjectIndexForOOBAttachment, WriteObjectFlag_Attachments | WriteObjectFlag_HugeObject);
@@ -2991,7 +3250,7 @@ uint32 FReplicationWriter::WriteOOBAttachments(FNetSerializationContext& Context
 			WrittenObjectCount += Result;
 		}
 
-		if (WriteContext.bHasOOBAttachmentsToSend && CanSendObject(ObjectIndexForOOBAttachment))
+		if (WriteContext.bHasOOBAttachmentsToSend && CanSendObject(ObjectIndexForOOBAttachment) && (ReplicationRecord.GetUnusedInfoCount() > 0U))
 		{
 			IRIS_PROFILER_SCOPE(FReplicationWriter_WriteOOBAttachments);
 			const int32 Result = WriteObjectBatch(Context, ObjectIndexForOOBAttachment, WriteObjectFlag_Attachments);
@@ -3018,7 +3277,8 @@ uint32 FReplicationWriter::WriteObjects(FNetSerializationContext& Context)
 	uint32 ObjectListIt = WriteContext.CurrentIndex;
 	uint32 SortedCount = WriteContext.SortedObjectCount;
 
-	bool bContinue = WriteContext.bHasUpdatedObjectsToSend;
+	// Write only if there are objects to send and we're not in replication record starvation mode.
+	bool bContinue = WriteContext.bHasUpdatedObjectsToSend && !WriteContext.bIsInReplicationRecordStarvation;
 
 	auto SendObjectFunction = [this, &Context, &WrittenObjectCount](FInternalNetRefIndex InternalIndex)
 	{
@@ -3300,6 +3560,12 @@ UDataStream::EWriteResult FReplicationWriter::BeginWrite(const UDataStream::FBeg
 		return UDataStream::EWriteResult::NoData;
 	}
 
+	// If we've run out of replication records we cannot send anything
+	if (ReplicationRecord.GetUnusedInfoCount() == 0)
+	{
+		return UDataStream::EWriteResult::NoData;
+	}
+
 	// Initialize context which can be used over multiple calls to WriteData
 	WriteContext.bHasUpdatedObjectsToSend  = 0U;
 	WriteContext.bHasDestroyedObjectsToSend = 0U;
@@ -3357,6 +3623,9 @@ UDataStream::EWriteResult FReplicationWriter::BeginWrite(const UDataStream::FBeg
 	WriteContext.NumWrittenPacketsInThisBatch = 0U;
 	WriteContext.bCanWriteMoreData = Params.bCanWriteMoreData;
 
+	// Set in Write
+	WriteContext.bIsInReplicationRecordStarvation = 0;
+
 	// Clear net stats. Used for CVS and Network Insights stats.
 	WriteContext.Stats.Reset();
 	WriteContext.Stats.SetNumberOfRootObjectsScheduledForReplication(WriteContext.ScheduledObjectCount);
@@ -3431,6 +3700,9 @@ UDataStream::EWriteResult FReplicationWriter::Write(FNetSerializationContext& Co
 	// We have somethings in the WriteContext that must be reset each packet
 	WriteContext.FailedToWriteSmallObjectCount = 0U;
 	WriteContext.WrittenBatchCount = 0U;
+	WriteContext.WrittenDestroyObjectCount = 0U;
+	// Check whether we are running low on replication records or not and need to throttle object replication.
+	WriteContext.bIsInReplicationRecordStarvation = ReplicationRecord.GetUnusedInfoCount() < static_cast<uint32>(GReplicationWriterReplicationRecordStarvationThreshold);
 
 	FNetBitStreamWriter& Writer = *Context.GetBitStreamWriter();
 
@@ -3482,7 +3754,7 @@ UDataStream::EWriteResult FReplicationWriter::Write(FNetSerializationContext& Co
 		{
 			WriteResult = UDataStream::EWriteResult::HasMoreData;
 		}
-		else if (WriteContext.bCanWriteMoreData && ((int32)WriteContext.NumWrittenPacketsInThisBatch < GReplicationWriterMaxAllowedPacketsIfNotHugeObject))
+		else if (!WriteContext.bIsInReplicationRecordStarvation && WriteContext.bCanWriteMoreData && ((int32)WriteContext.NumWrittenPacketsInThisBatch < GReplicationWriterMaxAllowedPacketsIfNotHugeObject))
 		{
 			WriteResult = UDataStream::EWriteResult::HasMoreData;
 		}
diff --git a/Engine/Source/Runtime/Experimental/Iris/Core/Private/Iris/ReplicationSystem/ReplicationWriter.h b/Engine/Source/Runtime/Experimental/Iris/Core/Private/Iris/ReplicationSystem/ReplicationWriter.h
--- a/Engine/Source/Runtime/Experimental/Iris/Core/Private/Iris/ReplicationSystem/ReplicationWriter.h
+++ b/Engine/Source/Runtime/Experimental/Iris/Core/Private/Iris/ReplicationSystem/ReplicationWriter.h
@@ -236,16 +236,26 @@ private:
 		// How many objects that were attempted to be replicated but which ultimately didn't fit in the packet.
 		uint32 FailedToWriteSmallObjectCount;
 
+		// How many root object destroys and destruction infos have been replicated.
+		uint32 WrittenDestroyObjectCount;
+
 		// How many packets have we written to?
 		uint32 NumWrittenPacketsInThisBatch = 0U;
 		bool bCanWriteMoreData = false;
 		
 		EDataStreamWriteMode WriteMode = EDataStreamWriteMode::Full;	
 	
+		// Whether we're running low on ReplicationRecords or not. If starving only OOB and huge objects are sent, if at all possible.
+		uint32 bIsInReplicationRecordStarvation : 1;
+		// Whether there are destroyed objects to send or not
 		uint32 bHasDestroyedObjectsToSend : 1;
+		// Whether there are dirty objects to send or not
 		uint32 bHasUpdatedObjectsToSend : 1;
+		// Whether there is at least one huge object to send or not
 		uint32 bHasHugeObjectToSend : 1;
+		// Whether there are OOB attachments to send or not
 		uint32 bHasOOBAttachmentsToSend : 1;
+		// Whether this context is valid or not
 		uint32 bIsValid : 1;
 
 		FNetSendStats Stats;
@@ -444,6 +454,8 @@ private:
 	
 	// Push and link info for written object to ReplicationRecord
 	void CommitObjectRecord(uint32 InternalObjectIndex, const FObjectRecord& Record);
+	// Push and link info for destroyed root object
+	void CommitObjectDestroyRecord(uint32 InternalObjectIndex, const FObjectRecord& ObjectRecord, const FReplicationRecord::FSubObjectRecord& SubObjectRecord);
 
 	void CommitBatchRecord(const FBatchRecord& BatchRecord);
 
@@ -454,8 +466,12 @@ private:
 	// Partial sort of OutScheduledObjectIndices, will sort at most PartialSortObjectCount objects
 	uint32 SortScheduledObjects(FScheduleObjectInfo* ScheduledObjectIndices, uint32 ScheduledObjectCount, uint32 StartIndex);
 
-	// Write all objects pending destroy (or as many as we fit in the current packet)
+	// Write all objects pending destroy (or as many as we fit in the current packet). Will call one of WriteObjectsAndSubObjectsPendingDestroy and WriteRootObjectsPendingDestroy depending on cvar net.Iris.DestroyRootAndSubObjectsIndividually
 	uint32 WriteObjectsPendingDestroy(FNetSerializationContext& Context);
+	// Used when net.Iris.DestroyRootAndSubObjectsIndividually is true. Will send destroy for every root and subobject. Leagcy Iris behavior.
+	uint32 WriteObjectsAndSubObjectsPendingDestroy(FNetSerializationContext& Context);
+	// Used when net.Iris.DestroyRootAndSubObjectsIndividually is false. Will only send destroy for root objects and subobjects will be gathered on the client side and destroyed. This guarantees atomic object hierarchy destroys.
+	uint32 WriteRootObjectsPendingDestroy(FNetSerializationContext& Context);
 
 	// Write object and SubObjects
 	EWriteObjectStatus WriteObjectAndSubObjects(FNetSerializationContext& Context, uint32 InternalIndex, uint32 WriteObjectFlags, FBatchInfo& OutBatchInfo);
diff --git a/Engine/Source/Runtime/Experimental/Iris/Core/Public/Iris/Serialization/NetJournal.h b/Engine/Source/Runtime/Experimental/Iris/Core/Public/Iris/Serialization/NetJournal.h
--- a/Engine/Source/Runtime/Experimental/Iris/Core/Public/Iris/Serialization/NetJournal.h
+++ b/Engine/Source/Runtime/Experimental/Iris/Core/Public/Iris/Serialization/NetJournal.h
@@ -18,8 +18,10 @@ class UReplicationSystem;
 
 #if UE_NET_ENABLE_READ_JOURNAL
 	#define UE_ADD_READ_JOURNAL_ENTRY(SerializationContext, X) SerializationContext.AddReadJournalEntry(X);
+	#define UE_RESET_READ_JOURNAL(SerializationContext) SerializationContext.ResetReadJournal();
 #else
 	#define UE_ADD_READ_JOURNAL_ENTRY(...)
+#define UE_RESET_READ_JOURNAL(...)
 #endif
 
 namespace UE::Net
