Iris - Fixed a few cases where we did not properly propagate HasPushBasedDiritness. Make sure function descriptors are marked as push based. Added tests for HashPushBasedDirtiness trait.

diff --git a/Engine/Plugins/Runtime/GameplayAbilities/Source/GameplayAbilities/Private/Serialization/MinimalGameplayCueReplicationProxyReplicationFragment.cpp b/Engine/Plugins/Runtime/GameplayAbilities/Source/GameplayAbilities/Private/Serialization/MinimalGameplayCueReplicationProxyReplicationFragment.cpp
--- a/Engine/Plugins/Runtime/GameplayAbilities/Source/GameplayAbilities/Private/Serialization/MinimalGameplayCueReplicationProxyReplicationFragment.cpp
+++ b/Engine/Plugins/Runtime/GameplayAbilities/Source/GameplayAbilities/Private/Serialization/MinimalGameplayCueReplicationProxyReplicationFragment.cpp
@@ -49,6 +49,13 @@ FMinimalGameplayCueReplicationProxyReplicationFragment::FMinimalGameplayCueRepli
 
 		SrcReplicationState = MakeUnique<FPropertyReplicationState>(InDescriptor);
 	}
+
+#if WITH_PUSH_MODEL
+	if (EnumHasAnyFlags(InDescriptor->Traits, EReplicationStateTraits::HasPushBasedDirtiness))
+	{
+		Traits |= EReplicationFragmentTraits::HasPushBasedDirtiness;
+	}
+#endif
 	
 	if (EnumHasAnyFlags(InTraits, EReplicationFragmentTraits::CanReceive))
 	{
diff --git a/Engine/Plugins/Runtime/GameplayAbilities/Source/GameplayAbilities/Private/Serialization/MinimalReplicationTagCountMapReplicationFragment.cpp b/Engine/Plugins/Runtime/GameplayAbilities/Source/GameplayAbilities/Private/Serialization/MinimalReplicationTagCountMapReplicationFragment.cpp
--- a/Engine/Plugins/Runtime/GameplayAbilities/Source/GameplayAbilities/Private/Serialization/MinimalReplicationTagCountMapReplicationFragment.cpp
+++ b/Engine/Plugins/Runtime/GameplayAbilities/Source/GameplayAbilities/Private/Serialization/MinimalReplicationTagCountMapReplicationFragment.cpp
@@ -43,7 +43,14 @@ FMinimalReplicationTagCountMapReplicationFragment::FMinimalReplicationTagCountMa
 
 		SrcReplicationState = MakeUnique<FPropertyReplicationState>(InDescriptor);
 	}
-	
+
+#if WITH_PUSH_MODEL
+	if (EnumHasAnyFlags(InDescriptor->Traits, EReplicationStateTraits::HasPushBasedDirtiness))
+	{
+		Traits |= EReplicationFragmentTraits::HasPushBasedDirtiness;
+	}
+#endif
+		
 	if (EnumHasAnyFlags(InTraits, EReplicationFragmentTraits::CanReceive))
 	{
 		if (EnumHasAnyFlags(InDescriptor->Traits, EReplicationStateTraits::HasRepNotifies))
diff --git a/Engine/Plugins/Runtime/ReplicationSystemTestPlugin/Source/Private/Tests/ReplicationState/TestReplicationStateDescriptorBuilder.cpp b/Engine/Plugins/Runtime/ReplicationSystemTestPlugin/Source/Private/Tests/ReplicationState/TestReplicationStateDescriptorBuilder.cpp
--- a/Engine/Plugins/Runtime/ReplicationSystemTestPlugin/Source/Private/Tests/ReplicationState/TestReplicationStateDescriptorBuilder.cpp
+++ b/Engine/Plugins/Runtime/ReplicationSystemTestPlugin/Source/Private/Tests/ReplicationState/TestReplicationStateDescriptorBuilder.cpp
@@ -695,6 +695,19 @@ UE_NET_TEST_FIXTURE(FTestReplicationStateDescriptorBuilderFixture, UTestReplicat
 		const uint16 ExpectedArrayIndex = (MemberIt % UTestReplicationStateDescriptor_TestClassWithCArray::ArrayElementCount);
 		UE_NET_ASSERT_EQ(MemberPropertyDescriptor.ArrayIndex, ExpectedArrayIndex);
 	}
+
+	// There should be a RepIndex per CArray element. There are two arrays with ArrayElementCount elements each.
+	UE_NET_ASSERT_EQ(uint32(Descriptor->RepIndexCount), 2U*UTestReplicationStateDescriptor_TestClassWithCArray::ArrayElementCount);
+
+	// Each rep index should point to an identical member index as the class only has C arrays
+	{
+		const FReplicationStateMemberRepIndexToMemberIndexDescriptor* RepIndexDescriptor = Descriptor->MemberRepIndexToMemberIndexDescriptors;
+		UE_NET_ASSERT_NE(RepIndexDescriptor, nullptr);
+		for (uint16 RepIndexIt = 0, RepIndexEndIt = 2*UTestReplicationStateDescriptor_TestClassWithCArray::ArrayElementCount; RepIndexIt != RepIndexEndIt; ++RepIndexIt)
+		{
+			UE_NET_ASSERT_EQ(RepIndexDescriptor[RepIndexIt].MemberIndex, RepIndexIt);
+		}
+	}
 }
 
 UE_NET_TEST_FIXTURE(FTestReplicationStateDescriptorBuilderFixture, UTestReplicationStateDescriptor_TestClassWithStructWithCArray)
@@ -1184,6 +1197,61 @@ UE_NET_TEST_FIXTURE(FTestReplicationStateDescriptorBuilderFixture, UTestReplicat
 	UE_NET_ASSERT_EQ(SerializerDescriptors[0].Serializer, &UE_NET_GET_SERIALIZER(FFieldPathNetSerializer));
 }
 
+UE_NET_TEST_FIXTURE(FTestReplicationStateDescriptorBuilderFixture, TestClassWithOnlyPushBasedPropertiesHasPushBasedDirtinessTrait)
+{
+	InitDescriptorsFromClass(UTestReplicationStateDescriptor_TestClassWithOnlyPushBasedProperties::StaticClass());
+	UE_NET_ASSERT_GE(Descriptors.Num(), 1);
+
+	for (const FReplicationStateDescriptor* TheDescriptor : Descriptors)
+	{
+		UE_NET_ASSERT_TRUE(EnumHasAnyFlags(TheDescriptor->Traits, EReplicationStateTraits::HasPushBasedDirtiness));
+	}
+}
+
+UE_NET_TEST_FIXTURE(FTestReplicationStateDescriptorBuilderFixture, TestClassWithNeverReplicatePropertiesHasPushBasedDirtinessTrait)
+{
+	InitDescriptorsFromClass(UTestReplicationStateDescriptor_TestClassWithOnlyPropertiesThatNeverReplicates::StaticClass());
+
+	// If a descriptor is created at all it is expected to have the HasPushBasedDirtiness trait 
+	for (const FReplicationStateDescriptor* TheDescriptor : Descriptors)
+	{
+		UE_NET_ASSERT_TRUE(EnumHasAnyFlags(TheDescriptor->Traits, EReplicationStateTraits::HasPushBasedDirtiness));
+	}
+}
+
+UE_NET_TEST_FIXTURE(FTestReplicationStateDescriptorBuilderFixture, TestClassWithPushBasedPropertiesAndNeverReplicatedPropertiesHasPushBasedDirtinessTrait)
+{
+	InitDescriptorsFromClass(UTestReplicationStateDescriptor_TestClassWithPushBasedPropertiesAndNeverReplicatedProperties::StaticClass());
+	UE_NET_ASSERT_GE(Descriptors.Num(), 1);
+
+	for (const FReplicationStateDescriptor* TheDescriptor : Descriptors)
+	{
+		UE_NET_ASSERT_TRUE(EnumHasAnyFlags(TheDescriptor->Traits, EReplicationStateTraits::HasPushBasedDirtiness));
+	}
+}
+
+UE_NET_TEST_FIXTURE(FTestReplicationStateDescriptorBuilderFixture, TestClassWithPushBasedPropertiesAndNeverReplicatedPropertiesAndFunctionsHasPushBasedDirtinessTrait)
+{
+	InitDescriptorsFromClass(UTestReplicationStateDescriptor_TestClassWithPushBasedPropertiesAndNeverReplicatedPropertiesAndFunctions::StaticClass());
+	UE_NET_ASSERT_GE(Descriptors.Num(), 2);
+
+	for (const FReplicationStateDescriptor* TheDescriptor : Descriptors)
+	{
+		UE_NET_ASSERT_TRUE(EnumHasAnyFlags(TheDescriptor->Traits, EReplicationStateTraits::HasPushBasedDirtiness));
+	}
+}
+
+UE_NET_TEST_FIXTURE(FTestReplicationStateDescriptorBuilderFixture, UTestReplicationStateDescriptor_TestClassWithPushBasedAndNonPushBasedPropertiesDoesNotHavePushBasedDirtinessTrait)
+{
+	InitDescriptorsFromClass(UTestReplicationStateDescriptor_TestClassWithPushBasedAndNonPushBasedProperties::StaticClass());
+	UE_NET_ASSERT_EQ(Descriptors.Num(), 1);
+
+	for (const FReplicationStateDescriptor* TheDescriptor : Descriptors)
+	{
+		UE_NET_ASSERT_FALSE(EnumHasAnyFlags(TheDescriptor->Traits, EReplicationStateTraits::HasPushBasedDirtiness));
+	}
+}
+
 }
 
 // Function definitions
@@ -1418,3 +1486,76 @@ void UTestReplicationStateDescriptor_TestClassWithFieldPathProperty::GetLifetime
 	DOREPLIFETIME(ThisClass, FieldPathProperty);
 }
 
+void UTestReplicationStateDescriptor_TestClassWithOnlyPushBasedProperties::GetLifetimeReplicatedProps(TArray<FLifetimeProperty>& OutLifetimeProps) const
+{
+	FDoRepLifetimeParams LifetimeParams;
+	LifetimeParams.bIsPushBased = true;
+	DOREPLIFETIME_WITH_PARAMS_FAST(ThisClass, IntPushBasedA, LifetimeParams);
+	DOREPLIFETIME_WITH_PARAMS_FAST(ThisClass, IntPushBasedB, LifetimeParams);
+	DOREPLIFETIME_WITH_PARAMS_FAST(ThisClass, IntPushBasedC, LifetimeParams);
+}
+
+void UTestReplicationStateDescriptor_TestClassWithOnlyPropertiesThatNeverReplicates::GetLifetimeReplicatedProps(TArray<FLifetimeProperty>& OutLifetimeProps) const
+{
+	FDoRepLifetimeParams LifetimeParams;
+	LifetimeParams.bIsPushBased = false;
+	LifetimeParams.Condition = COND_Never;
+	DOREPLIFETIME_WITH_PARAMS_FAST(ThisClass, IntNeverReplicatesA, LifetimeParams);
+	DOREPLIFETIME_WITH_PARAMS_FAST(ThisClass, IntNeverReplicatesB, LifetimeParams);
+	DOREPLIFETIME_WITH_PARAMS_FAST(ThisClass, IntNeverReplicatesC, LifetimeParams);
+}
+
+void UTestReplicationStateDescriptor_TestClassWithPushBasedPropertiesAndNeverReplicatedProperties::GetLifetimeReplicatedProps(TArray<FLifetimeProperty>& OutLifetimeProps) const
+{
+	FDoRepLifetimeParams LifetimeParams;
+	LifetimeParams.bIsPushBased = true;
+	DOREPLIFETIME_WITH_PARAMS_FAST(ThisClass, IntPushBasedA, LifetimeParams);
+	DOREPLIFETIME_WITH_PARAMS_FAST(ThisClass, IntPushBasedB, LifetimeParams);
+	DOREPLIFETIME_WITH_PARAMS_FAST(ThisClass, IntPushBasedC, LifetimeParams);
+
+
+	LifetimeParams.bIsPushBased = false;
+	LifetimeParams.Condition = COND_Never;
+	DOREPLIFETIME_WITH_PARAMS_FAST(ThisClass, IntNeverReplicatesA, LifetimeParams);
+	DOREPLIFETIME_WITH_PARAMS_FAST(ThisClass, IntNeverReplicatesB, LifetimeParams);
+	DOREPLIFETIME_WITH_PARAMS_FAST(ThisClass, IntNeverReplicatesC, LifetimeParams);
+}
+
+void UTestReplicationStateDescriptor_TestClassWithPushBasedPropertiesAndNeverReplicatedPropertiesAndFunctions::GetLifetimeReplicatedProps(TArray<FLifetimeProperty>& OutLifetimeProps) const
+{
+	FDoRepLifetimeParams LifetimeParams;
+	LifetimeParams.bIsPushBased = true;
+	DOREPLIFETIME_WITH_PARAMS_FAST(ThisClass, IntPushBasedA, LifetimeParams);
+	DOREPLIFETIME_WITH_PARAMS_FAST(ThisClass, IntPushBasedB, LifetimeParams);
+	DOREPLIFETIME_WITH_PARAMS_FAST(ThisClass, IntPushBasedC, LifetimeParams);
+
+
+	LifetimeParams.bIsPushBased = false;
+	LifetimeParams.Condition = COND_Never;
+	DOREPLIFETIME_WITH_PARAMS_FAST(ThisClass, IntNeverReplicatesA, LifetimeParams);
+	DOREPLIFETIME_WITH_PARAMS_FAST(ThisClass, IntNeverReplicatesB, LifetimeParams);
+	DOREPLIFETIME_WITH_PARAMS_FAST(ThisClass, IntNeverReplicatesC, LifetimeParams);
+}
+
+void UTestReplicationStateDescriptor_TestClassWithPushBasedPropertiesAndNeverReplicatedPropertiesAndFunctions::ClientUnrealiableFunction_Implementation(int32 Param)
+{
+}
+
+void UTestReplicationStateDescriptor_TestClassWithPushBasedPropertiesAndNeverReplicatedPropertiesAndFunctions::ServerReliableFunction_Implementation(int32 Param)
+{
+}
+
+void UTestReplicationStateDescriptor_TestClassWithPushBasedAndNonPushBasedProperties::GetLifetimeReplicatedProps(TArray<FLifetimeProperty> & OutLifetimeProps) const
+{
+	FDoRepLifetimeParams LifetimeParams;
+	LifetimeParams.bIsPushBased = true;
+	DOREPLIFETIME_WITH_PARAMS_FAST(ThisClass, IntPushBasedA, LifetimeParams);
+	DOREPLIFETIME_WITH_PARAMS_FAST(ThisClass, IntPushBasedB, LifetimeParams);
+	DOREPLIFETIME_WITH_PARAMS_FAST(ThisClass, IntPushBasedC, LifetimeParams);
+
+
+	LifetimeParams.bIsPushBased = false;
+	DOREPLIFETIME_WITH_PARAMS_FAST(ThisClass, IntNotPushBasedA, LifetimeParams);
+	DOREPLIFETIME_WITH_PARAMS_FAST(ThisClass, IntNotPushBasedB, LifetimeParams);
+	DOREPLIFETIME_WITH_PARAMS_FAST(ThisClass, IntNotPushBasedC, LifetimeParams);
+}
diff --git a/Engine/Plugins/Runtime/ReplicationSystemTestPlugin/Source/Private/Tests/ReplicationState/TestReplicationStateDescriptorBuilder.h b/Engine/Plugins/Runtime/ReplicationSystemTestPlugin/Source/Private/Tests/ReplicationState/TestReplicationStateDescriptorBuilder.h
--- a/Engine/Plugins/Runtime/ReplicationSystemTestPlugin/Source/Private/Tests/ReplicationState/TestReplicationStateDescriptorBuilder.h
+++ b/Engine/Plugins/Runtime/ReplicationSystemTestPlugin/Source/Private/Tests/ReplicationState/TestReplicationStateDescriptorBuilder.h
@@ -749,3 +749,116 @@ public:
 	TFieldPath<FProperty> FieldPathProperty;
 };
 
+UCLASS()
+class UTestReplicationStateDescriptor_TestClassWithOnlyPushBasedProperties : public UObject
+{
+	GENERATED_BODY()
+
+private:
+	UPROPERTY(Replicated)
+	int32 IntPushBasedA;
+
+	UPROPERTY(Replicated)
+	int32 IntPushBasedB;
+
+	UPROPERTY(Replicated)
+	int32 IntPushBasedC;
+};
+
+UCLASS()
+class UTestReplicationStateDescriptor_TestClassWithOnlyPropertiesThatNeverReplicates : public UObject
+{
+	GENERATED_BODY()
+
+private:
+	UPROPERTY(Replicated)
+	int32 IntNeverReplicatesA;
+
+	UPROPERTY(Replicated)
+	int32 IntNeverReplicatesB;
+
+	UPROPERTY(Replicated)
+	int32 IntNeverReplicatesC;
+};
+
+UCLASS()
+class UTestReplicationStateDescriptor_TestClassWithPushBasedPropertiesAndNeverReplicatedProperties : public UObject
+{
+	GENERATED_BODY()
+
+private:
+	UPROPERTY(Replicated)
+	int32 IntPushBasedA;
+
+	UPROPERTY(Replicated)
+	int32 IntPushBasedB;
+
+	UPROPERTY(Replicated)
+	int32 IntPushBasedC;
+
+	UPROPERTY(Replicated)
+	int32 IntNeverReplicatesA;
+
+	UPROPERTY(Replicated)
+	int32 IntNeverReplicatesB;
+
+	UPROPERTY(Replicated)
+	int32 IntNeverReplicatesC;
+};
+
+UCLASS()
+class UTestReplicationStateDescriptor_TestClassWithPushBasedPropertiesAndNeverReplicatedPropertiesAndFunctions : public UObject
+{
+	GENERATED_BODY()
+
+protected:
+	UFUNCTION(Client, unreliable)
+	void ClientUnrealiableFunction(int32 Param);
+
+	UFUNCTION(Server, reliable)
+	void ServerReliableFunction(int32 Param);
+
+private:
+	UPROPERTY(Replicated)
+	int32 IntPushBasedA;
+
+	UPROPERTY(Replicated)
+	int32 IntPushBasedB;
+
+	UPROPERTY(Replicated)
+	int32 IntPushBasedC;
+
+	UPROPERTY(Replicated)
+	int32 IntNeverReplicatesA;
+
+	UPROPERTY(Replicated)
+	int32 IntNeverReplicatesB;
+
+	UPROPERTY(Replicated)
+	int32 IntNeverReplicatesC;
+};
+
+UCLASS()
+class UTestReplicationStateDescriptor_TestClassWithPushBasedAndNonPushBasedProperties : public UObject
+{
+	GENERATED_BODY()
+
+private:
+	UPROPERTY(Replicated)
+	int32 IntPushBasedA;
+
+	UPROPERTY(Replicated)
+	int32 IntPushBasedB;
+
+	UPROPERTY(Replicated)
+	int32 IntPushBasedC;
+
+	UPROPERTY(Replicated)
+	int32 IntNotPushBasedA;
+
+	UPROPERTY(Replicated)
+	int32 IntNotPushBasedB;
+
+	UPROPERTY(Replicated)
+	int32 IntNotPushBasedC;
+};
diff --git a/Engine/Source/Runtime/Experimental/Iris/Core/Private/Iris/ReplicationState/ReplicationStateDescriptorBuilder.cpp b/Engine/Source/Runtime/Experimental/Iris/Core/Private/Iris/ReplicationState/ReplicationStateDescriptorBuilder.cpp
--- a/Engine/Source/Runtime/Experimental/Iris/Core/Private/Iris/ReplicationState/ReplicationStateDescriptorBuilder.cpp
+++ b/Engine/Source/Runtime/Experimental/Iris/Core/Private/Iris/ReplicationState/ReplicationStateDescriptorBuilder.cpp
@@ -356,7 +356,7 @@ private:
 	EReplicationStateTraits BuildReplicationStateTraits(const FBuilderContext& Context) const;
 
 	static FRepTag GetRepTagFromProperty(const FMemberCacheEntry& MemberCacheEntry, const FMemberProperty& Property);
-
+	static uint32 GetRepIndexCount(const TConstArrayView<FMemberProperty>& Members);
 	void FinalizeDescriptor(FReplicationStateDescriptor* Descriptor, FBuilderContext& Context) const;
 
 	TRefCountPtr<const FReplicationStateDescriptor> GetDescriptorForStructProperty(const FMemberProperty& MemberProperty, FReplicationStateDescriptorRegistry* DescriptorRegistry) const;
@@ -1253,12 +1253,7 @@ void FPropertyReplicationStateDescriptorBuilder::BuildMemberLifetimeConditionDes
 
 void FPropertyReplicationStateDescriptorBuilder::BuildMemberRepIndexToMemberDescriptors(FReplicationStateMemberRepIndexToMemberIndexDescriptor* MemberRepIndexToMemberIndexDescriptors, FReplicationStateDescriptor* Descriptor) const
 {
-	uint16 MaxRepIndex = 0;
-	for (const FMemberProperty& Member : Members)
-	{
-		MaxRepIndex = FPlatformMath::Max(MaxRepIndex, Member.Property->RepIndex);
-	}
-	const uint32 RepIndexCount = static_cast<uint32>(MaxRepIndex) + 1U;
+	const uint32 RepIndexCount = GetRepIndexCount(MakeConstArrayView(Members));
 
 	// Faster way of setting all the MemberIndex values to InvalidEntry
 	{
@@ -1266,11 +1261,16 @@ void FPropertyReplicationStateDescriptorBuilder::BuildMemberRepIndexToMemberDesc
 		FPlatformMemory::Memset(MemberRepIndexToMemberIndexDescriptors, 0xFF, RepIndexCount*sizeof(FReplicationStateMemberRepIndexToMemberIndexDescriptor));
 	}
 
+	const FProperty* LastProperty = nullptr;
+	uint16 ArrayIndex = 0;
 	for (const FMemberProperty& Info : Members)
 	{
+		ArrayIndex = (LastProperty && Info.Property == LastProperty ? uint16(ArrayIndex + 1) : uint16(0));
+		LastProperty = Info.Property;
+		CA_ASSUME(Info.Property != nullptr);
 
 		const SIZE_T MemberIndex = &Info - Members.GetData();
-		const uint16 RepIndex = Info.Property->RepIndex;
+		const uint16 RepIndex = Info.Property->RepIndex + ArrayIndex;
 		checkSlow(RepIndex < RepIndexCount);
 		MemberRepIndexToMemberIndexDescriptors[RepIndex].MemberIndex = static_cast<uint16>(MemberIndex);
 	}
@@ -1535,6 +1535,12 @@ EReplicationStateTraits FPropertyReplicationStateDescriptorBuilder::BuildReplica
 		Traits |= EReplicationStateTraits::SupportsDeltaCompression;
 	}
 
+	// For the main replication state for functions we do EReplicationStateTraits::HasPushBasedDirtiness even though it does not have any members.
+	if (Context.MemberFunctionCache.Num() && Members.Num() == 0)
+	{
+		Traits |= EReplicationStateTraits::HasPushBasedDirtiness;
+	}
+
 	Traits |= EReplicationStateTraits::NeedsRefCount;
 
 	return Traits;
@@ -1647,12 +1653,7 @@ FPropertyReplicationStateDescriptorBuilder::Build(const FString& StateName, FRep
 	{
 		if (Members.Num() > 0)
 		{
-			uint16 MaxRepIndex = 0;
-			for (const FMemberProperty& Member : Members)
-			{
-				MaxRepIndex = FPlatformMath::Max(MaxRepIndex, Member.Property->RepIndex);
-			}
-			const uint32 RepIndexCount = static_cast<uint32>(MaxRepIndex) + 1U;
+			const uint32 RepIndexCount = GetRepIndexCount(MakeConstArrayView(Members));
 			FMemoryLayoutUtil::AddToLayout<FReplicationStateMemberRepIndexToMemberIndexDescriptor>(Layout, LayoutData.MemberRepIndexToMemberIndexDescriptorsSizeAndOffset, RepIndexCount);
 		}
 	}
@@ -2155,6 +2156,25 @@ FRepTag FPropertyReplicationStateDescriptorBuilder::GetRepTagFromProperty(const
 	return GetInvalidRepTag();
 }
 
+uint32 FPropertyReplicationStateDescriptorBuilder::GetRepIndexCount(const TConstArrayView<FMemberProperty>& Members)
+{
+	uint16 MaxRepIndex = 0;
+	{
+		const FProperty* LastProperty = nullptr;
+		uint16 ArrayIndex = 0;
+		for (const FMemberProperty& Member : Members)
+		{
+			ArrayIndex = (LastProperty && Member.Property == LastProperty ? uint16(ArrayIndex + 1) : uint16(0));
+			LastProperty = Member.Property;
+			CA_ASSUME(Member.Property != nullptr);
+			MaxRepIndex = FPlatformMath::Max(MaxRepIndex, static_cast<uint16>(Member.Property->RepIndex + ArrayIndex));
+		}
+	}
+	
+	const uint32 RepIndexCount = static_cast<uint32>(MaxRepIndex) + 1U;
+	return RepIndexCount;
+}
+
 void FPropertyReplicationStateDescriptorBuilder::GetPropertyPathName(const FProperty* Property, FString& PathName)
 {
 	constexpr uint32 MaxHierarchyDepth = 64;
@@ -2686,9 +2706,11 @@ SIZE_T FReplicationStateDescriptorBuilder::CreateDescriptorsForClass(FResult& Cr
 	{
 		FPropertyReplicationStateDescriptorBuilder::FMemberProperty MemberProperty;
 		const FProperty* Property = RepRecord.Property;
+
+		const int32 RepIndex = static_cast<int32>(Property->RepIndex);
 	
 		// If we are building a class descriptor for a single property we only include the property matching the index
-		if (bSingleProperty && ((int32)Property->RepIndex != Parameters.SinglePropertyIndex))
+		if (bSingleProperty && (RepIndex != Parameters.SinglePropertyIndex))
 		{
 			continue;
 		}
