Iris - Added missing REPLICATED_BASE_CLASS macro.

#rnx

diff --git a/Engine/Plugins/Runtime/ReplicationSystemTestPlugin/Source/Private/Tests/ReplicationSystem/ReplicatedTestObject.h b/Engine/Plugins/Runtime/ReplicationSystemTestPlugin/Source/Private/Tests/ReplicationSystem/ReplicatedTestObject.h
--- a/Engine/Plugins/Runtime/ReplicationSystemTestPlugin/Source/Private/Tests/ReplicationSystem/ReplicatedTestObject.h
+++ b/Engine/Plugins/Runtime/ReplicationSystemTestPlugin/Source/Private/Tests/ReplicationSystem/ReplicatedTestObject.h
@@ -12,6 +12,7 @@
 #include "Iris/ReplicationSystem/TypedReplicationFragment.h"
 #include "Templates/UniquePtr.h"
 #include "Net/Core/NetBitArray.h"
+#include "Net/Core/PushModel/PushModelMacros.h"
 #include "Iris/ReplicationState/ReplicationStateDescriptorMacros.h"
 #include "ReplicatedTestObject.generated.h"
 
@@ -28,6 +29,7 @@ UCLASS()
 class UReplicatedTestObject : public UObject
 {
 	GENERATED_BODY()
+	REPLICATED_BASE_CLASS(UReplicatedTestObject)
 
 public:
 	// Register the fragments for this object
