Iris - Removed space from filename.

#rnx

diff --git a/Engine/Plugins/Runtime/ReplicationSystemTestPlugin/Source/Private/IrisTestMessageStreamOperators .cpp b/Engine/Plugins/Runtime/ReplicationSystemTestPlugin/Source/Private/IrisTestMessageStreamOperators .cpp
deleted file mode 100644
diff --git a/Engine/Plugins/Runtime/ReplicationSystemTestPlugin/Source/Private/IrisTestMessageStreamOperators.cpp b/Engine/Plugins/Runtime/ReplicationSystemTestPlugin/Source/Private/IrisTestMessageStreamOperators.cpp
new file mode 100644
--- /dev/null
+++ b/Engine/Plugins/Runtime/ReplicationSystemTestPlugin/Source/Private/IrisTestMessageStreamOperators.cpp
@@ -0,0 +1,14 @@
+// Copyright Epic Games, Inc. All Rights Reserved.
+
+#include "IrisTestMessageStreamOperators.h"
+#include "Iris/ReplicationSystem/Filtering/NetObjectFilter.h"
+
+namespace UE::Net
+{
+
+FTestMessage& operator<<(FTestMessage& Message, ENetFilterStatus Status)
+{
+	return Message << (Status == ENetFilterStatus::Disallow ? TEXT("ENetFilterStatus::Disallow") : TEXT("ENetFilterStatus::Allow"));
+}
+
+}
