Iris - NetBitStreamReader substream support.

diff --git a/Engine/Plugins/Runtime/ReplicationSystemTestPlugin/Source/Private/Tests/Serialization/TestNetBitStreams.cpp b/Engine/Plugins/Runtime/ReplicationSystemTestPlugin/Source/Private/Tests/Serialization/TestNetBitStreams.cpp
--- a/Engine/Plugins/Runtime/ReplicationSystemTestPlugin/Source/Private/Tests/Serialization/TestNetBitStreams.cpp
+++ b/Engine/Plugins/Runtime/ReplicationSystemTestPlugin/Source/Private/Tests/Serialization/TestNetBitStreams.cpp
@@ -48,7 +48,11 @@ protected:
 	FNetBitStreamReader Reader;
 };
 
-class FNetBitStreamSubstreamTest : public FNetBitStreamWriterTest
+class FNetBitStreamWriterSubstreamTest : public FNetBitStreamWriterTest
+{
+};
+
+class FNetBitStreamReaderSubstreamTest : public FNetBitStreamReaderTest
 {
 };
 
@@ -340,7 +344,8 @@ UE_NET_TEST_FIXTURE(FNetBitStreamReaderWriterTest, ReadStreamWithBitsWrittenAtAr
 	}
 }
 
-UE_NET_TEST_FIXTURE(FNetBitStreamSubstreamTest, CanCreateSubstream)
+// Writer substream tests
+UE_NET_TEST_FIXTURE(FNetBitStreamWriterSubstreamTest, CanCreateSubstream)
 {
 	Writer.InitBytes(Buffer.GetBuffer(), Buffer.GetBufferCapacity());
 	Writer.Seek(1U);
@@ -354,7 +359,7 @@ UE_NET_TEST_FIXTURE(FNetBitStreamSubstreamTest, CanCreateSubstream)
 	Writer.DiscardSubstream(Substream);
 }
 
-UE_NET_TEST_FIXTURE(FNetBitStreamSubstreamTest, CanCreateSmallSubstream)
+UE_NET_TEST_FIXTURE(FNetBitStreamWriterSubstreamTest, CanCreateSmallSubstream)
 {
 	Writer.InitBytes(Buffer.GetBuffer(), Buffer.GetBufferCapacity());
 	Writer.Seek(1U);
@@ -368,7 +373,21 @@ UE_NET_TEST_FIXTURE(FNetBitStreamSubstreamTest, CanCreateSmallSubstream)
 	Writer.DiscardSubstream(Substream);
 }
 
-UE_NET_TEST_FIXTURE(FNetBitStreamSubstreamTest, CanDiscardSubstream)
+UE_NET_TEST_FIXTURE(FNetBitStreamWriterSubstreamTest, CanCreateEmptySubstream)
+{
+	Writer.InitBytes(Buffer.GetBuffer(), Buffer.GetBufferCapacity());
+	Writer.Seek(1U);
+
+	FNetBitStreamWriter Substream = Writer.CreateSubstream(0U);
+
+	UE_NET_ASSERT_FALSE(Substream.IsOverflown());
+	UE_NET_ASSERT_EQ(Substream.GetPosBits(), 0U);
+	UE_NET_ASSERT_EQ(Substream.GetBitsLeft(), 0U);
+
+	Writer.DiscardSubstream(Substream);
+}
+
+UE_NET_TEST_FIXTURE(FNetBitStreamWriterSubstreamTest, CanDiscardSubstream)
 {
 	Writer.InitBytes(Buffer.GetBuffer(), Buffer.GetBufferCapacity());
 
@@ -384,7 +403,7 @@ UE_NET_TEST_FIXTURE(FNetBitStreamSubstreamTest, CanDiscardSubstream)
 	UE_NET_ASSERT_EQ(Writer.GetPosBits(), 0U);
 }
 
-UE_NET_TEST_FIXTURE(FNetBitStreamSubstreamTest, CanCommitSubstream)
+UE_NET_TEST_FIXTURE(FNetBitStreamWriterSubstreamTest, CanCommitSubstream)
 {
 	Writer.InitBytes(Buffer.GetBuffer(), Buffer.GetBufferCapacity());
 	constexpr uint32 SeekPos = 33;
@@ -404,7 +423,7 @@ UE_NET_TEST_FIXTURE(FNetBitStreamSubstreamTest, CanCommitSubstream)
 	UE_NET_ASSERT_EQ(Writer.GetPosBits(), SeekPos + WriteBitCount);
 }
 
-UE_NET_TEST_FIXTURE(FNetBitStreamSubstreamTest, CanCreateSubstreamFromOverflowedStream)
+UE_NET_TEST_FIXTURE(FNetBitStreamWriterSubstreamTest, CanCreateSubstreamFromOverflowedStream)
 {
 	Writer.InitBytes(Buffer.GetBuffer(), Buffer.GetBufferCapacity());
 	Writer.Seek(Buffer.GetBufferCapacity()*8 + 1);
@@ -430,7 +449,7 @@ UE_NET_TEST_FIXTURE(FNetBitStreamSubstreamTest, CanCreateSubstreamFromOverflowed
 	UE_NET_ASSERT_EQ(CurrentBitPos, PreviousBitPos);
 }
 
-UE_NET_TEST_FIXTURE(FNetBitStreamSubstreamTest, CanCreateSubSubstream)
+UE_NET_TEST_FIXTURE(FNetBitStreamWriterSubstreamTest, CanCreateSubSubstream)
 {
 	Writer.InitBytes(Buffer.GetBuffer(), Buffer.GetBufferCapacity());
 
@@ -452,7 +471,7 @@ UE_NET_TEST_FIXTURE(FNetBitStreamSubstreamTest, CanCreateSubSubstream)
 	UE_NET_ASSERT_EQ(CurrentBitPos, SeekCount1 + WriteCount2);
 }
 
-UE_NET_TEST_FIXTURE(FNetBitStreamSubstreamTest, CanWriteToSubSubstream)
+UE_NET_TEST_FIXTURE(FNetBitStreamWriterSubstreamTest, CanWriteToSubSubstream)
 {
 	uint32 Word = 0;
 	Writer.InitBytes(&Word, sizeof(Word));
@@ -474,7 +493,7 @@ UE_NET_TEST_FIXTURE(FNetBitStreamSubstreamTest, CanWriteToSubSubstream)
 	UE_NET_ASSERT_EQ(Word, 0xFFFF0000U);
 }
 
-UE_NET_TEST_FIXTURE(FNetBitStreamSubstreamTest, CanReadDataCommittedFromSubstream)
+UE_NET_TEST_FIXTURE(FNetBitStreamWriterSubstreamTest, CanReadDataCommittedFromSubstream)
 {
 	Writer.InitBytes(Buffer.GetBuffer(), Buffer.GetBufferCapacity());
 
@@ -506,7 +525,7 @@ UE_NET_TEST_FIXTURE(FNetBitStreamSubstreamTest, CanReadDataCommittedFromSubstrea
 	}
 }
 
-UE_NET_TEST_FIXTURE(FNetBitStreamWriterTest, CanWriteToEndOfSubStreamTest1)
+UE_NET_TEST_FIXTURE(FNetBitStreamWriterSubstreamTest, CanWriteToEndOfSubStreamTest1)
 {
 	constexpr uint32 Sentinel = 0xBAAAAAADU;
 
@@ -521,6 +540,7 @@ UE_NET_TEST_FIXTURE(FNetBitStreamWriterTest, CanWriteToEndOfSubStreamTest1)
 		Writer.CommitWrites();
 	}
 
+	// Test normal reading
 	{
 		FNetBitStreamReader Reader;
 		Reader.InitBits(Buffer.GetBuffer(), Writer.GetPosBits());
@@ -529,9 +549,22 @@ UE_NET_TEST_FIXTURE(FNetBitStreamWriterTest, CanWriteToEndOfSubStreamTest1)
 		const uint32 ReadSentinel = Reader.ReadBits(32);
 		UE_NET_ASSERT_EQ(ReadSentinel, Sentinel);
 	}
+
+	// Test reading via substream
+	{
+		FNetBitStreamReader Reader;
+		Reader.InitBits(Buffer.GetBuffer(), Buffer.GetBufferCapacity()*8U);
+
+		FNetBitStreamReader Substream = Reader.CreateSubstream(Writer.GetPosBits());
+		Substream.Seek(Substream.GetBitsLeft() - 32U);
+
+		const uint32 ReadSentinel = Substream.ReadBits(32U);
+		UE_NET_ASSERT_EQ(ReadSentinel, Sentinel);
+		Reader.CommitSubstream(Substream);
+	}
 }
 
-UE_NET_TEST_FIXTURE(FNetBitStreamWriterTest, CanWriteToEndOfSubStreamTest2)
+UE_NET_TEST_FIXTURE(FNetBitStreamWriterSubstreamTest, CanWriteToEndOfSubStreamTest2)
 {
 	constexpr uint32 Sentinel1 = 0xBAAAAAADU;
 	constexpr uint32 Sentinel2 = 0xBAADF00DU;
@@ -552,19 +585,233 @@ UE_NET_TEST_FIXTURE(FNetBitStreamWriterTest, CanWriteToEndOfSubStreamTest2)
 		Writer.CommitWrites();
 	}
 
+	// Test normal reading
+	{
+		FNetBitStreamReader Reader;
+		Reader.InitBits(Buffer.GetBuffer(), Writer.GetPosBits());
+		Reader.Seek(Reader.GetBitsLeft() - 64U);
+
+		const uint32 ReadSentinel1 = Reader.ReadBits(32U);
+		UE_NET_ASSERT_EQ(ReadSentinel1, Sentinel1);
+
+		const uint32 ReadSentinel2 = Reader.ReadBits(32U);
+		UE_NET_ASSERT_EQ(ReadSentinel2, Sentinel2);
+	}
+
+	// Test reading via substream
 	{
 		FNetBitStreamReader Reader;
 		Reader.InitBits(Buffer.GetBuffer(), Writer.GetPosBits());
 		Reader.Seek(Reader.GetBitsLeft() - 64U);
 
-		const uint32 ReadSentinel1 = Reader.ReadBits(32);
+		FNetBitStreamReader Substream1 = Reader.CreateSubstream(32U);
+		const uint32 ReadSentinel1 = Substream1.ReadBits(32U);
 		UE_NET_ASSERT_EQ(ReadSentinel1, Sentinel1);
+		Reader.CommitSubstream(Substream1);
 
-		const uint32 ReadSentinel2 = Reader.ReadBits(32);
+		FNetBitStreamReader Substream2 = Reader.CreateSubstream();
+		const uint32 ReadSentinel2 = Substream2.ReadBits(32U);
 		UE_NET_ASSERT_EQ(ReadSentinel2, Sentinel2);
+		Reader.CommitSubstream(Substream2);
 	}
 }
 
+UE_NET_TEST_FIXTURE(FNetBitStreamWriterSubstreamTest, SubStreamSeek)
+{
+	Writer.InitBytes(Buffer.GetBuffer(), Buffer.GetBufferCapacity());
+	Writer.WriteBits(0x12345678,32);
+
+	FNetBitStreamWriter SubStream = Writer.CreateSubstream();
+	const uint32 SubStreamStartPos = SubStream.GetPosBits();
+	SubStream.Seek(0);
+	UE_NET_ASSERT_EQ(SubStreamStartPos, SubStream.GetPosBits());
+
+	Writer.CommitSubstream(SubStream);
+}
+
+// Reader substream tests
+UE_NET_TEST_FIXTURE(FNetBitStreamReaderSubstreamTest, CanCreateSubstream)
+{
+	Reader.InitBits(Buffer.GetBuffer(), Buffer.GetBufferCapacity()*8);
+	Reader.Seek(1U);
+
+	FNetBitStreamReader Substream = Reader.CreateSubstream(~0U);
+
+	UE_NET_ASSERT_FALSE(Substream.IsOverflown());
+	UE_NET_ASSERT_EQ(Substream.GetPosBits(), 0U);
+	UE_NET_ASSERT_EQ(Substream.GetBitsLeft(), Reader.GetBitsLeft());
+
+	Reader.DiscardSubstream(Substream);
+}
+
+UE_NET_TEST_FIXTURE(FNetBitStreamReaderSubstreamTest, CanCreateSmallSubstream)
+{
+	Reader.InitBits(Buffer.GetBuffer(), Buffer.GetBufferCapacity()*8);
+	Reader.Seek(1U);
+
+	FNetBitStreamReader Substream = Reader.CreateSubstream(3U);
+
+	UE_NET_ASSERT_FALSE(Substream.IsOverflown());
+	UE_NET_ASSERT_EQ(Substream.GetPosBits(), 0U);
+	UE_NET_ASSERT_EQ(Substream.GetBitsLeft(), 3U);
+
+	Reader.DiscardSubstream(Substream);
+}
+
+UE_NET_TEST_FIXTURE(FNetBitStreamReaderSubstreamTest, CanCreateEmptySubstream)
+{
+	Reader.InitBits(Buffer.GetBuffer(), Buffer.GetBufferCapacity()*8);
+	Reader.Seek(1U);
+
+	FNetBitStreamReader Substream = Reader.CreateSubstream(0U);
+
+	UE_NET_ASSERT_FALSE(Substream.IsOverflown());
+	UE_NET_ASSERT_EQ(Substream.GetPosBits(), 0U);
+	UE_NET_ASSERT_EQ(Substream.GetBitsLeft(), 0U);
+
+	Reader.DiscardSubstream(Substream);
+}
+
+UE_NET_TEST_FIXTURE(FNetBitStreamReaderSubstreamTest, CanDiscardSubstream)
+{
+	Reader.InitBits(Buffer.GetBuffer(), Buffer.GetBufferCapacity()*8);
+	Reader.Seek(1U);
+
+	const uint32 PreviousBitPos = Reader.GetPosBits();
+
+	{
+		FNetBitStreamReader Substream = Reader.CreateSubstream();
+		Substream.ReadBits(32U);
+
+		UE_NET_ASSERT_FALSE(Substream.IsOverflown());
+
+		Reader.DiscardSubstream(Substream);
+	}
+
+	// A discarded substream should not affect its parent's position.
+	UE_NET_ASSERT_EQ(Reader.GetPosBits(), PreviousBitPos);
+}
+
+UE_NET_TEST_FIXTURE(FNetBitStreamReaderSubstreamTest, CanCommitSubstream)
+{
+	Reader.InitBits(Buffer.GetBuffer(), Buffer.GetBufferCapacity()*8);
+
+	constexpr uint32 SeekPos = 33;
+	constexpr uint32 ReadBitCount = 32;
+
+	Reader.Seek(SeekPos);
+
+	{
+		FNetBitStreamReader Substream = Reader.CreateSubstream();
+		Substream.ReadBits(ReadBitCount);
+
+		UE_NET_ASSERT_FALSE(Substream.IsOverflown());
+
+		Reader.CommitSubstream(Substream);
+	}
+
+	UE_NET_ASSERT_EQ(Reader.GetPosBits(), SeekPos + ReadBitCount);
+}
+
+UE_NET_TEST_FIXTURE(FNetBitStreamReaderSubstreamTest, CanCreateSubstreamFromOverflowedStream)
+{
+	Reader.InitBits(Buffer.GetBuffer(), Buffer.GetBufferCapacity()*8);
+	Reader.Seek(Reader.GetBitsLeft() + 1);
+
+	UE_NET_ASSERT_TRUE(Reader.IsOverflown());
+
+	const uint32 PreviousBitPos = Reader.GetPosBits();
+	{
+		FNetBitStreamReader Substream = Reader.CreateSubstream();
+		UE_NET_ASSERT_TRUE(Substream.IsOverflown());
+
+		Substream.Seek(0U);
+		UE_NET_ASSERT_EQ(Substream.GetBitsLeft(), 0U);
+		
+		Substream.ReadBits(1U);
+		UE_NET_ASSERT_TRUE(Substream.IsOverflown());
+
+		// Commit overflown substream. Because of the overflow this should not affects its parent's position.
+		Reader.CommitSubstream(Substream);
+	}
+	const uint32 CurrentBitPos = Reader.GetPosBits();
+
+	UE_NET_ASSERT_EQ(CurrentBitPos, PreviousBitPos);
+}
+
+UE_NET_TEST_FIXTURE(FNetBitStreamReaderSubstreamTest, CanCreateSubSubstream)
+{
+	Reader.InitBits(Buffer.GetBuffer(), Buffer.GetBufferCapacity()*8);
+
+	constexpr uint32 SeekPosSubstream1 = 33;
+	constexpr uint32 ReadCountSubstream2 = 32;
+
+	{
+		FNetBitStreamReader Substream1 = Reader.CreateSubstream();
+		Substream1.Seek(SeekPosSubstream1);
+
+		FNetBitStreamReader Substream2 = Substream1.CreateSubstream();
+		Substream2.ReadBits(ReadCountSubstream2);
+
+		Substream1.CommitSubstream(Substream2);
+		Reader.CommitSubstream(Substream1);
+	}
+	const uint32 CurrentBitPos = Reader.GetPosBits();
+
+	UE_NET_ASSERT_EQ(CurrentBitPos, SeekPosSubstream1 + ReadCountSubstream2);
+}
+
+UE_NET_TEST_FIXTURE(FNetBitStreamReaderSubstreamTest, CanReadFromSubstream)
+{
+	const uint32 Word = 0xFFFF0000U;
+	Reader.InitBits(&Word, sizeof(Word)*8);
+	Reader.Seek(16);
+
+	constexpr uint32 SubstreamBitCount = 15U;
+	FNetBitStreamReader Substream = Reader.CreateSubstream(SubstreamBitCount);
+	const uint32 SubStreamReadValue = Substream.ReadBits(SubstreamBitCount);
+	UE_NET_ASSERT_EQ(SubStreamReadValue, 0b111111111111111U);
+
+	Reader.CommitSubstream(Substream);
+	const uint32 StreamReadValue = Reader.ReadBits(1U);
+	UE_NET_ASSERT_EQ(StreamReadValue, 0b1U);
+}
+
+UE_NET_TEST_FIXTURE(FNetBitStreamReaderSubstreamTest, CanReadFromSubSubstream)
+{
+	const uint32 Word = 0xFFFF0000U;
+	Reader.InitBits(&Word, sizeof(Word)*8);
+
+	const uint32 FirstStreamReadValue = Reader.ReadBits(16U);
+	UE_NET_ASSERT_EQ(FirstStreamReadValue, 0U);
+	{
+		constexpr uint32 SubstreamBitCount = 15U;
+		FNetBitStreamReader Substream1 = Reader.CreateSubstream(SubstreamBitCount);
+		FNetBitStreamReader Substream2 = Substream1.CreateSubstream();
+		const uint32 SubSubStreamReadValue = Substream2.ReadBits(SubstreamBitCount);
+		UE_NET_ASSERT_EQ(SubSubStreamReadValue, 0b111111111111111U);
+
+		Substream1.CommitSubstream(Substream2);
+		Reader.CommitSubstream(Substream1);
+	}
+	const uint32 SecondStreamReadValue = Reader.ReadBits(1U);
+	UE_NET_ASSERT_EQ(SecondStreamReadValue, 0b1U);
+}
+
+UE_NET_TEST_FIXTURE(FNetBitStreamReaderSubstreamTest, SubStreamSeek)
+{
+	Reader.InitBits(Buffer.GetBuffer(), Buffer.GetBufferCapacity()*8U);
+	Reader.ReadBits(32);
+
+	FNetBitStreamReader SubStream = Reader.CreateSubstream();
+	const uint32 SubStreamStartPos = SubStream.GetPosBits();
+	SubStream.Seek(0);
+	UE_NET_ASSERT_EQ(SubStreamStartPos, SubStream.GetPosBits());
+
+	Reader.CommitSubstream(SubStream);
+}
+
+// Misc tests
 UE_NET_TEST_FIXTURE(FNetBitStreamWriterTest, RollbackToValidPositionAfterOverflow)
 {
 	Writer.InitBytes(Buffer.GetBuffer(), Buffer.GetBufferCapacity());
@@ -625,19 +872,6 @@ UE_NET_TEST_FIXTURE(FNetBitStreamWriterTest, WriteScopeCanRewriteBits)
 	}
 }
 
-UE_NET_TEST_FIXTURE(FNetBitStreamWriterTest, SubStreamSeek)
-{
-	Writer.InitBytes(Buffer.GetBuffer(), Buffer.GetBufferCapacity());
-	Writer.WriteBits(0x12345678,32);
-
-	FNetBitStreamWriter SubStream = Writer.CreateSubstream();
-	const uint32 SubStreamStartPos = SubStream.GetPosBits();
-	SubStream.Seek(0);
-	UE_NET_ASSERT_EQ(SubStreamStartPos, SubStream.GetPosBits());
-
-	Writer.CommitSubstream(SubStream);
-}
-
 UE_NET_TEST_FIXTURE(FNetBitStreamWriterTest, WriteScopeOverflow)
 {
 	Writer.InitBytes(Buffer.GetBuffer(), Buffer.GetBufferCapacity());
diff --git a/Engine/Source/Runtime/Experimental/Iris/Core/Private/Iris/Serialization/NetBitStreamReader.cpp b/Engine/Source/Runtime/Experimental/Iris/Core/Private/Iris/Serialization/NetBitStreamReader.cpp
--- a/Engine/Source/Runtime/Experimental/Iris/Core/Private/Iris/Serialization/NetBitStreamReader.cpp
+++ b/Engine/Source/Runtime/Experimental/Iris/Core/Private/Iris/Serialization/NetBitStreamReader.cpp
@@ -2,9 +2,17 @@
 
 #include "Iris/Serialization/NetBitStreamReader.h"
 #include "Iris/Serialization/NetBitStreamUtils.h"
+#include "Iris/IrisConfig.h"
 #include "HAL/PlatformMemory.h"
 #include "Misc/AssertionMacros.h"
-#include <stdint.h>
+
+#if UE_NETBITSTREAMREADER_VALIDATE
+#	define UE_NETBITSTREAMREADER_CHECK(expr) check(expr)
+#	define UE_NETBITSTREAMREADER_CHECKF(expr, format, ...) checkf(expr, format, ##__VA_ARGS__)
+#else
+#	define UE_NETBITSTREAMREADER_CHECK(...) 
+#	define UE_NETBITSTREAMREADER_CHECKF(...) 
+#endif
 
 namespace UE::Net
 {
@@ -12,20 +20,28 @@ namespace UE::Net
 FNetBitStreamReader::FNetBitStreamReader()
 : Buffer(nullptr)
 , BufferBitCapacity(0)
+, BufferBitStartOffset(0)
 , BufferBitPosition(0)
 , PendingWord(0)
 , OverflowBitCount(0)
+, bHasSubstream(0)
+, bIsSubstream(0)
+, bIsInvalid(0)
 {
 }
 
 FNetBitStreamReader::~FNetBitStreamReader()
 {
+	UE_NETBITSTREAMREADER_CHECKF(!bHasSubstream, TEXT("FNetBitStreamReader is destroyed with active substream. "));
 }
 
 void FNetBitStreamReader::InitBits(const void* InBuffer, uint32 BitCount)
 {
 	check(InBuffer != nullptr);
-	checkf((uintptr_t(InBuffer) & 3) == 0, TEXT("Buffer needs to be 4-byte aligned."));
+	checkf((UPTRINT(InBuffer) & 3) == 0, TEXT("Buffer needs to be 4-byte aligned."));
+	// Re-initializing a substream or while having an active substream is not supported.
+	UE_NETBITSTREAMREADER_CHECK(!bHasSubstream && !bIsSubstream);
+
 	Buffer = static_cast<const uint32*>(InBuffer);
 	BufferBitCapacity = BitCount;
 	BufferBitPosition = 0;
@@ -37,6 +53,9 @@ void FNetBitStreamReader::InitBits(const void* InBuffer, uint32 BitCount)
 
 uint32 FNetBitStreamReader::ReadBits(uint32 BitCount)
 {
+	// Must be valid and must not read from main stream if it has a substream. Technically the latter would work, as we're just reading, but it's weird.
+	UE_NETBITSTREAMREADER_CHECK((!bIsInvalid) & (!bHasSubstream));
+
 	if (OverflowBitCount != 0)
 	{
 		return 0U;
@@ -81,6 +100,9 @@ uint32 FNetBitStreamReader::ReadBits(uint32 BitCount)
 
 void FNetBitStreamReader::ReadBitStream(uint32* InDst, uint32 BitCount)
 {
+	// Must be valid and must not read from main stream if it has a substream. Technically the latter would work, as we're just reading, but it's weird.
+	UE_NETBITSTREAMREADER_CHECK((!bIsInvalid) & (!bHasSubstream));
+
 	if (OverflowBitCount != 0)
 	{
 		return;
@@ -153,15 +175,19 @@ void FNetBitStreamReader::ReadBitStream(uint32* InDst, uint32 BitCount)
 
 void FNetBitStreamReader::Seek(uint32 BitPosition)
 {
-	if (BitPosition > BufferBitCapacity)
+	UE_NETBITSTREAMREADER_CHECK((!bIsInvalid) & (!bHasSubstream));
+
+	const uint32 AdjustedBitPosition = BitPosition + BufferBitStartOffset;
+	// We handle uint32 overflow as well which makes this code a bit more complicated. The OverflowBitCount may not always end up correct, but will be at least 1.
+	if ((AdjustedBitPosition > BufferBitCapacity) | (AdjustedBitPosition < BitPosition))
 	{
-		OverflowBitCount = BitPosition - BufferBitCapacity;
+		OverflowBitCount = FPlatformMath::Max(AdjustedBitPosition, BufferBitCapacity + 1U) - BufferBitCapacity;
 		return;
 	}
 
 	OverflowBitCount = 0;
 
-	BufferBitPosition = BitPosition;
+	BufferBitPosition = AdjustedBitPosition;
 	if ((BufferBitPosition & ~31U) < BufferBitCapacity)
 	{
 		PendingWord = INTEL_ORDER32(Buffer[BufferBitPosition >> 5U]);
@@ -176,4 +202,70 @@ void FNetBitStreamReader::DoOverflow()
 	}
 }
 
+FNetBitStreamReader FNetBitStreamReader::CreateSubstream(uint32 MaxBitCount)
+{
+	UE_NETBITSTREAMREADER_CHECK((!bIsInvalid) & (!bHasSubstream));
+
+	// Create a copy of this stream and overwrite the necessary members.
+	FNetBitStreamReader Substream = *this;
+	Substream.BufferBitStartOffset = BufferBitPosition;
+	Substream.bHasSubstream = 0;
+	Substream.bIsSubstream = 1;
+
+	bHasSubstream = 1;
+
+	/* If this stream is overflown make sure the substream will always be overflown as well!
+	 * We must be careful to ensure that a seek to the beginning of this stream will still cause the substream to be overflown.
+	 * We can ignore MaxBitCount completely because no writes will succeed anyway.
+	 */
+	if (OverflowBitCount)
+	{
+		Substream.BufferBitCapacity = Substream.BufferBitStartOffset;
+		// It's not vital that the OverflowBitCount is set as the user can reset it with a Seek(0) call. In any case no modifications to the bitstream can be done.
+		Substream.OverflowBitCount = OverflowBitCount;
+	}
+	else
+	{
+		Substream.BufferBitCapacity = BufferBitPosition + FPlatformMath::Min(MaxBitCount, BufferBitCapacity - BufferBitPosition);
+	}
+
+	return Substream;
 }
+
+void FNetBitStreamReader::CommitSubstream(FNetBitStreamReader& Substream)
+{
+	// Only accept substreams iff this is the parent and the substream has not overflown and has not previously been commited or discarded.
+	if (!ensure(bHasSubstream & (!Substream.bHasSubstream) & (!bIsInvalid) & (!Substream.bIsInvalid) & (Buffer == Substream.Buffer) & (BufferBitPosition == Substream.BufferBitStartOffset)))
+	{
+		return;
+	}
+
+	if (!Substream.IsOverflown())
+	{
+		BufferBitPosition = Substream.BufferBitPosition;
+		if ((Substream.BufferBitPosition & ~31U) < BufferBitCapacity)
+		{
+			PendingWord = INTEL_ORDER32(Buffer[Substream.BufferBitPosition >> 5U]);
+		}
+	}
+
+	bHasSubstream = 0;
+	Substream.bIsInvalid = 1;
+}
+
+void FNetBitStreamReader::DiscardSubstream(FNetBitStreamReader& Substream)
+{
+	// Only accept substreams iff this is the parent and the substream has not previously been commited or discarded.
+	if (!ensure(bHasSubstream & (!Substream.bHasSubstream) & (!bIsInvalid) & (!Substream.bIsInvalid) & (Buffer == Substream.Buffer) & (BufferBitPosition == Substream.BufferBitStartOffset)))
+	{
+		return;
+	}
+
+	bHasSubstream = 0;
+	Substream.bIsInvalid = 1;
+}
+
+}
+
+#undef UE_NETBITSTREAMREADER_CHECK
+#undef UE_NETBITSTREAMREADER_CHECKF
diff --git a/Engine/Source/Runtime/Experimental/Iris/Core/Public/Iris/IrisConfig.h b/Engine/Source/Runtime/Experimental/Iris/Core/Public/Iris/IrisConfig.h
--- a/Engine/Source/Runtime/Experimental/Iris/Core/Public/Iris/IrisConfig.h
+++ b/Engine/Source/Runtime/Experimental/Iris/Core/Public/Iris/IrisConfig.h
@@ -23,11 +23,11 @@ IRISCORE_API EReplicationSystem GetUseIrisReplicationCmdlineValue();
 
 /* NetBitStreamReader/Writer validation support */
 #ifndef UE_NETBITSTREAMWRITER_VALIDATE
-#if !(UE_BUILD_SHIPPING || UE_BUILD_TEST)
-#define UE_NETBITSTREAMWRITER_VALIDATE 1
-#else
-#define UE_NETBITSTREAMWRITER_VALIDATE 0
+#define UE_NETBITSTREAMWRITER_VALIDATE !(UE_BUILD_SHIPPING || UE_BUILD_TEST)
 #endif
+
+#ifndef UE_NETBITSTREAMREADER_VALIDATE
+#define UE_NETBITSTREAMREADER_VALIDATE !(UE_BUILD_SHIPPING || UE_BUILD_TEST)
 #endif
 
 /** CSV stats. */
@@ -38,4 +38,4 @@ IRISCORE_API EReplicationSystem GetUseIrisReplicationCmdlineValue();
 /** Enables code that detects non-thread safe access to network data */
 #ifndef UE_NET_THREAD_SAFETY_CHECK
 #	define UE_NET_THREAD_SAFETY_CHECK DO_CHECK
-#endif
\ No newline at end of file
+#endif
diff --git a/Engine/Source/Runtime/Experimental/Iris/Core/Public/Iris/Serialization/NetBitStreamReader.h b/Engine/Source/Runtime/Experimental/Iris/Core/Public/Iris/Serialization/NetBitStreamReader.h
--- a/Engine/Source/Runtime/Experimental/Iris/Core/Public/Iris/Serialization/NetBitStreamReader.h
+++ b/Engine/Source/Runtime/Experimental/Iris/Core/Public/Iris/Serialization/NetBitStreamReader.h
@@ -43,13 +43,19 @@ public:
 	IRISCORE_API void ReadBitStream(uint32* Dst, uint32 BitCount);
 
 	/**
-	 * Seek to a specific BitPosition. If the stream is overflown and you seek back to a position
+	 * Seek to a specific position from the start of the stream or substream. If the stream is overflown and you seek back to a position
 	 * where you can still read bits the stream will no longer be considered overflown.
 	 */
 	IRISCORE_API void Seek(uint32 BitPosition);
 
-	/** Returns the current bit position. */
-	inline uint32 GetPosBits() const { return BufferBitPosition; }
+	/** Returns the the current byte position. */
+	inline uint32 GetPosBytes() const { return (BufferBitPosition - BufferBitStartOffset + 7) >> 3U; }
+
+	/** Returns the current bit position */
+	inline uint32 GetPosBits() const { return BufferBitPosition - BufferBitStartOffset; }
+
+	/** Returns the absolute bit position */
+	inline uint32 GetAbsolutePosBits() const { return BufferBitPosition; }
 
 	/** Returns the number of bits that can be read before overflowing. */
 	inline uint32 GetBitsLeft() const { return (OverflowBitCount ? 0U : (BufferBitCapacity - BufferBitPosition)); }
@@ -60,15 +66,43 @@ public:
 	/** Returns whether the stream is overflown or not. */
 	inline bool IsOverflown() const { return OverflowBitCount != 0; }
 
+	/** 
+	 * Creates a substream at the current bit position. The substream must be committed or discarded. Only one active substream at a time is allowed,
+	 * but a substream can have an active substream as well. Once the substream has been commited or discarded a new substream may be created. No
+	 * reads may be performed on this stream until the substream has been committed or discarded.  
+	 *
+	 * The returned FNetBitStreamReader will have similar behavior to a newly constructed regular FNetBitStreamWriter. 
+	 * 
+	 * @param MaxBitCount The maximum allowed bits that may be read. The value will be clamped to the number of bits left in this stream/substream. If it's a requirement a specific size is supported you can verify it with GetBitsLeft().
+	 * @return A FNetBitStreamReader.
+	 */
+	IRISCORE_API FNetBitStreamReader CreateSubstream(uint32 MaxBitCount = ~0U);
+
+	/**
+	 * Commits a substream to this stream. Substreams that are overflown or do not belong to this stream will be ignored. 
+	 * If the substream is valid then this stream's bit position will be updated.
+	 */
+	IRISCORE_API void CommitSubstream(FNetBitStreamReader& Substream);
+
+	/** Discards a substream of this stream. This stream's bit position will remain intact. */
+	IRISCORE_API void DiscardSubstream(FNetBitStreamReader& Substream);
+
 private:
 	const uint32* Buffer;
+	// The BufferBitCapacity is an absolute bit position indicating the bit after the last valid bit position to read.
 	uint32 BufferBitCapacity;
+	// For substreams this indicate the absolute bit position in the buffer where it will start reading
+	uint32 BufferBitStartOffset;
 	uint32 BufferBitPosition;
 	uint32 PendingWord;
 	uint32 OverflowBitCount;
+
+	uint32 bHasSubstream : 1;
+	uint32 bIsSubstream : 1;
+	uint32 bIsInvalid : 1;
 };
 
 }
 
 // Always report the actual bitstream position, even on overflow. This normally allows for better comparisons between sending and receiving side when bitstream errors occur.
-inline uint32 GetBitStreamPositionForNetTrace(const UE::Net::FNetBitStreamReader& Stream) { return Stream.GetPosBits(); }
+inline uint32 GetBitStreamPositionForNetTrace(const UE::Net::FNetBitStreamReader& Stream) { return Stream.GetAbsolutePosBits(); }
