Iris
- Fixed compilation error in ReplicationSystemTests when compiling for Linux

diff --git a/Engine/Plugins/Runtime/ReplicationSystemTestPlugin/Source/Private/Tests/DataStream/TestDataStream.cpp b/Engine/Plugins/Runtime/ReplicationSystemTestPlugin/Source/Private/Tests/DataStream/TestDataStream.cpp
--- a/Engine/Plugins/Runtime/ReplicationSystemTestPlugin/Source/Private/Tests/DataStream/TestDataStream.cpp
+++ b/Engine/Plugins/Runtime/ReplicationSystemTestPlugin/Source/Private/Tests/DataStream/TestDataStream.cpp
@@ -385,7 +385,7 @@ protected:
 	void RoundTrip()
 	{
 		Server->UpdateAndSend({Client}, true);
-		Client->UpdateAndSend({Server}, true);
+		Client->UpdateAndSend(Server, true);
 	}
 	const FName DataStreamName = "DynamicNetToken";
 
@@ -558,7 +558,7 @@ UE_NET_TEST_FIXTURE(FTestDynamicCreateDataStreamFixture, TestDynamicDataStream_D
 	}
 
 	// Drop PendingCreate from client
-	Client->UpdateAndSend({Server}, false);
+	Client->UpdateAndSend(Server, false);
 	
 	// Server state should be WaitForCreateConfirmation
 	UE_NET_ASSERT_TRUE(GetDataStreamStateOnServer() == UDataStream::EDataStreamState::WaitOnCreateConfirmation);
@@ -567,7 +567,7 @@ UE_NET_TEST_FIXTURE(FTestDynamicCreateDataStreamFixture, TestDynamicDataStream_D
 	UE_NET_ASSERT_TRUE(GetDataStreamStateOnClient() == UDataStream::EDataStreamState::PendingCreate);
 
 	// Send and deliver to server
-	Client->UpdateAndSend({Server}, true);
+	Client->UpdateAndSend(Server, true);
 
 	// Server state should now be open
 	UE_NET_ASSERT_TRUE(GetDataStreamStateOnServer() == UDataStream::EDataStreamState::Open);
@@ -602,7 +602,7 @@ UE_NET_TEST_FIXTURE(FTestDynamicCreateDataStreamFixture, TestDynamicDataStream_R
 	UE_NET_ASSERT_TRUE(GetDataStreamStateOnClient() == UDataStream::EDataStreamState::PendingClose);
 
 	// Send and deliver to server
-	Client->UpdateAndSend({Server}, true);
+	Client->UpdateAndSend(Server, true);
 	
 	UE_NET_ASSERT_TRUE(GetDataStreamStateOnServer() == UDataStream::EDataStreamState::PendingClose);
 	UE_NET_ASSERT_TRUE(GetDataStreamStateOnClient() == UDataStream::EDataStreamState::WaitOnCloseConfirmation);
@@ -639,14 +639,14 @@ UE_NET_TEST_FIXTURE(FTestDynamicCreateDataStreamFixture, TestDynamicDataStream_R
 	UE_NET_ASSERT_TRUE(GetDataStreamStateOnClient() == UDataStream::EDataStreamState::PendingClose);
 
 	// Drop send to server
-	Client->UpdateAndSend({Server}, false);
+	Client->UpdateAndSend(Server, false);
 	
 	// Verify expected stream state
 	UE_NET_ASSERT_TRUE(GetDataStreamStateOnServer() == UDataStream::EDataStreamState::Open);
 	UE_NET_ASSERT_TRUE(GetDataStreamStateOnClient() == UDataStream::EDataStreamState::PendingClose);
 
 	// Send and deliver data
-	Client->UpdateAndSend({Server}, true);
+	Client->UpdateAndSend(Server, true);
 
 	// Verify expected stream state
 	UE_NET_ASSERT_TRUE(GetDataStreamStateOnServer() == UDataStream::EDataStreamState::PendingClose);
@@ -769,7 +769,7 @@ UE_NET_TEST_FIXTURE(FTestDynamicCreateDataStreamFixture, TestDynamicDataStreamCl
 	Server->DeliverTo(Client, true);
 
 	// Acknowledge pending close
-	Client->UpdateAndSend({Server}, true);
+	Client->UpdateAndSend(Server, true);
 
 	// We expect the stream to be PendingClose as we still have important data in flight
 	UE_NET_ASSERT_TRUE(GetDataStreamStateOnServer() == UDataStream::EDataStreamState::PendingClose);
@@ -778,7 +778,7 @@ UE_NET_TEST_FIXTURE(FTestDynamicCreateDataStreamFixture, TestDynamicDataStreamCl
 	Server->UpdateAndSend({Client}, true);
 
 	// Acknowledge pending close
-	Client->UpdateAndSend({Server}, true);
+	Client->UpdateAndSend(Server, true);
 
 	// Deliver PendingClose
 	RoundTrip();
diff --git a/Engine/Plugins/Runtime/ReplicationSystemTestPlugin/Source/Private/Tests/ReplicationSystem/ReplicationSystemServerClientTestFixture.h b/Engine/Plugins/Runtime/ReplicationSystemTestPlugin/Source/Private/Tests/ReplicationSystem/ReplicationSystemServerClientTestFixture.h
--- a/Engine/Plugins/Runtime/ReplicationSystemTestPlugin/Source/Private/Tests/ReplicationSystem/ReplicationSystemServerClientTestFixture.h
+++ b/Engine/Plugins/Runtime/ReplicationSystemTestPlugin/Source/Private/Tests/ReplicationSystem/ReplicationSystemServerClientTestFixture.h
@@ -43,7 +43,9 @@ public:
 		bool bDynamicCreate = false;
 		bool bValid = true;
 	};
-	void AddDataStreamDefinition(const TCHAR* StreamName, const TCHAR* ClassPath, const FAddDataStreamDefinitionParams& Params = FAddDataStreamDefinitionParams());
+
+	void AddDataStreamDefinition(const TCHAR* StreamName, const TCHAR* ClassPath, const FAddDataStreamDefinitionParams& Params);
+	void AddDataStreamDefinition(const TCHAR* StreamName, const TCHAR* ClassPath) { AddDataStreamDefinition(StreamName, ClassPath, FAddDataStreamDefinitionParams()); }
 	void FixupDefinitions() { DataStreamDefinitions->FixupDefinitions(); }
 
 protected:
