Added utility methods to write bytes to and from NetBitStream without requirement on storagetype alignment.

diff --git a/Engine/Plugins/Runtime/ReplicationSystemTestPlugin/Source/Private/Tests/Serialization/TestNetBitStreamUtil.cpp b/Engine/Plugins/Runtime/ReplicationSystemTestPlugin/Source/Private/Tests/Serialization/TestNetBitStreamUtil.cpp
--- a/Engine/Plugins/Runtime/ReplicationSystemTestPlugin/Source/Private/Tests/Serialization/TestNetBitStreamUtil.cpp
+++ b/Engine/Plugins/Runtime/ReplicationSystemTestPlugin/Source/Private/Tests/Serialization/TestNetBitStreamUtil.cpp
@@ -115,4 +115,101 @@ UE_NET_TEST_FIXTURE(FNetBitStreamUtilTest, TestTooLongStringWritesEmptyString)
 	Writer = FNetBitStreamWriter();
 }
 
+UE_NET_TEST_FIXTURE(FNetBitStreamUtilTest, TestReadWriteBytes)
+{
+	alignas(16) const char SrcBuffer[] = "012345679";
+
+	// Test write from offset
+	{
+		Writer.InitBytes(BitStreamBuffer, BitStreamBufferSize);
+
+		const char* SrcBytes = SrcBuffer + 1;
+		const uint32 ByteCount = 3U;
+
+		WriteBytes(&Writer, reinterpret_cast<const uint8*>(SrcBytes), ByteCount);
+		Writer.CommitWrites();
+
+		UE_NET_ASSERT_FALSE(Writer.IsOverflown());
+		UE_NET_ASSERT_GT(Writer.GetPosBits(), 0U);
+
+		char DstBytes[ByteCount];
+
+		Reader.InitBits(BitStreamBuffer, Writer.GetPosBits());
+		ReadBytes(&Reader, reinterpret_cast<uint8*>(DstBytes), ByteCount);
+		UE_NET_ASSERT_FALSE(Reader.IsOverflown());
+		UE_NET_ASSERT_EQ(FCStringAnsi::Strncmp(SrcBytes, DstBytes, ByteCount), 0);
+	}
+
+	// Test write entire buffer
+	{
+		Writer.InitBytes(BitStreamBuffer, BitStreamBufferSize);
+
+		const char* SrcBytes = SrcBuffer;
+		const uint32 ByteCount = 10;
+
+		WriteBytes(&Writer, reinterpret_cast<const uint8*>(SrcBytes), ByteCount);
+		Writer.CommitWrites();
+
+		UE_NET_ASSERT_FALSE(Writer.IsOverflown());
+		UE_NET_ASSERT_GT(Writer.GetPosBits(), 0U);
+
+		char DstBytes[ByteCount];
+
+		Reader.InitBits(BitStreamBuffer, Writer.GetPosBits());
+		ReadBytes(&Reader, reinterpret_cast<uint8*>(DstBytes), ByteCount);
+		UE_NET_ASSERT_FALSE(Reader.IsOverflown());
+		UE_NET_ASSERT_EQ(FCStringAnsi::Strncmp(SrcBytes, DstBytes, ByteCount), 0);
+	}
+
+	// Test write with misaligned start
+	{
+		Writer.InitBytes(BitStreamBuffer, BitStreamBufferSize);
+
+		const char* SrcBytes = SrcBuffer + 3;
+		const uint32 ByteCount = 6;
+
+		WriteBytes(&Writer, reinterpret_cast<const uint8*>(SrcBytes), ByteCount);
+		Writer.CommitWrites();
+
+		UE_NET_ASSERT_FALSE(Writer.IsOverflown());
+		UE_NET_ASSERT_GT(Writer.GetPosBits(), 0U);
+
+		char DstBytes[ByteCount];
+
+		Reader.InitBits(BitStreamBuffer, Writer.GetPosBits());
+		ReadBytes(&Reader, reinterpret_cast<uint8*>(DstBytes), ByteCount);
+		UE_NET_ASSERT_FALSE(Reader.IsOverflown());
+		UE_NET_ASSERT_EQ(FCStringAnsi::Strncmp(SrcBytes, DstBytes, ByteCount), 0);
+	}
+
+	// Test read to offset
+	{
+		Writer.InitBytes(BitStreamBuffer, BitStreamBufferSize);
+
+		const char* SrcBytes = SrcBuffer;
+		const uint32 ByteCount = 10;
+		const uint32 DstOffset = 3;
+
+		WriteBytes(&Writer, reinterpret_cast<const uint8*>(SrcBytes), ByteCount);
+		Writer.CommitWrites();
+
+		UE_NET_ASSERT_FALSE(Writer.IsOverflown());
+		UE_NET_ASSERT_GT(Writer.GetPosBits(), 0U);
+
+		alignas(16) char DstBuffer[16] = {};
+
+		char* DstBytes = DstBuffer + DstOffset;
+
+		Reader.InitBits(BitStreamBuffer, Writer.GetPosBits());
+		ReadBytes(&Reader, reinterpret_cast<uint8*>(DstBytes), ByteCount);
+
+		UE_NET_ASSERT_FALSE(Reader.IsOverflown());
+		UE_NET_ASSERT_EQ(FCStringAnsi::Strncmp(SrcBytes, DstBytes, ByteCount), 0);
+		// Verify byte before and after
+		UE_NET_ASSERT_EQ(DstBuffer[DstOffset - 1], (char)0);
+		UE_NET_ASSERT_EQ(DstBuffer[DstOffset + ByteCount], (char)0);
+	}
+}
+
+
 }
diff --git a/Engine/Source/Runtime/Experimental/Iris/Core/Private/Iris/Serialization/NetBitStreamUtil.cpp b/Engine/Source/Runtime/Experimental/Iris/Core/Private/Iris/Serialization/NetBitStreamUtil.cpp
--- a/Engine/Source/Runtime/Experimental/Iris/Core/Private/Iris/Serialization/NetBitStreamUtil.cpp
+++ b/Engine/Source/Runtime/Experimental/Iris/Core/Private/Iris/Serialization/NetBitStreamUtil.cpp
@@ -915,4 +915,62 @@ bool ReadAndVerifySentinelBits(FNetBitStreamReader* Reader, const TCHAR* ErrorSt
 	return ensureAlwaysMsgf(!Reader->IsOverflown() &&  ExpectedValue == ReadValue, TEXT("ReadAndVerifySentinelBits %s failed OverFlow %u Got 0x%u != 0x%u"), ErrorString, Reader->IsOverflown() ? 1U : 0U, ReadValue, ExpectedValue);
 }
 
+void ReadBytes(FNetBitStreamReader* Reader, uint8* Destination, uint32 BytesToRead)
+{
+	if (BytesToRead == 0)
+	{
+		return;
+	}
+
+	uint8* DstData = Destination;
+	// Fix up dst alignment
+	while (((uintptr_t)(DstData) & 3) && BytesToRead)
+	{
+		*(DstData++) = static_cast<uint8>(Reader->ReadBits(8U));
+		--BytesToRead;
+	}
+	// Write full words
+	const uint32 FullWordsToRead = BytesToRead / 4U;
+	if (FullWordsToRead)
+	{
+		Reader->ReadBitStream(reinterpret_cast<uint32*>(DstData), FullWordsToRead * 32U);
+		DstData += FullWordsToRead * 4U;
+		BytesToRead -= FullWordsToRead * 4U;
+	}
+	while (BytesToRead)
+	{
+		*(DstData++) = static_cast<uint8>(Reader->ReadBits(8U));
+		--BytesToRead;
+	}
+}
+
+void WriteBytes(FNetBitStreamWriter* Writer, const uint8* Src, uint32 BytesToWrite)
+{
+	if (BytesToWrite == 0)
+	{
+		return;
+	}
+
+	const uint8* SrcData = Src;
+	// Fix up src alignment
+	while (((uintptr_t)(SrcData) & 3) && BytesToWrite)
+	{
+		Writer->WriteBits(*(SrcData++), 8U);
+		--BytesToWrite;
+	}
+	// Write full words
+	const uint32 FullWordsToWrite = BytesToWrite / 4U;
+	if (FullWordsToWrite)
+	{
+		Writer->WriteBitStream(reinterpret_cast<const uint32*>(SrcData), 0U, FullWordsToWrite * 32);
+		SrcData += FullWordsToWrite * 4U;
+		BytesToWrite -= FullWordsToWrite * 4U;
+	}
+	while (BytesToWrite)
+	{
+		Writer->WriteBits(*(SrcData++), 8U);
+		--BytesToWrite;
+	}
+}
+
 }
diff --git a/Engine/Source/Runtime/Experimental/Iris/Core/Public/Iris/Serialization/NetBitStreamUtil.h b/Engine/Source/Runtime/Experimental/Iris/Core/Public/Iris/Serialization/NetBitStreamUtil.h
--- a/Engine/Source/Runtime/Experimental/Iris/Core/Public/Iris/Serialization/NetBitStreamUtil.h
+++ b/Engine/Source/Runtime/Experimental/Iris/Core/Public/Iris/Serialization/NetBitStreamUtil.h
@@ -216,5 +216,17 @@ IRISCORE_API void WriteSentinelBits(FNetBitStreamWriter* Writer, uint32 BitCount
 
 IRISCORE_API bool ReadAndVerifySentinelBits(FNetBitStreamReader* Reader, const TCHAR* ErrorString = TEXT("Sentinel"), uint32 BitCount = 32U);
 
+/** 
+* Write bytes
+* Note: The purpose of this method is for cases when we cannot guarantee that the buffer is word aligned. If buffer is known to be word aligned, it is faster to use NetBitStreamReader::WriteBitStream directly.
+*       Does not write the size.
+*/
+IRISCORE_API void WriteBytes(FNetBitStreamWriter* Writer, const uint8* Src, uint32 BytesToWrite);
+
+/** 
+ * Read bytes into preallocated buffer
+ * Note: The purpose of this method is for cases when we cannot guarantee that the buffer is word aligned. If buffer is known to be word aligned, it is faster to use NetBitStreamReader::ReadBitStream directly.
+ */
+IRISCORE_API void ReadBytes(FNetBitStreamReader* Reader, uint8* Destination, uint32 BytesToRead);
 
 }
