NetBitArray
* Added AndBitValue and OrBitValue that allows you to AND|OR a single bit index

#rnx

diff --git a/Engine/Plugins/Runtime/ReplicationSystemTestPlugin/Source/Private/Tests/Core/TestNetBitArray.cpp b/Engine/Plugins/Runtime/ReplicationSystemTestPlugin/Source/Private/Tests/Core/TestNetBitArray.cpp
--- a/Engine/Plugins/Runtime/ReplicationSystemTestPlugin/Source/Private/Tests/Core/TestNetBitArray.cpp
+++ b/Engine/Plugins/Runtime/ReplicationSystemTestPlugin/Source/Private/Tests/Core/TestNetBitArray.cpp
@@ -425,6 +425,125 @@ UE_NET_TEST(FNetBitArrayView, ClearBit)
 	UE_NET_ASSERT_FALSE(BitArray.GetBit(TestBitIndex));
 }
 
+UE_NET_TEST(FNetBitArrayView, OrBitValue)
+{
+	uint32 WordBuffer[] = { ~(0U), 0U, 0xAAAA };
+	FNetBitArrayView BitArray(&WordBuffer[0], 8 * sizeof(WordBuffer));
+	
+	// Test first uint32 (~0)
+	{
+		UE_NET_ASSERT_EQ(BitArray.CountSetBits(0, 32), 32U);
+
+		UE_NET_ASSERT_TRUE(BitArray.GetBit(1));
+		BitArray.OrBitValue(1, true);
+		UE_NET_ASSERT_TRUE(BitArray.GetBit(1));
+
+		UE_NET_ASSERT_EQ(BitArray.CountSetBits(0, 32), 32U);
+
+		UE_NET_ASSERT_TRUE(BitArray.GetBit(2));
+		BitArray.OrBitValue(2, false);
+		UE_NET_ASSERT_TRUE(BitArray.GetBit(2));
+
+		UE_NET_ASSERT_EQ(BitArray.CountSetBits(0, 32), 32U);
+	}
+
+	// Test second uint32 (0)
+	{
+		UE_NET_ASSERT_EQ(BitArray.CountSetBits(32, 32), 0U);
+
+		UE_NET_ASSERT_FALSE(BitArray.GetBit(33));
+		BitArray.OrBitValue(33, true);
+		UE_NET_ASSERT_TRUE(BitArray.GetBit(33));
+
+		UE_NET_ASSERT_EQ(BitArray.CountSetBits(32, 32), 1U);
+
+		UE_NET_ASSERT_FALSE(BitArray.GetBit(34));
+		BitArray.OrBitValue(34, false);
+		UE_NET_ASSERT_FALSE(BitArray.GetBit(34));
+
+		UE_NET_ASSERT_EQ(BitArray.CountSetBits(32, 32), 1U);
+	}
+
+	// Test third uint32 (0xAAAA)
+	{
+		UE_NET_ASSERT_EQ(BitArray.CountSetBits(64, 32), 8U);
+
+		UE_NET_ASSERT_TRUE(BitArray.GetBit(65));
+		BitArray.OrBitValue(65, true);
+		UE_NET_ASSERT_TRUE(BitArray.GetBit(65));
+
+		UE_NET_ASSERT_EQ(BitArray.CountSetBits(64, 32), 8U);
+
+		UE_NET_ASSERT_FALSE(BitArray.GetBit(66));
+		BitArray.OrBitValue(66, true);
+		UE_NET_ASSERT_TRUE(BitArray.GetBit(66));
+
+		UE_NET_ASSERT_EQ(BitArray.CountSetBits(64, 32), 9U);
+	}	
+}
+
+UE_NET_TEST(FNetBitArrayView, AndBitValue)
+{
+	uint32 WordBuffer[] = { ~(0U), 0U, 0xAAAA };
+	FNetBitArrayView BitArray(&WordBuffer[0], 8 * sizeof(WordBuffer));
+
+	// Test first uint32 (~0)
+	{
+		UE_NET_ASSERT_EQ(BitArray.CountSetBits(0, 32), 32U);
+
+		UE_NET_ASSERT_TRUE(BitArray.GetBit(1));
+		BitArray.AndBitValue(1, true);
+		UE_NET_ASSERT_TRUE(BitArray.GetBit(1));
+
+		UE_NET_ASSERT_EQ(BitArray.CountSetBits(0, 32), 32U);
+
+		UE_NET_ASSERT_TRUE(BitArray.GetBit(2));
+		BitArray.AndBitValue(2, false);
+		UE_NET_ASSERT_FALSE(BitArray.GetBit(2));
+
+		UE_NET_ASSERT_EQ(BitArray.CountSetBits(0, 32), 31U);
+	}
+
+	// Test second uint32 (0)
+	{
+		UE_NET_ASSERT_EQ(BitArray.CountSetBits(32, 32), 0U);
+
+		UE_NET_ASSERT_FALSE(BitArray.GetBit(33));
+		BitArray.AndBitValue(33, true);
+		UE_NET_ASSERT_FALSE(BitArray.GetBit(33));
+
+		UE_NET_ASSERT_EQ(BitArray.CountSetBits(32, 32), 0U);
+
+		UE_NET_ASSERT_FALSE(BitArray.GetBit(34));
+		BitArray.AndBitValue(34, false);
+		UE_NET_ASSERT_FALSE(BitArray.GetBit(34));
+
+		UE_NET_ASSERT_EQ(BitArray.CountSetBits(32, 32), 0U);
+	}
+
+	// Test third uint32 (0xAAAA)
+	{
+		UE_NET_ASSERT_EQ(BitArray.CountSetBits(64, 32), 8U);
+
+		UE_NET_ASSERT_FALSE(BitArray.GetBit(64));
+		BitArray.AndBitValue(64, true);
+		UE_NET_ASSERT_FALSE(BitArray.GetBit(64));
+
+		UE_NET_ASSERT_EQ(BitArray.CountSetBits(64, 32), 8U);
+
+		UE_NET_ASSERT_TRUE(BitArray.GetBit(65));
+		BitArray.AndBitValue(65, true);
+		UE_NET_ASSERT_TRUE(BitArray.GetBit(65));
+
+		UE_NET_ASSERT_EQ(BitArray.CountSetBits(64, 32), 8U);
+
+		UE_NET_ASSERT_TRUE(BitArray.GetBit(65));
+		BitArray.AndBitValue(65, false);
+		UE_NET_ASSERT_FALSE(BitArray.GetBit(65));
+
+		UE_NET_ASSERT_EQ(BitArray.CountSetBits(64, 32), 7U);
+	}
+}
 UE_NET_TEST(FNetBitArrayView, GetSetBit)
 {	
 	{
diff --git a/Engine/Source/Runtime/Net/Core/Public/Net/Core/NetBitArray.h b/Engine/Source/Runtime/Net/Core/Public/Net/Core/NetBitArray.h
--- a/Engine/Source/Runtime/Net/Core/Public/Net/Core/NetBitArray.h
+++ b/Engine/Source/Runtime/Net/Core/Public/Net/Core/NetBitArray.h
@@ -200,6 +200,12 @@ public:
 	/** Clear the bit with the specified index */
 	void ClearBit(uint32 Index);
 
+	/** Do an Or operation on the specified index */
+	inline void OrBitValue(uint32 Index, bool bValue);
+
+	/** Do an And operation on the specified index */
+	inline void AndBitValue(uint32 Index, bool bValue);
+
 	/** Get the bit with the specified index */
 	bool GetBit(uint32 Index) const;
 
@@ -379,6 +385,12 @@ public:
 	/** Get the bit with the specified index */
 	inline bool GetBit(uint32 Index) const;
 
+	/** Do an Or operation on the specified index */
+	inline void OrBitValue(uint32 Index, bool bValue);
+
+	/** Do an And operation on the specified index */
+	inline void AndBitValue(uint32 Index, bool bValue);
+
 	/** Find first zero bit. Returns InvalidIndex if no zero bit was found */
 	inline uint32 FindFirstZero() const;
 
@@ -682,6 +694,28 @@ private:
 		return (Storage[WordIndex] & WordMask) != 0u;
 	}
 
+	static void OrBitValue(StorageWordType* Storage, const uint32 BitCount, const uint32 Index, const bool bValue)
+	{
+		UE_NETBITARRAY_CHECK(Index < BitCount);
+
+		const SIZE_T WordIndex = Index / WordBitCount;
+		const StorageWordType WordMask = (StorageWordType(1) << (Index & (WordBitCount - 1)));
+		const StorageWordType ValueMask = bValue ? WordMask : StorageWordType(0);
+
+		Storage[WordIndex] = Storage[WordIndex] | ValueMask;
+	}
+
+	static void AndBitValue(StorageWordType* Storage, const uint32 BitCount, const uint32 Index, const bool bValue)
+	{
+		UE_NETBITARRAY_CHECK(Index < BitCount);
+
+		const SIZE_T WordIndex = Index / WordBitCount;
+		const StorageWordType WordMask = (StorageWordType(1) << (Index & (WordBitCount - 1)));
+		const StorageWordType ValueMask = bValue ? ~StorageWordType(0) : ~WordMask;
+
+		Storage[WordIndex] = Storage[WordIndex] & ValueMask;
+	}
+
 	static uint32 FindFirstZero(const StorageWordType* Storage, const uint32 WordCount, const uint32 BitCount)
 	{
 		for (uint32 WordIt = 0, CurrentBitIndex = 0; WordIt < WordCount; ++WordIt, CurrentBitIndex += WordBitCount)
@@ -1146,6 +1180,16 @@ inline bool FNetBitArray::GetBit(uint32 Index) const
 	return FNetBitArrayHelper::GetBit(Storage.GetData(), BitCount, Index);
 }
 
+inline void FNetBitArray::OrBitValue(uint32 Index, bool bValue)
+{
+	FNetBitArrayHelper::OrBitValue(Storage.GetData(), BitCount, Index, bValue);
+}
+
+inline void FNetBitArray::AndBitValue(uint32 Index, bool bValue)
+{
+	FNetBitArrayHelper::AndBitValue(Storage.GetData(), BitCount, Index, bValue);
+}
+
 inline uint32 FNetBitArray::FindFirstZero() const
 {
 	return FNetBitArrayHelper::FindFirstZero(Storage.GetData(), Storage.Num(), BitCount);
@@ -1408,6 +1452,16 @@ bool FNetBitArrayView::GetBit(uint32 Index) const
 	return FNetBitArrayHelper::GetBit(Storage, BitCount, Index);
 }
 
+void FNetBitArrayView::OrBitValue(uint32 Index, bool bValue)
+{
+	FNetBitArrayHelper::OrBitValue(Storage, BitCount, Index, bValue);
+}
+
+void FNetBitArrayView::AndBitValue(uint32 Index, bool bValue)
+{
+	FNetBitArrayHelper::AndBitValue(Storage, BitCount, Index, bValue);
+}
+
 uint32 FNetBitArrayView::FindFirstZero() const
 {
 	return FNetBitArrayHelper::FindFirstZero(Storage, WordCount, BitCount);
