Added FNetIDVariant, a simple variant wrapper to hold a generic or Iris network object ID (FNetworkGUID or FNetRefHandle).
Intended to aid in transitioning systems that were using FNetworkGUID to Iris.

diff --git a/Engine/Plugins/Runtime/ReplicationSystemTestPlugin/Source/Private/Tests/Engine/TestNetIDVariant.cpp b/Engine/Plugins/Runtime/ReplicationSystemTestPlugin/Source/Private/Tests/Engine/TestNetIDVariant.cpp
new file mode 100644
--- /dev/null
+++ b/Engine/Plugins/Runtime/ReplicationSystemTestPlugin/Source/Private/Tests/Engine/TestNetIDVariant.cpp
@@ -0,0 +1,173 @@
+// Copyright Epic Games, Inc. All Rights Reserved.
+
+#include "NetworkAutomationTest.h"
+#include "NetworkAutomationTestMacros.h"
+#include "Tests/ReplicationSystem/ReplicationSystemServerClientTestFixture.h"
+#include "Net/NetIDVariant.h"
+#include "UObject/CoreNet.h"
+
+namespace UE::Net::Private
+{
+	UE_NET_TEST(NetIDVariant, Empty)
+	{
+		FNetIDVariant DefaultVariant;
+
+		UE_NET_ASSERT_FALSE(DefaultVariant.GetVariant().IsType<FNetworkGUID>());
+		UE_NET_ASSERT_FALSE(DefaultVariant.GetVariant().IsType<FNetRefHandle>());
+
+		// Serialization
+		FNetBitWriter TempWriter(1024);
+		TempWriter << DefaultVariant;
+
+		FNetBitReader TempReader(nullptr, TempWriter.GetData(), TempWriter.GetNumBits());
+		FNetIDVariant DeserializedDefaultVariant;
+		TempReader << DeserializedDefaultVariant;
+
+		UE_NET_ASSERT_TRUE(DeserializedDefaultVariant.GetVariant().IsType<FNetIDVariant::FEmptyID>());
+	}
+
+	UE_NET_TEST(NetIDVariant, Hashing)
+	{
+		// Testing hashing, equality, and TSet/TMap support
+		FNetIDVariant DefaultVariant;
+		FNetIDVariant InvalidGUID{FNetworkGUID()};
+		FNetIDVariant InvalidRefHandle{FNetRefHandle()};
+
+		const uint32 DefaultHash = GetTypeHash(DefaultVariant);
+		const uint32 InvalidGUIDHash = GetTypeHash(InvalidGUID);
+		const uint32 InvalidRefHandleHash = GetTypeHash(InvalidRefHandle);
+
+		// Try to avoid some simple hash collisions
+		UE_NET_ASSERT_NE(DefaultHash, InvalidGUIDHash);
+		UE_NET_ASSERT_NE(DefaultHash, InvalidRefHandleHash);
+		UE_NET_ASSERT_NE(InvalidGUIDHash, InvalidRefHandleHash);
+
+		// Set uniqueness
+		TSet<FNetIDVariant> SetTest;
+		SetTest.Add(DefaultVariant);
+		SetTest.Add(InvalidGUID);
+		SetTest.Add(InvalidRefHandle);
+		UE_NET_ASSERT_EQ(SetTest.Num(), 3);
+
+		// Set duplicates
+		SetTest.Add(DefaultVariant);
+		SetTest.Add(InvalidGUID);
+		SetTest.Add(InvalidRefHandle);
+		UE_NET_ASSERT_EQ(SetTest.Num(), 3);
+
+		FNetIDVariant ValidGUID1{FNetworkGUID::CreateFromIndex(1, false)};
+		FNetIDVariant ValidGUID2{FNetworkGUID::CreateFromIndex(2, false)};
+
+		const uint32 ValidGUID1Hash = GetTypeHash(ValidGUID1);
+		const uint32 ValidGUID2Hash = GetTypeHash(ValidGUID2);
+
+		UE_NET_ASSERT_NE(ValidGUID1Hash, ValidGUID2Hash);
+
+		// Need to get valid FNetRefHandles
+		UE::Net::FReplicationSystemTestNode TestNode(true, TEXT("NetIDVariantSetTest"));
+		UTestReplicatedIrisObject* TestObject1 = TestNode.CreateObject(0, 0);
+		UTestReplicatedIrisObject* TestObject2 = TestNode.CreateObject(0, 0);
+
+		FNetIDVariant ValidRefHandle1{TestNode.GetReplicationBridge()->BeginReplication(TestObject1)};
+		FNetIDVariant ValidRefHandle2{TestNode.GetReplicationBridge()->BeginReplication(TestObject2)};
+
+		const uint32 ValidRefHandle1Hash = GetTypeHash(ValidRefHandle1);
+		const uint32 ValidRefHandle2Hash = GetTypeHash(ValidRefHandle2);
+
+		UE_NET_ASSERT_NE(ValidRefHandle1Hash, ValidRefHandle2Hash);
+
+		// Set uniqueness
+		SetTest.Add(ValidGUID1);
+		SetTest.Add(ValidGUID2);
+		SetTest.Add(ValidRefHandle1);
+		SetTest.Add(ValidRefHandle2);
+		UE_NET_ASSERT_EQ(SetTest.Num(), 7);
+
+		// Set duplicates
+		SetTest.Add(ValidGUID1);
+		SetTest.Add(ValidGUID2);
+		SetTest.Add(ValidRefHandle1);
+		SetTest.Add(ValidRefHandle2);
+		UE_NET_ASSERT_EQ(SetTest.Num(), 7);
+	}
+
+	UE_NET_TEST(NetIDVariant, NetworkGUID)
+	{
+		const FNetworkGUID ValidGUID = FNetworkGUID::CreateFromIndex(100, false);
+		FNetIDVariant ValidVariantGUID(ValidGUID);
+
+		const FNetworkGUID InvalidGUID;
+		FNetIDVariant InvalidVariantGUID(InvalidGUID);
+
+		// Basic round-trip
+		FNetworkGUID ValidOutputGUID = ValidVariantGUID.GetVariant().Get<FNetworkGUID>();
+		UE_NET_ASSERT_EQ(ValidGUID, ValidOutputGUID);
+
+		FNetworkGUID InvalidOutputGUID = InvalidVariantGUID.GetVariant().Get<FNetworkGUID>();
+		UE_NET_ASSERT_EQ(InvalidGUID, InvalidOutputGUID);
+
+		// Validity
+		UE_NET_ASSERT_EQ(ValidGUID.IsValid(), ValidVariantGUID.IsValid());
+		UE_NET_ASSERT_EQ(InvalidGUID.IsValid(), InvalidVariantGUID.IsValid());
+
+		// Serialization
+		FNetBitWriter TempWriter(1024);
+		TempWriter << ValidVariantGUID;
+		TempWriter << InvalidVariantGUID;
+
+		FNetBitReader TempReader(nullptr, TempWriter.GetData(), TempWriter.GetNumBits());
+		FNetIDVariant DeserializedValidVariant;
+		TempReader << DeserializedValidVariant;
+
+		FNetIDVariant DeserializedInvalidVariant;
+		TempReader << DeserializedInvalidVariant;
+
+		UE_NET_ASSERT_TRUE(DeserializedValidVariant.GetVariant().IsType<FNetworkGUID>());
+		UE_NET_ASSERT_EQ(DeserializedValidVariant.GetVariant().Get<FNetworkGUID>(), ValidGUID);
+
+		UE_NET_ASSERT_TRUE(DeserializedInvalidVariant.GetVariant().IsType<FNetworkGUID>());
+		UE_NET_ASSERT_EQ(DeserializedInvalidVariant.GetVariant().Get<FNetworkGUID>(), InvalidGUID);
+	}
+
+	UE_NET_TEST(NetIDVariant, NetRefHandle)
+	{
+		// Just need to get a valid FNetRefHandle
+		UE::Net::FReplicationSystemTestNode TestNode(true, TEXT("NetIDVariantTest"));
+		UTestReplicatedIrisObject* TestObject = TestNode.CreateObject(0, 0);
+
+		FNetRefHandle ValidRefHandle = TestNode.GetReplicationBridge()->BeginReplication(TestObject);
+		FNetIDVariant ValidVariantRefHandle(ValidRefHandle);
+
+		FNetRefHandle InvalidRefHandle;
+		FNetIDVariant InvalidVariantRefHandle(InvalidRefHandle);
+
+		// Basic round-trip
+		FNetRefHandle ValidOutputRefHandle = ValidVariantRefHandle.GetVariant().Get<FNetRefHandle>();
+		UE_NET_ASSERT_EQ(ValidRefHandle, ValidOutputRefHandle);
+
+		FNetRefHandle InvalidOutputRefHandle = InvalidVariantRefHandle.GetVariant().Get<FNetRefHandle>();
+		UE_NET_ASSERT_EQ(InvalidRefHandle, InvalidOutputRefHandle);
+
+		// Validity
+		UE_NET_ASSERT_EQ(ValidRefHandle.IsValid(), ValidVariantRefHandle.IsValid());
+		UE_NET_ASSERT_EQ(InvalidRefHandle.IsValid(), InvalidVariantRefHandle.IsValid());
+
+		// Serialization
+		FNetBitWriter TempWriter(1024);
+		TempWriter << ValidVariantRefHandle;
+		TempWriter << InvalidVariantRefHandle;
+
+		FNetBitReader TempReader(nullptr, TempWriter.GetData(), TempWriter.GetNumBits());
+		FNetIDVariant DeserializedValidVariant;
+		TempReader << DeserializedValidVariant;
+
+		FNetIDVariant DeserializedInvalidVariant;
+		TempReader << DeserializedInvalidVariant;
+
+		UE_NET_ASSERT_TRUE(DeserializedValidVariant.GetVariant().IsType<FNetRefHandle>());
+		UE_NET_ASSERT_EQ(DeserializedValidVariant.GetVariant().Get<FNetRefHandle>(), ValidRefHandle);
+
+		UE_NET_ASSERT_TRUE(DeserializedInvalidVariant.GetVariant().IsType<FNetRefHandle>());
+		UE_NET_ASSERT_EQ(DeserializedInvalidVariant.GetVariant().Get<FNetRefHandle>(), InvalidRefHandle);
+	}
+}
diff --git a/Engine/Source/Runtime/Engine/Private/Net/NetIDVariant.cpp b/Engine/Source/Runtime/Engine/Private/Net/NetIDVariant.cpp
new file mode 100644
--- /dev/null
+++ b/Engine/Source/Runtime/Engine/Private/Net/NetIDVariant.cpp
@@ -0,0 +1,120 @@
+// Copyright Epic Games, Inc. All Rights Reserved.
+
+#include "Net/NetIDVariant.h"
+#include "EngineLogs.h"
+
+namespace UE::Net
+{
+
+FNetIDVariant::FNetIDVariant(FNetworkGUID NetGUID)
+	: Variant(TInPlaceType<FNetworkGUID>(), NetGUID)
+{
+}
+
+#if UE_WITH_IRIS
+FNetIDVariant::FNetIDVariant(FNetRefHandle NetRefHandle)
+	: Variant(TInPlaceType<FNetRefHandle>(), NetRefHandle)
+{
+}
+#endif
+
+bool FNetIDVariant::operator==(FNetIDVariant RHS) const
+{
+	if (Variant.GetIndex() != RHS.Variant.GetIndex())
+	{
+		return false;
+	}
+
+	if (Variant.IsType<FNetworkGUID>())
+	{
+		return Variant.Get<FNetworkGUID>() == RHS.Variant.Get<FNetworkGUID>();
+	}
+#if UE_WITH_IRIS
+	else if (Variant.IsType<FNetRefHandle>())
+	{
+		return Variant.Get<FNetRefHandle>() == RHS.Variant.Get<FNetRefHandle>();
+	}
+#endif
+
+	return true;
+}
+
+bool FNetIDVariant::IsValid() const
+{
+	if (Variant.IsType<FNetworkGUID>())
+	{
+		return Variant.Get<FNetworkGUID>().IsValid();
+	}
+	else
+#if UE_WITH_IRIS
+	if (Variant.IsType<FNetRefHandle>())
+	{
+		return Variant.Get<FNetRefHandle>().IsValid();
+	}
+	else
+#endif
+	{
+		return false;
+	}
+}
+
+FArchive& operator<<(FArchive& Ar, FNetIDVariant& NetID)
+{
+#if UE_WITH_IRIS
+	constexpr SIZE_T ExpectedSize = 3;
+#else
+	constexpr SIZE_T ExpectedSize = 2;
+#endif
+	static_assert(TVariantSize_V<FNetIDVariant::FVariantType> == ExpectedSize, "FNetIDVariant variant size changed, potential serialization incompatibility.");
+
+	uint32 TypeIndex = static_cast<uint32>(NetID.Variant.GetIndex());
+
+	Ar.SerializeInt(TypeIndex, TVariantSize_V<FNetIDVariant::FVariantType>);
+
+	if (Ar.IsSaving())
+	{
+		if (NetID.Variant.IsType<FNetworkGUID>())
+		{
+			FNetworkGUID WriteGUID = NetID.Variant.Get<FNetworkGUID>();
+			Ar << WriteGUID;
+		}
+#if UE_WITH_IRIS
+		else if (NetID.Variant.IsType<FNetRefHandle>())
+		{
+			FNetRefHandle WriteHandle = NetID.Variant.Get<FNetRefHandle>();
+			Ar << WriteHandle;
+		}
+#endif
+	}
+	else
+	{
+		if (TypeIndex == FNetIDVariant::FVariantType::IndexOfType<FNetworkGUID>())
+		{
+			FNetworkGUID ReadGUID;
+			Ar << ReadGUID;
+			NetID.Variant.Set<FNetworkGUID>(ReadGUID);
+		}
+#if UE_WITH_IRIS
+		else if (TypeIndex == FNetIDVariant::FVariantType::IndexOfType<FNetRefHandle>())
+		{
+			FNetRefHandle ReadRefHandle;
+			Ar << ReadRefHandle;
+			NetID.Variant.Set<FNetRefHandle>(ReadRefHandle);
+		}
+#endif
+		else if(TypeIndex == FNetIDVariant::FVariantType::IndexOfType<FNetIDVariant::FEmptyID>())
+		{
+			NetID.Variant = FNetIDVariant::FVariantType();
+		}
+		else
+		{
+			// Invalid TypeIndex
+			UE_LOG(LogNet, Warning, TEXT("Unknown TypeIndex %u reading an FNetIDVariant."), TypeIndex);
+			Ar.SetError();
+		}
+	}
+
+	return Ar;
+}
+
+}
\ No newline at end of file
diff --git a/Engine/Source/Runtime/Engine/Public/Net/NetIDVariant.h b/Engine/Source/Runtime/Engine/Public/Net/NetIDVariant.h
new file mode 100644
--- /dev/null
+++ b/Engine/Source/Runtime/Engine/Public/Net/NetIDVariant.h
@@ -0,0 +1,80 @@
+// Copyright Epic Games, Inc. All Rights Reserved.
+
+#pragma once
+
+#if UE_WITH_IRIS
+#include "Iris/ReplicationSystem/NetRefHandle.h"
+#endif
+#include "Misc/NetworkGuid.h"
+#include "Misc/TVariant.h"
+#include "Templates/TypeHash.h"
+
+class FArchive;
+
+namespace UE::Net
+{
+
+/**
+ * Simple variant wrapper to hold a generic or Iris network object ID (FNetworkGUID or FNetRefHandle).
+ * Intended for transitioning systems that were using FNetworkGUID directly to Iris, not intended for general use.
+ */
+class FNetIDVariant
+{
+public:
+	/** Placeholder to represent empty/unknown type. Always considered invalid. */
+	struct FEmptyID {};
+
+#if UE_WITH_IRIS
+	using FVariantType = TVariant<FEmptyID, FNetworkGUID, FNetRefHandle>;
+#else
+	using FVariantType = TVariant<FEmptyID, FNetworkGUID>;
+#endif
+
+	/** Constructs a variant holding an FEmptyID */
+	ENGINE_API FNetIDVariant() = default;
+
+	explicit ENGINE_API FNetIDVariant(FNetworkGUID NetGUID);
+#if UE_WITH_IRIS
+	explicit ENGINE_API FNetIDVariant(FNetRefHandle NetRefHandle);
+#endif
+
+	/** Serializes or deserializes the stored ID or handle, suitable for networking */
+	friend ENGINE_API FArchive& operator<<(FArchive& Ar, FNetIDVariant& NetID);
+
+	/** Returns whether the stored ID or handle is valid */
+	bool ENGINE_API IsValid() const;
+
+	/** Returns the stored TVariant */
+	FVariantType GetVariant() const
+	{
+		return Variant;
+	}
+
+	/** Equality comparison */
+	bool ENGINE_API operator==(FNetIDVariant RHS) const;
+
+private:
+	FVariantType Variant;
+};
+
+inline uint32 GetTypeHash(const FNetIDVariant& NetID)
+{
+	// Assuming the internal index of TVariant is uint8 (which it is) even though the public API uses SIZE_T
+	uint8 TypeIndex = static_cast<uint8>(NetID.GetVariant().GetIndex());
+
+	uint32 IDHash = 0;
+	if (NetID.GetVariant().IsType<FNetworkGUID>())
+	{
+		IDHash = GetTypeHash(NetID.GetVariant().Get<FNetworkGUID>());
+	}
+#if UE_WITH_IRIS
+	else if (NetID.GetVariant().IsType<FNetRefHandle>())
+	{
+		IDHash = GetTypeHash(NetID.GetVariant().Get<FNetRefHandle>());
+	}
+#endif
+
+	return HashCombineFast(TypeIndex, IDHash);
+}
+
+}
\ No newline at end of file
