Iris - Don't clobber not replicated properties when applying instanced structs.

diff --git a/Engine/Plugins/Runtime/ReplicationSystemTestPlugin/Source/Private/Tests/ReplicationState/TestDerivedStruct.cpp b/Engine/Plugins/Runtime/ReplicationSystemTestPlugin/Source/Private/Tests/ReplicationState/TestDerivedStruct.cpp
--- a/Engine/Plugins/Runtime/ReplicationSystemTestPlugin/Source/Private/Tests/ReplicationState/TestDerivedStruct.cpp
+++ b/Engine/Plugins/Runtime/ReplicationSystemTestPlugin/Source/Private/Tests/ReplicationState/TestDerivedStruct.cpp
@@ -128,7 +128,7 @@ UE_NET_TEST_FIXTURE(FTestDerivedStructFixture, DerivedStructWithoutNetSerializer
 	// The derived struct should have two members, the first using the NetSerializer of the parent struct.
 	UE_NET_ASSERT_EQ(StructDescriptor->MemberCount, uint16(2));
 	UE_NET_ASSERT_EQ(StructDescriptor->MemberSerializerDescriptors[0].Serializer, &UE_NET_GET_SERIALIZER(FTestDerivedStruct_Inherited_WithNetSerializer_NetSerializer));
-	UE_NET_ASSERT_TRUE(EnumHasAnyFlags(StructDescriptor->Traits, EReplicationStateTraits::IsDerivedStruct));
+	UE_NET_ASSERT_TRUE(EnumHasAnyFlags(StructDescriptor->Traits, EReplicationStateTraits::IsDerivedFromStructWithCustomSerializer));
 }
 
 UE_NET_TEST_FIXTURE(FTestDerivedStructFixture, DerivedStructHasValidDebugDescriptor)
@@ -149,7 +149,7 @@ UE_NET_TEST_FIXTURE(FTestDerivedStructFixture, DeepInheritanceOfStructWithNetSer
 {
 	TRefCountPtr<const FReplicationStateDescriptor> StructDescriptor = FReplicationStateDescriptorBuilder::CreateDescriptorForStruct(FTestDerivedStruct_DeepInheritanceOfStructWithNetSerializer::StaticStruct());
 
-	UE_NET_ASSERT_TRUE(EnumHasAnyFlags(StructDescriptor->Traits, EReplicationStateTraits::IsDerivedStruct));
+	UE_NET_ASSERT_TRUE(EnumHasAnyFlags(StructDescriptor->Traits, EReplicationStateTraits::IsDerivedFromStructWithCustomSerializer));
 
 	// The derived struct should have three members, the first using the NetSerializer of the parent struct.
 	UE_NET_ASSERT_EQ(StructDescriptor->MemberCount, uint16(3));
@@ -164,7 +164,7 @@ UE_NET_TEST_FIXTURE(FTestDerivedStructFixture, DeepInheritanceAndStructWithNetSe
 {
 	TRefCountPtr<const FReplicationStateDescriptor> StructDescriptor = FReplicationStateDescriptorBuilder::CreateDescriptorForStruct(FTestDerivedStruct_DeepInheritanceOfStructWithNetSerializer::StaticStruct());
 
-	UE_NET_ASSERT_TRUE(EnumHasAnyFlags(StructDescriptor->Traits, EReplicationStateTraits::IsDerivedStruct));
+	UE_NET_ASSERT_TRUE(EnumHasAnyFlags(StructDescriptor->Traits, EReplicationStateTraits::IsDerivedFromStructWithCustomSerializer));
 
 	// The derived struct should have three members, the first using the NetSerializer of the parent struct.
 	UE_NET_ASSERT_EQ(StructDescriptor->MemberCount, uint16(3));
@@ -179,7 +179,7 @@ UE_NET_TEST_FIXTURE(FTestDerivedStructFixture, StructUsesClosestNetSerializerInH
 {
 	TRefCountPtr<const FReplicationStateDescriptor> StructDescriptor = FReplicationStateDescriptorBuilder::CreateDescriptorForStruct(FTestDerivedStruct_Inherited_Inherited_WithNetSerializer_DeepInheritanceOfStructWithNetSerializer::StaticStruct());
 
-	UE_NET_ASSERT_TRUE(EnumHasAnyFlags(StructDescriptor->Traits, EReplicationStateTraits::IsDerivedStruct));
+	UE_NET_ASSERT_TRUE(EnumHasAnyFlags(StructDescriptor->Traits, EReplicationStateTraits::IsDerivedFromStructWithCustomSerializer));
 
 	// The derived struct should have two members, the first using the NetSerializer of the parent struct.
 	UE_NET_ASSERT_EQ(StructDescriptor->MemberCount, uint16(2));
diff --git a/Engine/Plugins/Runtime/ReplicationSystemTestPlugin/Source/Private/Tests/Serialization/TestInstancedStructNetSerializer.cpp b/Engine/Plugins/Runtime/ReplicationSystemTestPlugin/Source/Private/Tests/Serialization/TestInstancedStructNetSerializer.cpp
--- a/Engine/Plugins/Runtime/ReplicationSystemTestPlugin/Source/Private/Tests/Serialization/TestInstancedStructNetSerializer.cpp
+++ b/Engine/Plugins/Runtime/ReplicationSystemTestPlugin/Source/Private/Tests/Serialization/TestInstancedStructNetSerializer.cpp
@@ -5,6 +5,7 @@
 #include "Tests/ReplicationSystem/ReplicationSystemServerClientTestFixture.h"
 #include "Iris/ReplicationSystem/ReplicationFragmentUtil.h"
 #include "Iris/ReplicationSystem/ReplicationSystemInternal.h"
+#include "Iris/ReplicationState/PropertyNetSerializerInfoRegistry.h"
 #include "Iris/ReplicationState/ReplicationStateDescriptorBuilder.h"
 #include "Iris/Serialization/InternalNetSerializationContext.h"
 #include "Iris/Serialization/NetReferenceCollector.h"
@@ -21,6 +22,84 @@ FTestMessage& operator<<(FTestMessage& Message, const FInstancedStruct& Instance
 	return Message << String;
 }
 
+struct FStructForInstancedStructTestWithCustomApplyNetSerializer
+{
+	// Version
+	static const uint32 Version = 0;
+
+	// Types
+	typedef FStructForInstancedStructTestWithCustomApply SourceType;
+	typedef int QuantizedType;
+	typedef FNetSerializerConfig ConfigType;
+
+	inline static const ConfigType DefaultConfig;
+
+	// 
+	static void Serialize(FNetSerializationContext&, const FNetSerializeArgs&);
+	static void Deserialize(FNetSerializationContext&, const FNetDeserializeArgs&);
+
+	static void Quantize(FNetSerializationContext&, const FNetQuantizeArgs&);
+	static void Dequantize(FNetSerializationContext&, const FNetDequantizeArgs&);
+
+	static bool IsEqual(FNetSerializationContext&, const FNetIsEqualArgs&);
+
+	static void Apply(FNetSerializationContext&, const FNetApplyArgs&);
+
+private:
+	class FNetSerializerRegistryDelegates final : private UE::Net::FNetSerializerRegistryDelegates
+	{
+	public:
+		virtual ~FNetSerializerRegistryDelegates();
+
+	private:
+		virtual void OnPreFreezeNetSerializerRegistry() override;
+
+		inline static const FName StructForInstancedStructTestWithCustomApplyNetSerializer = FName("StructForInstancedStructTestWithCustomApply");
+		UE_NET_IMPLEMENT_NAMED_STRUCT_NETSERIALIZER_INFO(StructForInstancedStructTestWithCustomApplyNetSerializer, FStructForInstancedStructTestWithCustomApplyNetSerializer);
+	};
+
+	inline static FStructForInstancedStructTestWithCustomApplyNetSerializer::FNetSerializerRegistryDelegates NetSerializerRegistryDelegates;
+};
+UE_NET_IMPLEMENT_SERIALIZER(FStructForInstancedStructTestWithCustomApplyNetSerializer);
+
+struct FStructForInstancedStructTestWithCustomSerializerNetSerializer
+{
+	// Version
+	static const uint32 Version = 0;
+
+	// Types
+	typedef FStructForInstancedStructTestWithCustomSerializer SourceType;
+	typedef int QuantizedType;
+	typedef FNetSerializerConfig ConfigType;
+
+	inline static const ConfigType DefaultConfig;
+
+	// 
+	static void Serialize(FNetSerializationContext&, const FNetSerializeArgs&);
+	static void Deserialize(FNetSerializationContext&, const FNetDeserializeArgs&);
+
+	static void Quantize(FNetSerializationContext&, const FNetQuantizeArgs&);
+	static void Dequantize(FNetSerializationContext&, const FNetDequantizeArgs&);
+
+	static bool IsEqual(FNetSerializationContext&, const FNetIsEqualArgs&);
+
+private:
+	class FNetSerializerRegistryDelegates final : private UE::Net::FNetSerializerRegistryDelegates
+	{
+	public:
+		virtual ~FNetSerializerRegistryDelegates();
+
+	private:
+		virtual void OnPreFreezeNetSerializerRegistry() override;
+
+		inline static const FName StructForInstancedStructTestWithCustomSerializerNetSerializer = FName("StructForInstancedStructTestWithCustomSerializer");
+		UE_NET_IMPLEMENT_NAMED_STRUCT_NETSERIALIZER_INFO(StructForInstancedStructTestWithCustomSerializerNetSerializer, FStructForInstancedStructTestWithCustomSerializerNetSerializer);
+	};
+
+	inline static FStructForInstancedStructTestWithCustomSerializerNetSerializer::FNetSerializerRegistryDelegates NetSerializerRegistryDelegates;
+};
+UE_NET_IMPLEMENT_SERIALIZER(FStructForInstancedStructTestWithCustomSerializerNetSerializer);
+
 }
 
 namespace UE::Net::Private
@@ -413,6 +492,118 @@ UE_NET_TEST_FIXTURE(FTestInstancedStructNetSerializerFixture, ModifyArray)
 	UE_NET_ASSERT_EQ(ClientObject->InstancedStructArray[2].Get<FStructForInstancedStructTestA>().SomeUint16, ServerObject->InstancedStructArray[2].Get<FStructForInstancedStructTestA>().SomeUint16);
 }
 
+UE_NET_TEST_FIXTURE(FTestInstancedStructNetSerializerFixture, NonReplicatedPropertiesAreNotClobbered)
+{
+	// Add a client
+	FReplicationSystemTestClient* Client = CreateClient();
+
+	// Spawn objects on server
+	UInstancedStructNetSerializerTestObject* ServerObject = Server->CreateObject<UInstancedStructNetSerializerTestObject>();
+
+	// Initialize struct with some non-replicated properties.
+	ServerObject->InstancedStruct.InitializeAs<FStructForInstancedStructTestWithNonReplicatedProperties>();
+	ServerObject->InstancedStruct.GetMutable<FStructForInstancedStructTestWithNonReplicatedProperties>().IntA = 1333;
+	
+	// Replicate
+	Server->UpdateAndSend({ Client });
+	
+	UInstancedStructNetSerializerTestObject* ClientObject = Client->GetObjectAs<UInstancedStructNetSerializerTestObject>(ServerObject->NetRefHandle);
+	UE_NET_ASSERT_NE(ClientObject, nullptr);
+	CA_ASSUME(ClientObject != nullptr);
+	UE_NET_ASSERT_EQ(ClientObject->InstancedStruct.Get<FStructForInstancedStructTestWithNonReplicatedProperties>().IntA, ServerObject->InstancedStruct.Get<FStructForInstancedStructTestWithNonReplicatedProperties>().IntA);
+
+	// Modify value and verify that it's replicated as expected
+	ServerObject->InstancedStruct.GetMutable<FStructForInstancedStructTestWithNonReplicatedProperties>().IntA = 1337;
+
+	// Modify non-replicated property on the client
+	ClientObject->InstancedStruct.GetMutable<FStructForInstancedStructTestWithNonReplicatedProperties>().IntB = 1331;
+
+	// Replicate
+	Server->UpdateAndSend({ Client });
+
+	// Verified that the client got the modified value
+	UE_NET_ASSERT_EQ(ClientObject->InstancedStruct.Get<FStructForInstancedStructTestWithNonReplicatedProperties>().IntA, ServerObject->InstancedStruct.Get<FStructForInstancedStructTestWithNonReplicatedProperties>().IntA);
+
+	// Verified that the non-replicated property remained unchanged on the client
+	UE_NET_ASSERT_EQ(ClientObject->InstancedStruct.Get<FStructForInstancedStructTestWithNonReplicatedProperties>().IntB, 1331);
+}
+
+UE_NET_TEST_FIXTURE(FTestInstancedStructNetSerializerFixture, NetSerializerApplyIsCalled)
+{
+	// Add a client
+	FReplicationSystemTestClient* Client = CreateClient();
+
+	// Spawn objects on server
+	UInstancedStructNetSerializerTestObject* ServerObject = Server->CreateObject<UInstancedStructNetSerializerTestObject>();
+
+	// Initialize struct with custom apply.
+	ServerObject->InstancedStruct.InitializeAs<FStructForInstancedStructTestWithCustomApply>();
+	ServerObject->InstancedStruct.GetMutable<FStructForInstancedStructTestWithCustomApply>().Uint = 1333;
+	
+	// Replicate
+	Server->UpdateAndSend({ Client });
+	
+	UInstancedStructNetSerializerTestObject* ClientObject = Client->GetObjectAs<UInstancedStructNetSerializerTestObject>(ServerObject->NetRefHandle);
+	UE_NET_ASSERT_NE(ClientObject, nullptr);
+	CA_ASSUME(ClientObject != nullptr);
+	UE_NET_ASSERT_EQ(ClientObject->InstancedStruct.Get<FStructForInstancedStructTestWithCustomApply>().Uint, ServerObject->InstancedStruct.Get<FStructForInstancedStructTestWithCustomApply>().Uint);
+
+	UE_NET_ASSERT_EQ(ClientObject->InstancedStruct.Get<FStructForInstancedStructTestWithCustomApply>().ApplyCallCount, 1U);
+
+	// Modify value and verify that it's replicated as expected
+	ServerObject->InstancedStruct.GetMutable<FStructForInstancedStructTestWithCustomApply>().Uint = 1337;
+
+	// Modify Apply handled property and verify that it doesn't overwrite the client value
+	ServerObject->InstancedStruct.GetMutable<FStructForInstancedStructTestWithCustomApply>().ApplyCallCount = 1337;
+
+	// Replicate
+	Server->UpdateAndSend({ Client });
+
+	// Verified that the client got the modified value
+	UE_NET_ASSERT_EQ(ClientObject->InstancedStruct.Get<FStructForInstancedStructTestWithCustomApply>().Uint, ServerObject->InstancedStruct.Get<FStructForInstancedStructTestWithCustomApply>().Uint);
+
+	// Verified that the custom handled property in Apply is as expected
+	UE_NET_ASSERT_EQ(ClientObject->InstancedStruct.Get<FStructForInstancedStructTestWithCustomApply>().ApplyCallCount, 2U);
+}
+
+UE_NET_TEST_FIXTURE(FTestInstancedStructNetSerializerFixture, CustomSerializedStructIsOverwritten)
+{
+	// Add a client
+	FReplicationSystemTestClient* Client = CreateClient();
+
+	// Spawn objects on server
+	UInstancedStructNetSerializerTestObject* ServerObject = Server->CreateObject<UInstancedStructNetSerializerTestObject>();
+
+	// Initialize struct with custom serializer.
+	ServerObject->InstancedStruct.InitializeAs<FStructForInstancedStructTestWithCustomSerializer>();
+	ServerObject->InstancedStruct.GetMutable<FStructForInstancedStructTestWithCustomSerializer>().Uint = 1333;
+	
+	// Replicate
+	Server->UpdateAndSend({ Client });
+	
+	UInstancedStructNetSerializerTestObject* ClientObject = Client->GetObjectAs<UInstancedStructNetSerializerTestObject>(ServerObject->NetRefHandle);
+	UE_NET_ASSERT_NE(ClientObject, nullptr);
+	CA_ASSUME(ClientObject != nullptr);
+	UE_NET_ASSERT_EQ(ClientObject->InstancedStruct.Get<FStructForInstancedStructTestWithCustomSerializer>().Uint, ServerObject->InstancedStruct.Get<FStructForInstancedStructTestWithCustomSerializer>().Uint);
+
+	// Modify value and verify that it's replicated as expected
+	ServerObject->InstancedStruct.GetMutable<FStructForInstancedStructTestWithCustomSerializer>().Uint = 1337;
+
+	// Modify not replicated client values and verify they're overwritten since it's a custom serializer without apply
+	ClientObject->InstancedStruct.GetMutable<FStructForInstancedStructTestWithCustomSerializer>().NotReplicatedUint = 1338;
+	ClientObject->InstancedStruct.GetMutable<FStructForInstancedStructTestWithCustomSerializer>().NotPropertyUint = 1339;
+
+	// Replicate
+	Server->UpdateAndSend({ Client });
+
+	// Verified that the client got the modified value
+	UE_NET_ASSERT_EQ(ClientObject->InstancedStruct.Get<FStructForInstancedStructTestWithCustomSerializer>().Uint, ServerObject->InstancedStruct.Get<FStructForInstancedStructTestWithCustomSerializer>().Uint);
+
+	// Verified that the not replicated properties did get overwritten
+	UE_NET_ASSERT_EQ(ClientObject->InstancedStruct.Get<FStructForInstancedStructTestWithCustomSerializer>().NotReplicatedUint, ServerObject->InstancedStruct.Get<FStructForInstancedStructTestWithCustomSerializer>().NotReplicatedUint);
+	UE_NET_ASSERT_EQ(ClientObject->InstancedStruct.Get<FStructForInstancedStructTestWithCustomSerializer>().NotPropertyUint, ServerObject->InstancedStruct.Get<FStructForInstancedStructTestWithCustomSerializer>().NotPropertyUint);
+}
+
 // Fixture implementation
 void FTestInstancedStructNetSerializerFixture::SetUp()
 {
@@ -749,6 +940,128 @@ const FNetSerializerConfig* FTestInstancedStructNetSerializerFixture::GetInstanc
 
 }
 
+namespace UE::Net
+{
+
+void FStructForInstancedStructTestWithCustomApplyNetSerializer::Serialize(FNetSerializationContext& Context, const FNetSerializeArgs& Args)
+{
+	const QuantizedType Value = *reinterpret_cast<const QuantizedType*>(Args.Source);
+	Context.GetBitStreamWriter()->WriteBits(Value, 32);
+}
+
+void FStructForInstancedStructTestWithCustomApplyNetSerializer::Deserialize(FNetSerializationContext& Context, const FNetDeserializeArgs& Args)
+{
+	QuantizedType& Target = *reinterpret_cast<QuantizedType*>(Args.Target);
+	Target = Context.GetBitStreamReader()->ReadBits(32);
+}
+
+void FStructForInstancedStructTestWithCustomApplyNetSerializer::Quantize(FNetSerializationContext&, const FNetQuantizeArgs& Args)
+{
+	const SourceType& Source = *reinterpret_cast<const SourceType*>(Args.Source);
+	QuantizedType& Target = *reinterpret_cast<QuantizedType*>(Args.Target);
+
+	Target = Source.Uint;
+}
+
+void FStructForInstancedStructTestWithCustomApplyNetSerializer::Dequantize(FNetSerializationContext&, const FNetDequantizeArgs& Args)
+{
+	const QuantizedType& Source = *reinterpret_cast<const QuantizedType*>(Args.Source);
+	SourceType& Target = *reinterpret_cast<SourceType*>(Args.Target);
+
+	Target.Uint = Source;
+}
+
+bool FStructForInstancedStructTestWithCustomApplyNetSerializer::IsEqual(FNetSerializationContext&, const FNetIsEqualArgs& Args)
+{
+	if (Args.bStateIsQuantized)
+	{
+		const QuantizedType& QuantizedValue0 = *reinterpret_cast<const QuantizedType*>(Args.Source0);
+		const QuantizedType& QuantizedValue1 = *reinterpret_cast<const QuantizedType*>(Args.Source1);
+		return QuantizedValue0 == QuantizedValue1;
+	}
+	else
+	{
+		const SourceType& SourceValue0 = *reinterpret_cast<const SourceType*>(Args.Source0);
+		const SourceType& SourceValue1 = *reinterpret_cast<const SourceType*>(Args.Source1);
+		return SourceValue0.Uint == SourceValue1.Uint;
+	}
+}
+
+void FStructForInstancedStructTestWithCustomApplyNetSerializer::Apply(FNetSerializationContext& Context, const FNetApplyArgs& Args)
+{
+	const SourceType& Source = *reinterpret_cast<const SourceType*>(Args.Source);
+	SourceType& Target = *reinterpret_cast<SourceType*>(Args.Target);
+
+	Target.Uint = Source.Uint;
+	Target.ApplyCallCount++;
+}
+
+FStructForInstancedStructTestWithCustomApplyNetSerializer::FNetSerializerRegistryDelegates::~FNetSerializerRegistryDelegates()
+{
+	UE_NET_UNREGISTER_NETSERIALIZER_INFO(StructForInstancedStructTestWithCustomApplyNetSerializer);
+}
+
+void FStructForInstancedStructTestWithCustomApplyNetSerializer::FNetSerializerRegistryDelegates::OnPreFreezeNetSerializerRegistry()
+{
+	UE_NET_REGISTER_NETSERIALIZER_INFO(StructForInstancedStructTestWithCustomApplyNetSerializer);
+}
+
+void FStructForInstancedStructTestWithCustomSerializerNetSerializer::Serialize(FNetSerializationContext& Context, const FNetSerializeArgs& Args)
+{
+	const QuantizedType Value = *reinterpret_cast<const QuantizedType*>(Args.Source);
+	Context.GetBitStreamWriter()->WriteBits(Value, 32);
+}
+
+void FStructForInstancedStructTestWithCustomSerializerNetSerializer::Deserialize(FNetSerializationContext& Context, const FNetDeserializeArgs& Args)
+{
+	QuantizedType& Target = *reinterpret_cast<QuantizedType*>(Args.Target);
+	Target = Context.GetBitStreamReader()->ReadBits(32);
+}
+
+void FStructForInstancedStructTestWithCustomSerializerNetSerializer::Quantize(FNetSerializationContext&, const FNetQuantizeArgs& Args)
+{
+	const SourceType& Source = *reinterpret_cast<const SourceType*>(Args.Source);
+	QuantizedType& Target = *reinterpret_cast<QuantizedType*>(Args.Target);
+
+	Target = Source.Uint;
+}
+
+void FStructForInstancedStructTestWithCustomSerializerNetSerializer::Dequantize(FNetSerializationContext&, const FNetDequantizeArgs& Args)
+{
+	const QuantizedType& Source = *reinterpret_cast<const QuantizedType*>(Args.Source);
+	SourceType& Target = *reinterpret_cast<SourceType*>(Args.Target);
+
+	Target.Uint = Source;
+}
+
+bool FStructForInstancedStructTestWithCustomSerializerNetSerializer::IsEqual(FNetSerializationContext&, const FNetIsEqualArgs& Args)
+{
+	if (Args.bStateIsQuantized)
+	{
+		const QuantizedType& QuantizedValue0 = *reinterpret_cast<const QuantizedType*>(Args.Source0);
+		const QuantizedType& QuantizedValue1 = *reinterpret_cast<const QuantizedType*>(Args.Source1);
+		return QuantizedValue0 == QuantizedValue1;
+	}
+	else
+	{
+		const SourceType& SourceValue0 = *reinterpret_cast<const SourceType*>(Args.Source0);
+		const SourceType& SourceValue1 = *reinterpret_cast<const SourceType*>(Args.Source1);
+		return SourceValue0.Uint == SourceValue1.Uint;
+	}
+}
+
+FStructForInstancedStructTestWithCustomSerializerNetSerializer::FNetSerializerRegistryDelegates::~FNetSerializerRegistryDelegates()
+{
+	UE_NET_UNREGISTER_NETSERIALIZER_INFO(StructForInstancedStructTestWithCustomSerializerNetSerializer);
+}
+
+void FStructForInstancedStructTestWithCustomSerializerNetSerializer::FNetSerializerRegistryDelegates::OnPreFreezeNetSerializerRegistry()
+{
+	UE_NET_REGISTER_NETSERIALIZER_INFO(StructForInstancedStructTestWithCustomSerializerNetSerializer);
+}
+
+}
+
 //
 void UInstancedStructNetSerializerTestObject::GetLifetimeReplicatedProps(TArray<class FLifetimeProperty>& OutLifetimeProps) const
 {
diff --git a/Engine/Plugins/Runtime/ReplicationSystemTestPlugin/Source/Private/Tests/Serialization/TestInstancedStructNetSerializer.h b/Engine/Plugins/Runtime/ReplicationSystemTestPlugin/Source/Private/Tests/Serialization/TestInstancedStructNetSerializer.h
--- a/Engine/Plugins/Runtime/ReplicationSystemTestPlugin/Source/Private/Tests/Serialization/TestInstancedStructNetSerializer.h
+++ b/Engine/Plugins/Runtime/ReplicationSystemTestPlugin/Source/Private/Tests/Serialization/TestInstancedStructNetSerializer.h
@@ -3,10 +3,19 @@
 #pragma once
 
 #include "UObject/ObjectMacros.h"
+#include "Iris/Serialization/NetSerializer.h"
 #include "StructUtils/InstancedStruct.h"
 #include "Tests/ReplicationSystem/ReplicatedTestObject.h"
 #include "TestInstancedStructNetSerializer.generated.h"
 
+namespace UE::Net
+{
+
+UE_NET_DECLARE_SERIALIZER(FStructForInstancedStructTestWithCustomApplyNetSerializer, REPLICATIONSYSTEMTESTPLUGIN_API);
+UE_NET_DECLARE_SERIALIZER(FStructForInstancedStructTestWithCustomSerializerNetSerializer, REPLICATIONSYSTEMTESTPLUGIN_API);
+
+}
+
 USTRUCT()
 struct FTestInstancedStruct
 {
@@ -78,6 +87,45 @@ struct FStructForInstancedStructTestWithObjectReference
 	TObjectPtr<UObject> SomeObject = nullptr;
 };
 
+USTRUCT()
+struct FStructForInstancedStructTestWithCustomApply
+{
+	GENERATED_BODY()
+
+	UPROPERTY();
+	uint32 Uint = 0;
+
+	// Incremented by 1 in this struct's NetSerializer Apply method. It's set to replicate but Apply trumps that.
+	UPROPERTY()
+	uint32 ApplyCallCount = 0;
+};
+
+USTRUCT()
+struct FStructForInstancedStructTestWithCustomSerializer
+{
+	GENERATED_BODY()
+
+	UPROPERTY();
+	uint32 Uint = 0;
+
+	UPROPERTY(NotReplicated)
+	uint32 NotReplicatedUint = 0;
+
+	uint32 NotPropertyUint = 0;
+};
+
+USTRUCT()
+struct FStructForInstancedStructTestWithNonReplicatedProperties
+{
+	GENERATED_BODY()
+
+	UPROPERTY();
+	int IntA = 0;
+
+	UPROPERTY(NotReplicated);
+	int IntB = 0;
+};
+
 UCLASS()
 class UInstancedStructNetSerializerTestObject : public UReplicatedTestObject
 {
diff --git a/Engine/Source/Runtime/Experimental/Iris/Core/Private/Iris/ReplicationState/InternalPropertyReplicationState.cpp b/Engine/Source/Runtime/Experimental/Iris/Core/Private/Iris/ReplicationState/InternalPropertyReplicationState.cpp
--- a/Engine/Source/Runtime/Experimental/Iris/Core/Private/Iris/ReplicationState/InternalPropertyReplicationState.cpp
+++ b/Engine/Source/Runtime/Experimental/Iris/Core/Private/Iris/ReplicationState/InternalPropertyReplicationState.cpp
@@ -233,7 +233,7 @@ bool InternalCompareStructProperty(const FReplicationStateDescriptor* StructDesc
 	const FReplicationStateMemberPropertyDescriptor* MemberPropertyDescriptors = StructDescriptor->MemberPropertyDescriptors;
 
 	// For derived structs the first property is a NetSerializer for some super struct. Its property will be null so we can't use standard iteration over properties. Let's use serializer equal for it. 
-	if (EnumHasAnyFlags(StructDescriptor->Traits, EReplicationStateTraits::IsDerivedStruct))
+	if (EnumHasAnyFlags(StructDescriptor->Traits, EReplicationStateTraits::IsDerivedFromStructWithCustomSerializer))
 	{
 		FirstStructMemberForCompare = 1U;
 
@@ -270,12 +270,12 @@ bool InternalCompareStructProperty(const FReplicationStateDescriptor* StructDesc
 void InternalApplyStructProperty(const FReplicationStateDescriptor* StructDescriptor, void* RESTRICT Dst, const void* RESTRICT Src)
 {
 	uint32 FirstStructMemberForApply = 0U;
-	if (EnumHasAnyFlags(StructDescriptor->Traits, EReplicationStateTraits::IsDerivedStruct))
+	if (EnumHasAnyFlags(StructDescriptor->Traits, EReplicationStateTraits::IsStructWithCustomSerializer | EReplicationStateTraits::IsDerivedFromStructWithCustomSerializer))
 	{
 		// The first member in a derived struct descriptor is the base struct. We can skip after we've applied the value.
 		FirstStructMemberForApply = 1U;
 
-		// If derived struct has custom Apply we need to use it.
+		// If the custom serializer has an Apply method we need to use it.
 		const FReplicationStateMemberSerializerDescriptor& MemberSerializerDescriptor = StructDescriptor->MemberSerializerDescriptors[0];
 		if (EnumHasAnyFlags(MemberSerializerDescriptor.Serializer->Traits, ENetSerializerTraits::HasApply))
 		{
@@ -288,7 +288,7 @@ void InternalApplyStructProperty(const FReplicationStateDescriptor* StructDescri
 		}
 		else
 		{
-		// The base struct is the closest parent with a custom serializer. It was determined above that it did not have a custom apply so we should copy it in its entirety.
+			// For derived structs the base struct is the closest parent with a custom serializer. For root structs it's the actual struct.
 			const UScriptStruct* BaseStruct = StructDescriptor->BaseStruct;
 			BaseStruct->CopyScriptStruct(Dst, Src, 1);
 		}
@@ -392,7 +392,7 @@ void InternalApplyPropertyValue(const FReplicationStateDescriptor* Descriptor, u
 void InternalCopyStructProperty(const FReplicationStateDescriptor* StructDescriptor, void* RESTRICT Dst, const void* RESTRICT Src)
 {
 	uint32 FirstStructMemberForCopy = 0U;
-	if (EnumHasAnyFlags(StructDescriptor->Traits, EReplicationStateTraits::IsDerivedStruct))
+	if (EnumHasAnyFlags(StructDescriptor->Traits, EReplicationStateTraits::IsDerivedFromStructWithCustomSerializer))
 	{
 		// The base struct is the closest parent with a custom serializer, for which we always copy the entire state.
 		const UScriptStruct* BaseStruct = StructDescriptor->BaseStruct;
diff --git a/Engine/Source/Runtime/Experimental/Iris/Core/Private/Iris/ReplicationState/ReplicationStateDescriptorBuilder.cpp b/Engine/Source/Runtime/Experimental/Iris/Core/Private/Iris/ReplicationState/ReplicationStateDescriptorBuilder.cpp
--- a/Engine/Source/Runtime/Experimental/Iris/Core/Private/Iris/ReplicationState/ReplicationStateDescriptorBuilder.cpp
+++ b/Engine/Source/Runtime/Experimental/Iris/Core/Private/Iris/ReplicationState/ReplicationStateDescriptorBuilder.cpp
@@ -106,7 +106,6 @@ enum class EMemberPropertyTraits : uint32
 	HasConnectionSpecificSerialization	= IsInvalidFastArray << 1U,
 	HasPushBasedDirtiness				= HasConnectionSpecificSerialization << 1U,
 	UseSerializerIsEqual				= HasPushBasedDirtiness << 1U,
-	IsBaseStruct						= UseSerializerIsEqual << 1U,
 };
 ENUM_CLASS_FLAGS(EMemberPropertyTraits);
 
@@ -224,8 +223,21 @@ public:
 	void AddMemberProperty(const FMemberProperty& Info) { Members.Add(Info); }
 	void AddMemberFunction(const FMemberFunction& Info) { Functions.Add(Info); }
 
-	void SetStructSizeAndAlignment(SIZE_T Size, SIZE_T Alignment) { StructInfo.SizeAndAlignment.Size = Size; StructInfo.SizeAndAlignment.Alignment = Alignment; }
-	void SetStruct(const UStruct* Struct) { StructInfo.Struct = Struct; }
+	void SetStructSizeAndAlignment(SIZE_T Size, SIZE_T Alignment)
+	{
+		StructInfo.SizeAndAlignment.Size = Size;
+		StructInfo.SizeAndAlignment.Alignment = Alignment;
+	}
+
+	void SetStruct(const UStruct* Struct)
+	{
+		StructInfo.Struct = Struct;
+	}
+
+	void SetStructSerializerType(EStructNetSerializerType SerializerType)
+	{
+		StructInfo.SerializerType = SerializerType;
+	}
 
 	bool HasDataToBuild() const { return (Members.Num() > 0) || (Functions.Num() > 0U); }
 	TRefCountPtr<const FReplicationStateDescriptor> Build(const FString& StateName, FReplicationStateDescriptorRegistry* DescriptorRegistry, const FBuildParameters& BuildParams);
@@ -329,6 +341,7 @@ private:
 	{
 		FSizeAndAlignment SizeAndAlignment = {};
 		const UStruct* Struct = nullptr;
+		EStructNetSerializerType SerializerType = EStructNetSerializerType::Struct;
 	};
 
 	void AddSerializerConfigMemoryBlockToLayout(FMemoryLayoutUtil::FLayout& Layout, FOffsetAndSize& OutOffsetAndSize, const FMemberCache& MemberCache) const;
@@ -1562,18 +1575,25 @@ EReplicationStateTraits FPropertyReplicationStateDescriptorBuilder::BuildReplica
 		Traits |= EReplicationStateTraits::UseSerializerIsEqual;
 	}
 
-	// If one of the members have the IsBaseStruct trait we're building a descriptor for a derived struct.
-	if (EnumHasAnyFlags(PropertyTraits, EMemberPropertyTraits::IsBaseStruct))
-	{
-		Traits |= EReplicationStateTraits::IsDerivedStruct;
-	}
-
 	// Only set the PushBasedDirtiness trait if all properties are push based. An alternative could be to split into multiple states.
 	if (EnumHasAnyFlags(SharedPropertyTraits, EMemberPropertyTraits::HasPushBasedDirtiness))
 	{
 		Traits |= EReplicationStateTraits::HasPushBasedDirtiness;
 	}
 
+	// Struct traits
+	if (Context.BuildParams.DescriptorType == EDescriptorType::Struct)
+	{
+		if (StructInfo.SerializerType == EStructNetSerializerType::DerivedFromCustom)
+		{
+			Traits |= EReplicationStateTraits::IsDerivedFromStructWithCustomSerializer;
+		}
+		else if (StructInfo.SerializerType == EStructNetSerializerType::Custom)
+		{
+			Traits |= EReplicationStateTraits::IsStructWithCustomSerializer;
+		}
+	}
+
 	// Special traits when all members are replicated
 	if (Context.BuildParams.bAllMembersAreReplicated)
 	{
@@ -1847,11 +1867,19 @@ FPropertyReplicationStateDescriptorBuilder::Build(const FString& StateName, FRep
 		}
 
 		// BaseStruct
-		if (EnumHasAnyFlags(Traits, EReplicationStateTraits::IsDerivedStruct))
+		if (BuildParams.DescriptorType == EDescriptorType::Struct)
 		{
-			const UScriptStruct* BaseStruct = Cast<const UScriptStruct>(FPropertyReplicationStateDescriptorBuilder::FindSuperStructWithCustomSerializer(StructInfo.Struct));
-			ensure(BaseStruct != nullptr);
-			Descriptor->BaseStruct = BaseStruct;
+			if (StructInfo.SerializerType == EStructNetSerializerType::DerivedFromCustom)
+			{
+				const UScriptStruct* BaseStruct = Cast<const UScriptStruct>(FPropertyReplicationStateDescriptorBuilder::FindSuperStructWithCustomSerializer(StructInfo.Struct));
+				ensure(BaseStruct != nullptr);
+				Descriptor->BaseStruct = BaseStruct;
+			}
+			else
+			{
+				Descriptor->BaseStruct = Cast<const UScriptStruct>(StructInfo.Struct);
+				ensure(Descriptor->BaseStruct);
+			}
 		}
 
 		if (LayoutData.MemberDebugDescriptorsSizeAndOffset.Size > 0)
@@ -2158,10 +2186,6 @@ EStructNetSerializerType FPropertyReplicationStateDescriptorBuilder::IsSupported
 	
 	// Init traits
 	OutMemberProperty.Traits = EMemberPropertyTraits::None;
-	if (SerializerType == EStructNetSerializerType::DerivedFromCustom)
-	{
-		OutMemberProperty.Traits |= EMemberPropertyTraits::IsBaseStruct;
-	}
 	GetSerializerTraits(OutMemberProperty, nullptr, NetSerializerInfo);
 
 	return SerializerType;
@@ -2600,6 +2624,7 @@ TRefCountPtr<const FReplicationStateDescriptor> FReplicationStateDescriptorBuild
 
 	Builder.SetStructSizeAndAlignment(InStruct->GetStructureSize(), InStruct->GetMinAlignment());
 	Builder.SetStruct(InStruct);
+	Builder.SetStructSerializerType(SerializerType);
 
 	FLazyGetPathNameHelper LazyPathHelper;
 	FPropertyReplicationStateDescriptorBuilder::FBuildParameters BuildParameters = {};
diff --git a/Engine/Source/Runtime/Experimental/Iris/Core/Private/Iris/ReplicationSystem/ReplicationOperations.cpp b/Engine/Source/Runtime/Experimental/Iris/Core/Private/Iris/ReplicationSystem/ReplicationOperations.cpp
--- a/Engine/Source/Runtime/Experimental/Iris/Core/Private/Iris/ReplicationSystem/ReplicationOperations.cpp
+++ b/Engine/Source/Runtime/Experimental/Iris/Core/Private/Iris/ReplicationSystem/ReplicationOperations.cpp
@@ -5,6 +5,7 @@
 #include "Iris/Core/IrisLog.h"
 #include "Net/Core/NetBitArray.h"
 #include "Iris/ReplicationState/ReplicationStateUtil.h"
+#include "Iris/ReplicationState/InternalPropertyReplicationState.h"
 #include "Iris/ReplicationState/InternalReplicationStateDescriptorUtils.h"
 #include "Iris/ReplicationSystem/ObjectReferenceCache.h"
 #include "Iris/ReplicationSystem/ReplicationFragment.h"
@@ -661,6 +662,11 @@ void FReplicationStateOperations::DeserializeDeltaWithMask(FNetSerializationCont
 	}
 }
 
+void FReplicationStateOperations::ApplyStruct(FNetSerializationContext& Context, uint8* RESTRICT DstExternalBuffer, const uint8* RESTRICT SrcExternalBuffer, const FReplicationStateDescriptor* Descriptor)
+{
+	Private::InternalApplyStructProperty(Descriptor, DstExternalBuffer, SrcExternalBuffer);
+}
+
 bool FReplicationInstanceOperations::PollAndCopyPropertyData(const FReplicationInstanceProtocol* InstanceProtocol, EReplicationFragmentTraits ExcludeTraits, EReplicationFragmentPollFlags PollOptions)
 {
 	bool bIsStateDirty = false;
diff --git a/Engine/Source/Runtime/Experimental/Iris/Core/Private/Iris/Serialization/InstancedStructNetSerializer.cpp b/Engine/Source/Runtime/Experimental/Iris/Core/Private/Iris/Serialization/InstancedStructNetSerializer.cpp
--- a/Engine/Source/Runtime/Experimental/Iris/Core/Private/Iris/Serialization/InstancedStructNetSerializer.cpp
+++ b/Engine/Source/Runtime/Experimental/Iris/Core/Private/Iris/Serialization/InstancedStructNetSerializer.cpp
@@ -9,8 +9,9 @@
 #include "HAL/IConsoleManager.h"
 #include "Iris/Core/IrisLog.h"
 #include "Iris/Core/NetObjectReference.h"
-#include "Iris/ReplicationState/ReplicationStateDescriptorBuilder.h"
 #include "Iris/ReplicationState/PropertyNetSerializerInfoRegistry.h"
+#include "Iris/ReplicationState/ReplicationStateDescriptorBuilder.h"
+#include "Iris/ReplicationSystem/ReplicationOperations.h"
 #include "Iris/Serialization/NetReferenceCollector.h"
 #include "Iris/Serialization/NetSerializerArrayStorage.h"
 #include "Iris/Serialization/NetSerializerDelegates.h"
@@ -92,6 +93,8 @@ struct FInstancedStructNetSerializer
 
 	static void CollectNetReferences(FNetSerializationContext&, const FNetCollectReferencesArgs&);
 
+	static void Apply(FNetSerializationContext&, const FNetApplyArgs&);
+
 private:
 	// Frees dynamic memory allocated by the struct instance. Zeros the struct storage. Does not free the struct storage. After the call the Value is ready to be re-purposed for a different type of struct.
 	static void FreeStructInstance(FNetSerializationContext&, FInstancedStructNetSerializerConfig*, QuantizedType& Value);
@@ -470,6 +473,26 @@ void FInstancedStructNetSerializer::CollectNetReferences(FNetSerializationContex
 	}
 }
 
+void FInstancedStructNetSerializer::Apply(FNetSerializationContext& Context, const FNetApplyArgs& Args)
+{
+	const SourceType& Source = *reinterpret_cast<const SourceType*>(Args.Source);
+	SourceType& Target = *reinterpret_cast<SourceType*>(Args.Target);
+	FInstancedStructNetSerializerConfig* Config = const_cast<FInstancedStructNetSerializerConfig*>(static_cast<const FInstancedStructNetSerializerConfig*>(Args.NetSerializerConfig));
+
+	// If source and target has the same type make sure to not clobber not replicated properties.
+	const UScriptStruct* ScriptStruct = Source.GetScriptStruct();
+	if (ScriptStruct != Target.GetScriptStruct())
+	{
+		Target.InitializeAs(ScriptStruct);
+	}
+
+	if (ScriptStruct)
+	{
+		const TRefCountPtr<const FReplicationStateDescriptor> Descriptor = Config->DescriptorCache.FindOrAddDescriptor(ScriptStruct);
+		FReplicationStateOperations::ApplyStruct(Context, Target.GetMutableMemory(), Source.GetMemory(), Descriptor);
+	}
+}
+
 void FInstancedStructNetSerializer::FreeStructInstance(FNetSerializationContext& Context, FInstancedStructNetSerializerConfig* Config, QuantizedType& Value)
 {
 	InternalFreeStructInstance(Context, Config, Value);
@@ -642,6 +665,18 @@ TRefCountPtr<const FReplicationStateDescriptor> FInstancedStructDescriptorCache:
 	}
 }
 
+TRefCountPtr<const FReplicationStateDescriptor> FInstancedStructDescriptorCache::FindDescriptor(const UScriptStruct* Struct)
+{
+	if (!Struct)
+	{
+		return TRefCountPtr<const FReplicationStateDescriptor>();
+	}
+
+	const FString& PathNameString = Struct->GetPathName();
+	const FName PathName(PathNameString);
+	return FindDescriptor(PathName);
+}
+
 TRefCountPtr<const FReplicationStateDescriptor> FInstancedStructDescriptorCache::FindOrAddDescriptor(FName StructPath)
 {
 	TRefCountPtr<const FReplicationStateDescriptor> Descriptor;
diff --git a/Engine/Source/Runtime/Experimental/Iris/Core/Public/Iris/ReplicationState/ReplicationStateDescriptor.h b/Engine/Source/Runtime/Experimental/Iris/Core/Public/Iris/ReplicationState/ReplicationStateDescriptor.h
--- a/Engine/Source/Runtime/Experimental/Iris/Core/Public/Iris/ReplicationState/ReplicationStateDescriptor.h
+++ b/Engine/Source/Runtime/Experimental/Iris/Core/Public/Iris/ReplicationState/ReplicationStateDescriptor.h
@@ -179,8 +179,9 @@ enum class EReplicationStateTraits : uint32
 	// Whether delta compression is supported or not
 	SupportsDeltaCompression			= HasPushBasedDirtiness << 1U,
 	UseSerializerIsEqual				= SupportsDeltaCompression << 1U,
-	// Whether this is a descriptor for a struct derived from something with a NetSerializer.
-	IsDerivedStruct						= UseSerializerIsEqual << 1U,
+	// Whether this is a descriptor for a struct derived from something with a NetSerializer, but doesn't have a custom NetSerializer itself.
+	IsDerivedFromStructWithCustomSerializer	= UseSerializerIsEqual << 1U,
+	IsStructWithCustomSerializer = IsDerivedFromStructWithCustomSerializer << 1U,
 };
 ENUM_CLASS_FLAGS(EReplicationStateTraits);
 
diff --git a/Engine/Source/Runtime/Experimental/Iris/Core/Public/Iris/ReplicationSystem/ReplicationOperations.h b/Engine/Source/Runtime/Experimental/Iris/Core/Public/Iris/ReplicationSystem/ReplicationOperations.h
--- a/Engine/Source/Runtime/Experimental/Iris/Core/Public/Iris/ReplicationSystem/ReplicationOperations.h
+++ b/Engine/Source/Runtime/Experimental/Iris/Core/Public/Iris/ReplicationSystem/ReplicationOperations.h
@@ -73,6 +73,9 @@ struct FReplicationStateOperations
 
 	/** Deserialize Replication state from BitStream to internal buffer */
 	static IRISCORE_API void DeserializeDeltaWithMask(FNetSerializationContext& Context, const FNetBitArrayView& ChangeMask, const uint32 ChangeMaskOffset, uint8* RESTRICT DstInternalBuffer, const uint8* RESTRICT PrevInternalBuffer, const FReplicationStateDescriptor* Descriptor);
+
+	/** Apply/copy state for a descriptor created for a struct. */
+	static IRISCORE_API void ApplyStruct(FNetSerializationContext& Context, uint8* RESTRICT DstExternalBuffer, const uint8* RESTRICT SrcExternalBuffer, const FReplicationStateDescriptor* Descriptor);
 };
 
 struct FReplicationInstanceOperations
diff --git a/Engine/Source/Runtime/Experimental/Iris/Core/Public/Iris/Serialization/InstancedStructNetSerializer.h b/Engine/Source/Runtime/Experimental/Iris/Core/Public/Iris/Serialization/InstancedStructNetSerializer.h
--- a/Engine/Source/Runtime/Experimental/Iris/Core/Public/Iris/Serialization/InstancedStructNetSerializer.h
+++ b/Engine/Source/Runtime/Experimental/Iris/Core/Public/Iris/Serialization/InstancedStructNetSerializer.h
@@ -34,6 +34,9 @@ public:
 	// Find descriptor for struct with fully qualified name.
 	TRefCountPtr<const FReplicationStateDescriptor> FindDescriptor(FName StructPath);
 
+	// Find descriptor for struct.
+	TRefCountPtr<const FReplicationStateDescriptor> FindDescriptor(const UScriptStruct* Struct);
+
 	// Find or create descriptor for struct with fully qualified name.
 	TRefCountPtr<const FReplicationStateDescriptor> FindOrAddDescriptor(FName StructPath);
 
