Iris - Structs with custom serializers with Apply will now prevent the AllMembersAreReplicated trait from being added to the ReplicationStateDescriptor.

diff --git a/Engine/Plugins/Runtime/ReplicationSystemTestPlugin/Source/Private/Tests/ReplicationState/TestReplicationStateDescriptorBuilder.cpp b/Engine/Plugins/Runtime/ReplicationSystemTestPlugin/Source/Private/Tests/ReplicationState/TestReplicationStateDescriptorBuilder.cpp
--- a/Engine/Plugins/Runtime/ReplicationSystemTestPlugin/Source/Private/Tests/ReplicationState/TestReplicationStateDescriptorBuilder.cpp
+++ b/Engine/Plugins/Runtime/ReplicationSystemTestPlugin/Source/Private/Tests/ReplicationState/TestReplicationStateDescriptorBuilder.cpp
@@ -6,6 +6,7 @@
 #include "Iris/ReplicationState/InternalReplicationStateDescriptorUtils.h"
 #include "Iris/ReplicationState/ReplicationStateDescriptorBuilder.h"
 #include "Iris/ReplicationState/InternalPropertyReplicationState.h"
+#include "Iris/ReplicationState/PropertyNetSerializerInfoRegistry.h"
 #include "Iris/Serialization/InternalNetSerializers.h"
 #include "Iris/Serialization/NetSerializers.h"
 // UnrealNetwork.h needed for DOREPLIFETIME macros.
@@ -1149,6 +1150,14 @@ UE_NET_TEST_FIXTURE(FTestReplicationStateDescriptorBuilderFixture, FTestReplicat
 	UE_NET_ASSERT_TRUE(EnumHasAnyFlags(Descriptor->Traits, EReplicationStateTraits::AllMembersAreReplicated));
 }
 
+UE_NET_TEST_FIXTURE(FTestReplicationStateDescriptorBuilderFixture, FTestReplicationStateDescriptor_StructWithCustomSerializerWithApply)
+{
+	InitDescriptorFromStruct(StaticStruct<FTestReplicationStateDescriptor_StructWithCustomSerializerWithApply>());
+	UE_NET_ASSERT_EQ(Descriptors.Num(), 1);
+	UE_NET_ASSERT_EQ(Descriptor->MemberCount, uint16(1));
+	UE_NET_ASSERT_FALSE(EnumHasAnyFlags(Descriptor->Traits, EReplicationStateTraits::AllMembersAreReplicated));
+}
+
 UE_NET_TEST_FIXTURE(FTestReplicationStateDescriptorBuilderFixture, FTestReplicationStateDescriptor_StructWithArrayOfNotFullyReplicatedStruct)
 {
 	InitDescriptorFromStruct(StaticStruct<FTestReplicationStateDescriptor_StructWithArrayOfNotFullyReplicatedStruct>());
@@ -1165,6 +1174,14 @@ UE_NET_TEST_FIXTURE(FTestReplicationStateDescriptorBuilderFixture, FTestReplicat
 	UE_NET_ASSERT_TRUE(EnumHasAnyFlags(Descriptor->Traits, EReplicationStateTraits::AllMembersAreReplicated));
 }
 
+UE_NET_TEST_FIXTURE(FTestReplicationStateDescriptorBuilderFixture, FTestReplicationStateDescriptor_StructWithArrayOfStructWithCustomSerializerWithApply)
+{
+	InitDescriptorFromStruct(StaticStruct<FTestReplicationStateDescriptor_StructWithArrayOfStructWithCustomSerializerWithApply>());
+	UE_NET_ASSERT_EQ(Descriptors.Num(), 1);
+	UE_NET_ASSERT_EQ(Descriptor->MemberCount, uint16(1));
+	UE_NET_ASSERT_FALSE(EnumHasAnyFlags(Descriptor->Traits, EReplicationStateTraits::AllMembersAreReplicated));
+}
+
 UE_NET_TEST_FIXTURE(FTestReplicationStateDescriptorBuilderFixture, UTestReplicationStateDescriptor_TestClassWithNotFullyReplicatedStructAndArrays)
 {
 	InitDescriptorsFromClass(UTestReplicationStateDescriptor_TestClassWithNotFullyReplicatedStructAndArrays::StaticClass());
@@ -1559,3 +1576,74 @@ void UTestReplicationStateDescriptor_TestClassWithPushBasedAndNonPushBasedProper
 	DOREPLIFETIME_WITH_PARAMS_FAST(ThisClass, IntNotPushBasedB, LifetimeParams);
 	DOREPLIFETIME_WITH_PARAMS_FAST(ThisClass, IntNotPushBasedC, LifetimeParams);
 }
+
+
+namespace UE::Net
+{
+
+struct FTestReplicationStateDescriptor_StructWithCustomSerializerWithApplyNetSerializer
+{
+	// Version
+	static const uint32 Version = 0;
+
+	// Types
+	typedef FTestReplicationStateDescriptor_StructWithCustomSerializerWithApply SourceType;
+	typedef int QuantizedType;
+	typedef FNetSerializerConfig ConfigType;
+
+	inline static const ConfigType DefaultConfig;
+
+	// 
+	static void Serialize(FNetSerializationContext&, const FNetSerializeArgs&)
+	{
+	}
+
+	static void Deserialize(FNetSerializationContext&, const FNetDeserializeArgs&)
+	{
+	}
+
+	static void Quantize(FNetSerializationContext&, const FNetQuantizeArgs&)
+	{
+	}
+
+	static void Dequantize(FNetSerializationContext&, const FNetDequantizeArgs&)
+	{
+	}
+
+	static bool IsEqual(FNetSerializationContext&, const FNetIsEqualArgs&)
+	{
+		return true;
+	}
+
+	static void Apply(FNetSerializationContext&, const FNetApplyArgs&)
+	{
+	}
+
+private:
+	class FNetSerializerRegistryDelegates final : private UE::Net::FNetSerializerRegistryDelegates
+	{
+	public:
+		virtual ~FNetSerializerRegistryDelegates();
+
+	private:
+		virtual void OnPreFreezeNetSerializerRegistry() override;
+
+		inline static const FName TestReplicationStateDescriptor_StructWithCustomSerializerWithApply = FName("TestReplicationStateDescriptor_StructWithCustomSerializerWithApply");
+		UE_NET_IMPLEMENT_NAMED_STRUCT_NETSERIALIZER_INFO(TestReplicationStateDescriptor_StructWithCustomSerializerWithApply, FTestReplicationStateDescriptor_StructWithCustomSerializerWithApplyNetSerializer);
+	};
+
+	inline static FTestReplicationStateDescriptor_StructWithCustomSerializerWithApplyNetSerializer::FNetSerializerRegistryDelegates NetSerializerRegistryDelegates;
+};
+UE_NET_IMPLEMENT_SERIALIZER(FTestReplicationStateDescriptor_StructWithCustomSerializerWithApplyNetSerializer);
+
+FTestReplicationStateDescriptor_StructWithCustomSerializerWithApplyNetSerializer::FNetSerializerRegistryDelegates::~FNetSerializerRegistryDelegates()
+{
+	UE_NET_UNREGISTER_NETSERIALIZER_INFO(TestReplicationStateDescriptor_StructWithCustomSerializerWithApply);
+}
+
+void FTestReplicationStateDescriptor_StructWithCustomSerializerWithApplyNetSerializer::FNetSerializerRegistryDelegates::OnPreFreezeNetSerializerRegistry()
+{
+	UE_NET_REGISTER_NETSERIALIZER_INFO(TestReplicationStateDescriptor_StructWithCustomSerializerWithApply);
+}
+
+}
diff --git a/Engine/Plugins/Runtime/ReplicationSystemTestPlugin/Source/Private/Tests/ReplicationState/TestReplicationStateDescriptorBuilder.h b/Engine/Plugins/Runtime/ReplicationSystemTestPlugin/Source/Private/Tests/ReplicationState/TestReplicationStateDescriptorBuilder.h
--- a/Engine/Plugins/Runtime/ReplicationSystemTestPlugin/Source/Private/Tests/ReplicationState/TestReplicationStateDescriptorBuilder.h
+++ b/Engine/Plugins/Runtime/ReplicationSystemTestPlugin/Source/Private/Tests/ReplicationState/TestReplicationStateDescriptorBuilder.h
@@ -5,14 +5,24 @@
 #include "NetworkAutomationTest.h"
 #include "UObject/ObjectMacros.h"
 #include "Engine/EngineTypes.h"
+#include "Iris/Serialization/NetSerializer.h"
 #include "Tests/Serialization/EnumTestTypes.h"
 #include "TestReplicationStateDescriptorBuilder.generated.h"
 
+
+namespace UE::Net
+{
+
+UE_NET_DECLARE_SERIALIZER(FTestReplicationStateDescriptor_StructWithCustomSerializerWithApplyNetSerializer, REPLICATIONSYSTEMTESTPLUGIN_API);
+
+}
+
 // N.B. Must use alignas(8) to get 8-byte aligned int64 on 32-bit, due to compiler options to pack structs with 4-byte alignment.
 
 USTRUCT()
 struct alignas(8) FTestReplicationStateDescriptor_TestStruct
 {
+
 	GENERATED_BODY()
 	
 	UPROPERTY(Transient)
@@ -707,6 +717,24 @@ struct FTestReplicationStateDescriptor_StructWithArrayOfFullyReplicatedStruct
 	TArray<FTestReplicationStateDescriptor_FullyReplicatedStruct> TArrayOfFullyReplicatedStruct;
 };
 
+USTRUCT()
+struct FTestReplicationStateDescriptor_StructWithCustomSerializerWithApply
+{
+	GENERATED_BODY()
+
+	UPROPERTY()
+	int32 IntA = 0;
+};
+
+USTRUCT()
+struct FTestReplicationStateDescriptor_StructWithArrayOfStructWithCustomSerializerWithApply
+{
+	GENERATED_BODY()
+
+	UPROPERTY()
+	TArray<FTestReplicationStateDescriptor_StructWithCustomSerializerWithApply> ArrayOfStructWithCustomSerializerWithApply;
+};
+
 UCLASS()
 class UTestReplicationStateDescriptor_TestClassWithNotFullyReplicatedStructAndArrays : public UObject
 {
diff --git a/Engine/Plugins/Runtime/ReplicationSystemTestPlugin/Source/Private/Tests/Serialization/TestInstancedStructNetSerializer.cpp b/Engine/Plugins/Runtime/ReplicationSystemTestPlugin/Source/Private/Tests/Serialization/TestInstancedStructNetSerializer.cpp
--- a/Engine/Plugins/Runtime/ReplicationSystemTestPlugin/Source/Private/Tests/Serialization/TestInstancedStructNetSerializer.cpp
+++ b/Engine/Plugins/Runtime/ReplicationSystemTestPlugin/Source/Private/Tests/Serialization/TestInstancedStructNetSerializer.cpp
@@ -604,6 +604,102 @@ UE_NET_TEST_FIXTURE(FTestInstancedStructNetSerializerFixture, CustomSerializedSt
 	UE_NET_ASSERT_EQ(ClientObject->InstancedStruct.Get<FStructForInstancedStructTestWithCustomSerializer>().NotPropertyUint, ServerObject->InstancedStruct.Get<FStructForInstancedStructTestWithCustomSerializer>().NotPropertyUint);
 }
 
+UE_NET_TEST_FIXTURE(FTestInstancedStructNetSerializerFixture, NotFullyReplicatedStructIsAppliedCorrectly)
+{
+	// Add a client
+	FReplicationSystemTestClient* Client = CreateClient();
+
+	// Spawn objects on server
+	UInstancedStructNetSerializerTestObject* ServerObject = Server->CreateObject<UInstancedStructNetSerializerTestObject>();
+
+	// Initialize with a not fully replicated struct
+	ServerObject->InstancedStruct.InitializeAs<FStructForInstancedStructTestWithNonReplicatedProperties>();
+	
+	// Replicate
+	Server->UpdateAndSend({ Client });
+	
+	UInstancedStructNetSerializerTestObject* ClientObject = Client->GetObjectAs<UInstancedStructNetSerializerTestObject>(ServerObject->NetRefHandle);
+
+	// Modify not replicated value on the client
+	ClientObject->InstancedStruct.GetMutable<FStructForInstancedStructTestWithNonReplicatedProperties>().IntB = 1337;
+
+	// Modify value on server to trigger replication and apply of the struct
+	ServerObject->InstancedStruct.GetMutable<FStructForInstancedStructTestWithNonReplicatedProperties>().IntA += 1;
+
+	// Replicate
+	Server->UpdateAndSend({ Client });
+
+	// Verified that the client got the modified value
+	UE_NET_ASSERT_EQ(ClientObject->InstancedStruct.Get<FStructForInstancedStructTestWithNonReplicatedProperties>().IntA, ServerObject->InstancedStruct.Get<FStructForInstancedStructTestWithNonReplicatedProperties>().IntA);
+
+	// Verify the not replicated value remained as previously set
+	UE_NET_ASSERT_EQ(ClientObject->InstancedStruct.Get<FStructForInstancedStructTestWithNonReplicatedProperties>().IntB, 1337);
+}
+
+UE_NET_TEST_FIXTURE(FTestInstancedStructNetSerializerFixture, NotFullyReplicatedStructInArrayIsAppliedCorrectly)
+{
+	// Add a client
+	FReplicationSystemTestClient* Client = CreateClient();
+
+	// Spawn objects on server
+	UInstancedStructNetSerializerTestObject* ServerObject = Server->CreateObject<UInstancedStructNetSerializerTestObject>();
+
+	// Initialize an array entry with a not fully replicated struct
+	ServerObject->InstancedStructArray.Add(FInstancedStruct::Make<FStructForInstancedStructTestWithNonReplicatedProperties>());
+	
+	// Replicate
+	Server->UpdateAndSend({ Client });
+	
+	UInstancedStructNetSerializerTestObject* ClientObject = Client->GetObjectAs<UInstancedStructNetSerializerTestObject>(ServerObject->NetRefHandle);
+
+	// Modify not replicated value on the client
+	ClientObject->InstancedStructArray[0].GetMutable<FStructForInstancedStructTestWithNonReplicatedProperties>().IntB = 1337;
+
+	// Modify value on server to trigger replication and apply of the struct
+	ServerObject->InstancedStructArray[0].GetMutable<FStructForInstancedStructTestWithNonReplicatedProperties>().IntA += 1;
+
+	// Replicate
+	Server->UpdateAndSend({ Client });
+
+	// Verified that the client got the modified value
+	UE_NET_ASSERT_EQ(ClientObject->InstancedStructArray[0].Get<FStructForInstancedStructTestWithNonReplicatedProperties>().IntA, ServerObject->InstancedStructArray[0].Get<FStructForInstancedStructTestWithNonReplicatedProperties>().IntA);
+
+	// Verify the not replicated value remained as previously set
+	UE_NET_ASSERT_EQ(ClientObject->InstancedStructArray[0].Get<FStructForInstancedStructTestWithNonReplicatedProperties>().IntB, 1337);
+}
+
+UE_NET_TEST_FIXTURE(FTestInstancedStructNetSerializerFixture, NotFullyReplicatedStructInArrayInStructIsAppliedCorrectly)
+{
+	// Add a client
+	FReplicationSystemTestClient* Client = CreateClient();
+
+	// Spawn objects on server
+	UInstancedStructNetSerializerTestObject* ServerObject = Server->CreateObject<UInstancedStructNetSerializerTestObject>();
+
+	// Initialize an array entry with a not fully replicated struct
+	ServerObject->StructWithInstancedStructArray.InstancedStructArray.Add(FInstancedStruct::Make<FStructForInstancedStructTestWithNonReplicatedProperties>());
+	
+	// Replicate
+	Server->UpdateAndSend({ Client });
+	
+	UInstancedStructNetSerializerTestObject* ClientObject = Client->GetObjectAs<UInstancedStructNetSerializerTestObject>(ServerObject->NetRefHandle);
+
+	// Modify not replicated value on the client
+	ClientObject->StructWithInstancedStructArray.InstancedStructArray[0].GetMutable<FStructForInstancedStructTestWithNonReplicatedProperties>().IntB = 1337;
+
+	// Modify value on server to trigger replication and apply of the struct
+	ServerObject->StructWithInstancedStructArray.InstancedStructArray[0].GetMutable<FStructForInstancedStructTestWithNonReplicatedProperties>().IntA += 1;
+
+	// Replicate
+	Server->UpdateAndSend({ Client });
+
+	// Verified that the client got the modified value
+	UE_NET_ASSERT_EQ(ClientObject->StructWithInstancedStructArray.InstancedStructArray[0].Get<FStructForInstancedStructTestWithNonReplicatedProperties>().IntA, ServerObject->StructWithInstancedStructArray.InstancedStructArray[0].Get<FStructForInstancedStructTestWithNonReplicatedProperties>().IntA);
+
+	// Verify the not replicated value remained as previously set
+	UE_NET_ASSERT_EQ(ClientObject->StructWithInstancedStructArray.InstancedStructArray[0].Get<FStructForInstancedStructTestWithNonReplicatedProperties>().IntB, 1337);
+}
+
 // Fixture implementation
 void FTestInstancedStructNetSerializerFixture::SetUp()
 {
@@ -1070,6 +1166,7 @@ void UInstancedStructNetSerializerTestObject::GetLifetimeReplicatedProps(TArray<
 
 	DOREPLIFETIME_WITH_PARAMS_FAST(ThisClass, InstancedStruct, Params);
 	DOREPLIFETIME_WITH_PARAMS_FAST(ThisClass, InstancedStructArray, Params);
+	DOREPLIFETIME_WITH_PARAMS_FAST(ThisClass, StructWithInstancedStructArray, Params);
 }
 
 void UInstancedStructNetSerializerTestObject::RegisterReplicationFragments(UE::Net::FFragmentRegistrationContext& Context, UE::Net::EFragmentRegistrationFlags RegistrationFlags)
diff --git a/Engine/Plugins/Runtime/ReplicationSystemTestPlugin/Source/Private/Tests/Serialization/TestInstancedStructNetSerializer.h b/Engine/Plugins/Runtime/ReplicationSystemTestPlugin/Source/Private/Tests/Serialization/TestInstancedStructNetSerializer.h
--- a/Engine/Plugins/Runtime/ReplicationSystemTestPlugin/Source/Private/Tests/Serialization/TestInstancedStructNetSerializer.h
+++ b/Engine/Plugins/Runtime/ReplicationSystemTestPlugin/Source/Private/Tests/Serialization/TestInstancedStructNetSerializer.h
@@ -140,5 +140,8 @@ public:
 
 	UPROPERTY(Replicated, Transient)
 	TArray<FInstancedStruct> InstancedStructArray;
+
+	UPROPERTY(Replicated, Transient)
+	FTestInstancedStructArray StructWithInstancedStructArray;
 };
 
diff --git a/Engine/Source/Runtime/Experimental/Iris/Core/Private/Iris/ReplicationState/ReplicationStateDescriptorBuilder.cpp b/Engine/Source/Runtime/Experimental/Iris/Core/Private/Iris/ReplicationState/ReplicationStateDescriptorBuilder.cpp
--- a/Engine/Source/Runtime/Experimental/Iris/Core/Private/Iris/ReplicationState/ReplicationStateDescriptorBuilder.cpp
+++ b/Engine/Source/Runtime/Experimental/Iris/Core/Private/Iris/ReplicationState/ReplicationStateDescriptorBuilder.cpp
@@ -613,6 +613,7 @@ void FPropertyReplicationStateDescriptorBuilder::BuildMemberCache(FBuilderContex
 {
 	EMemberPropertyTraits PropertyTraits = EMemberPropertyTraits::None;
 	EMemberPropertyTraits SharedPropertyTraits = (Members.Num() > 0 ? ~EMemberPropertyTraits(EMemberPropertyTraits::None) : EMemberPropertyTraits::None);
+	bool bAllMembersAreReplicated = Context.BuildParams.bAllMembersAreReplicated;
 
 	FMemberCache& MemberCache = Context.MemberCache;
 	MemberCache.SetNum(Members.Num());
@@ -624,6 +625,12 @@ void FPropertyReplicationStateDescriptorBuilder::BuildMemberCache(FBuilderContex
 		const FNetSerializer* Serializer = Member.SerializerInfo->GetNetSerializer(Member.Property);
 		CurrentCacheEntry->Serializer = Serializer;
 
+		// Disqualify any member having a serializer with an Apply from being considered as AllMembersAreReplicated as we need to be sure Apply is called as expected.
+		if (EnumHasAnyFlags(Serializer->Traits, ENetSerializerTraits::HasApply))
+		{
+			bAllMembersAreReplicated = false;
+		}
+
 		// This is a property based struct, make sure we have a descriptor for it.
 		if (IsStructNetSerializer(Serializer))
 		{
@@ -663,7 +670,7 @@ void FPropertyReplicationStateDescriptorBuilder::BuildMemberCache(FBuilderContex
 			if (!EnumHasAnyFlags(Descriptor->Traits, EReplicationStateTraits::AllMembersAreReplicated))
 			{
 				// If we are building a struct containing a struct that isn't fully replicated then this struct isn't fully replicated either.
-				Context.BuildParams.bAllMembersAreReplicated = false;
+				bAllMembersAreReplicated = false;
 			}
 
 			/**
@@ -700,7 +707,7 @@ void FPropertyReplicationStateDescriptorBuilder::BuildMemberCache(FBuilderContex
 			if (!EnumHasAnyFlags(Descriptor->Traits, EReplicationStateTraits::AllMembersAreReplicated))
 			{
 				// If we are building a struct containing a struct that isn't fully replicated then this struct isn't fully replicated either.
-				Context.BuildParams.bAllMembersAreReplicated = false;
+				bAllMembersAreReplicated = false;
 			}
 
 			CurrentCacheEntry->ExternalSizeAndAlignment = { (SIZE_T)Member.Property->GetElementSize(), (SIZE_T)Member.Property->GetMinAlignment() };
@@ -736,6 +743,9 @@ void FPropertyReplicationStateDescriptorBuilder::BuildMemberCache(FBuilderContex
 
 	// If some members are property based we generate the property descriptor.
 	Context.bGeneratePropertyDescriptors = bSomeMembersAreProperties;
+
+	// 
+	Context.BuildParams.bAllMembersAreReplicated = bAllMembersAreReplicated;
 }
 
 void FPropertyReplicationStateDescriptorBuilder::BuildMemberTagCache(FMemberTagCache& MemberTagCache, const FMemberCache& MemberCache) const
