Iris - Fix issues with owning connection information and replication conditionals not always being correct under rare circumstances.

diff --git a/Engine/Plugins/Runtime/ReplicationSystemTestPlugin/Source/Private/Tests/ReplicationSystem/Filtering/TestFiltering.cpp b/Engine/Plugins/Runtime/ReplicationSystemTestPlugin/Source/Private/Tests/ReplicationSystem/Filtering/TestFiltering.cpp
--- a/Engine/Plugins/Runtime/ReplicationSystemTestPlugin/Source/Private/Tests/ReplicationSystem/Filtering/TestFiltering.cpp
+++ b/Engine/Plugins/Runtime/ReplicationSystemTestPlugin/Source/Private/Tests/ReplicationSystem/Filtering/TestFiltering.cpp
@@ -3161,4 +3161,40 @@ UE_NET_TEST_FIXTURE(FTestFilteringFixture, LotsOfObjectsCanBeFilteredOutViaHyste
 	}
 }
 
+UE_NET_TEST_FIXTURE(FTestFilteringFixture, OwningConnectionIsSetProperly)
+{
+	FReplicationSystemTestClient* Client = CreateClient();
+
+	// Create object with subobject
+	UReplicatedTestObject* ServerObject1 = Server->CreateObject();
+	UReplicatedTestObject* ServerObject1SubObject = Server->CreateSubObject(ServerObject1->NetRefHandle, 0, 0);
+	Server->ReplicationSystem->SetOwningNetConnection(ServerObject1->NetRefHandle, Client->ConnectionIdOnServer);
+
+	Server->UpdateAndSend({Client});
+
+	// Delete subobject 
+	Server->DestroyObject(ServerObject1SubObject);
+	
+	// Create new subobject
+	ServerObject1SubObject = Server->CreateSubObject(ServerObject1->NetRefHandle, 0, 0);
+
+	// Update twice to have the subobject internal index available for reuse
+	Server->UpdateAndSend({Client});
+	Server->UpdateAndSend({Client});
+
+	// Create another object with subobject
+	UReplicatedTestObject* ServerObject2 = Server->CreateObject();
+	UReplicatedTestObject* ServerObject2SubObject = Server->CreateSubObject(ServerObject2->NetRefHandle, 0, 0);
+	Server->ReplicationSystem->SetOwningNetConnection(ServerObject2->NetRefHandle, Client->ConnectionIdOnServer);
+
+	Server->UpdateAndSend({Client});
+
+	// Verify all objects have the correct owning connection
+	UE_NET_ASSERT_EQ(Server->ReplicationSystem->GetOwningNetConnection(ServerObject1->NetRefHandle), Client->ConnectionIdOnServer);
+	UE_NET_ASSERT_EQ(Server->ReplicationSystem->GetOwningNetConnection(ServerObject1SubObject->NetRefHandle), Client->ConnectionIdOnServer);
+
+	UE_NET_ASSERT_EQ(Server->ReplicationSystem->GetOwningNetConnection(ServerObject2->NetRefHandle), Client->ConnectionIdOnServer);
+	UE_NET_ASSERT_EQ(Server->ReplicationSystem->GetOwningNetConnection(ServerObject2SubObject->NetRefHandle), Client->ConnectionIdOnServer);
+}
+
 } // end namespace UE::Net::Private
diff --git a/Engine/Plugins/Runtime/ReplicationSystemTestPlugin/Source/Private/Tests/ReplicationSystem/TestFlushBeforeDestroy.cpp b/Engine/Plugins/Runtime/ReplicationSystemTestPlugin/Source/Private/Tests/ReplicationSystem/TestFlushBeforeDestroy.cpp
--- a/Engine/Plugins/Runtime/ReplicationSystemTestPlugin/Source/Private/Tests/ReplicationSystem/TestFlushBeforeDestroy.cpp
+++ b/Engine/Plugins/Runtime/ReplicationSystemTestPlugin/Source/Private/Tests/ReplicationSystem/TestFlushBeforeDestroy.cpp
@@ -1316,4 +1316,463 @@ UE_NET_TEST_FIXTURE(FTestFlushBeforeDestroyFixture, TestCancelPendingDestroyWait
 	UE_NET_ASSERT_EQ(ServerObject->IntA, ClientObject->IntA);
 }
 
+// Test modifying an OwnerOnly property on an object and tear it off
+UE_NET_TEST_FIXTURE(FTestFlushBeforeDestroyFixture, TestModifyingOwnerOnlyProperty)
+{
+	// Add client
+	FReplicationSystemTestClient* Client = CreateClient();
+
+	// Spawn object on server
+	UTestReplicatedIrisObject* ServerObject = Server->CreateObject({.ConnectionFilteredComponentCount = 1U});
+	Server->ReplicationSystem->SetOwningNetConnection(ServerObject->NetRefHandle, Client->ConnectionIdOnServer);
+
+	// Send update
+	Server->UpdateAndSend({Client});
+
+	const UTestReplicatedIrisObject* ClientObject = Cast<UTestReplicatedIrisObject>(Client->GetReplicationBridge()->GetReplicatedObject(ServerObject->NetRefHandle));
+	const int PrevSkipOwnerA = ClientObject->ConnectionFilteredComponents[0]->SkipOwnerA;
+
+	ServerObject->ConnectionFilteredComponents[0]->ToOwnerA = 11;
+	ServerObject->ConnectionFilteredComponents[0]->SkipOwnerA = 17;
+
+	// Request tear off of object
+	Server->ReplicationBridge->EndReplication(ServerObject, EEndReplicationFlags::TearOff);
+
+	// Update once and fail sending to cause torn off object to no longer be in scope
+	Server->UpdateAndSend({Client}, DoNotDeliverPacket);
+
+	// Send update
+	Server->UpdateAndSend({Client}, DeliverPacket);
+
+	// Verify we got the updated owner only state
+	UE_NET_ASSERT_EQ(ClientObject->ConnectionFilteredComponents[0]->ToOwnerA, ServerObject->ConnectionFilteredComponents[0]->ToOwnerA);
+	UE_NET_ASSERT_EQ(ClientObject->ConnectionFilteredComponents[0]->SkipOwnerA, PrevSkipOwnerA);
+}
+
+// Test modifying a OwnerOnly property on a subobject and tear it off
+UE_NET_TEST_FIXTURE(FTestFlushBeforeDestroyFixture, TestModifyingOwnerOnlyPropertyOnSubobject)
+{
+	// Add client
+	FReplicationSystemTestClient* Client = CreateClient();
+
+	// Spawn object on server
+	UTestReplicatedIrisObject* ServerObject = Server->CreateObject();
+	UTestReplicatedIrisObject* ServerSubObject = Server->CreateSubObject(ServerObject->NetRefHandle, {.ConnectionFilteredComponentCount = 1U});
+	Server->ReplicationSystem->SetOwningNetConnection(ServerObject->NetRefHandle, Client->ConnectionIdOnServer);
+
+	// Send update
+	Server->UpdateAndSend({Client});
+
+	const UTestReplicatedIrisObject* ClientSubObject = Cast<UTestReplicatedIrisObject>(Client->GetReplicationBridge()->GetReplicatedObject(ServerSubObject->NetRefHandle));
+	const int PrevSkipOwnerA = ClientSubObject->ConnectionFilteredComponents[0]->SkipOwnerA;
+
+	ServerSubObject->ConnectionFilteredComponents[0]->ToOwnerA = 11;
+	ServerSubObject->ConnectionFilteredComponents[0]->SkipOwnerA = 17;
+
+	// Request tear off of subobject
+	Server->ReplicationBridge->EndReplication(ServerSubObject, EEndReplicationFlags::TearOff);
+
+	// Update once and fail sending to cause torn off subobject to no longer be in scope
+	Server->UpdateAndSend({Client}, DoNotDeliverPacket);
+
+	// Send update
+	Server->UpdateAndSend({Client}, DeliverPacket);
+
+	// Verify we got the updated owner only state
+	UE_NET_ASSERT_EQ(ClientSubObject->ConnectionFilteredComponents[0]->ToOwnerA, ServerSubObject->ConnectionFilteredComponents[0]->ToOwnerA);
+	UE_NET_ASSERT_EQ(ClientSubObject->ConnectionFilteredComponents[0]->SkipOwnerA, PrevSkipOwnerA);
+}
+
+// Test modifying a OwnerOnly property on a subobject and tear it off and switch owners before subobject state has flushed
+UE_NET_TEST_FIXTURE(FTestFlushBeforeDestroyFixture, TestModifyingOwnerOnlyPropertyOnSubobjectAndSwitchOwnerBeforeFlushed)
+{
+	// Add client
+	FReplicationSystemTestClient* Client = CreateClient();
+
+	// Spawn object on server
+	UTestReplicatedIrisObject* ServerObject = Server->CreateObject();
+	UTestReplicatedIrisObject* ServerSubObject = Server->CreateSubObject(ServerObject->NetRefHandle, {.ConnectionFilteredComponentCount = 1U});
+
+	// Send update
+	Server->UpdateAndSend({Client});
+
+	const UTestReplicatedIrisObject* ClientSubObject = Cast<UTestReplicatedIrisObject>(Client->GetReplicationBridge()->GetReplicatedObject(ServerSubObject->NetRefHandle));
+	const int PrevSkipOwnerA = ClientSubObject->ConnectionFilteredComponents[0]->SkipOwnerA;
+
+	ServerSubObject->ConnectionFilteredComponents[0]->ToOwnerA = 11;
+	ServerSubObject->ConnectionFilteredComponents[0]->SkipOwnerA = 17;
+
+	// Request tear off of subobject
+	Server->ReplicationBridge->EndReplication(ServerSubObject, EEndReplicationFlags::TearOff);
+
+	// Update once and fail sending to cause torn off subobject to no longer be in scope
+	Server->UpdateAndSend({Client}, DoNotDeliverPacket);
+
+	// Switch owners on root object
+	Server->ReplicationSystem->SetOwningNetConnection(ServerObject->NetRefHandle, Client->ConnectionIdOnServer);
+
+	// Send update
+	Server->UpdateAndSend({Client}, DeliverPacket);
+
+	// Verify we got the updated owner only state
+	UE_NET_ASSERT_EQ(ClientSubObject->ConnectionFilteredComponents[0]->ToOwnerA, ServerSubObject->ConnectionFilteredComponents[0]->ToOwnerA);
+	UE_NET_ASSERT_EQ(ClientSubObject->ConnectionFilteredComponents[0]->SkipOwnerA, PrevSkipOwnerA);
+}
+
+// Test modifying a SkipOwner property on a subobject and tear it off and switch owners before subobject state has flushed
+UE_NET_TEST_FIXTURE(FTestFlushBeforeDestroyFixture, TestModifyingSkipOwnerPropertyOnSubobjectAndSwitchOwnerBeforeFlushed)
+{
+	// Add client
+	FReplicationSystemTestClient* Client = CreateClient();
+
+	// Spawn object on server
+	UTestReplicatedIrisObject* ServerObject = Server->CreateObject();
+	UTestReplicatedIrisObject* ServerSubObject = Server->CreateSubObject(ServerObject->NetRefHandle, {.ConnectionFilteredComponentCount = 1U});
+
+	Server->ReplicationSystem->SetOwningNetConnection(ServerObject->NetRefHandle, Client->ConnectionIdOnServer);
+
+	// Send update
+	Server->UpdateAndSend({Client});
+
+	const UTestReplicatedIrisObject* ClientSubObject = Cast<UTestReplicatedIrisObject>(Client->GetReplicationBridge()->GetReplicatedObject(ServerSubObject->NetRefHandle));
+	const int PrevOwnerA = ClientSubObject->ConnectionFilteredComponents[0]->ToOwnerA;
+
+	ServerSubObject->ConnectionFilteredComponents[0]->ToOwnerA = 11;
+	ServerSubObject->ConnectionFilteredComponents[0]->SkipOwnerA = 17;
+
+	// Request tear off of subobject
+	Server->ReplicationBridge->EndReplication(ServerSubObject, EEndReplicationFlags::TearOff);
+
+	// Update once and fail sending to cause torn off subobject to no longer be in scope
+	Server->UpdateAndSend({Client}, DoNotDeliverPacket);
+
+	// Switch owners on root object
+	Server->ReplicationSystem->SetOwningNetConnection(ServerObject->NetRefHandle, InvalidConnectionId);
+
+	// Send update
+	Server->UpdateAndSend({Client}, DeliverPacket);
+
+	// Verify we got the updated owner only state
+	UE_NET_ASSERT_EQ(ClientSubObject->ConnectionFilteredComponents[0]->ToOwnerA, PrevOwnerA);
+	UE_NET_ASSERT_EQ(ClientSubObject->ConnectionFilteredComponents[0]->SkipOwnerA, ServerSubObject->ConnectionFilteredComponents[0]->SkipOwnerA);
+}
+
+// Test setting subobject conditions before root object tear off
+UE_NET_TEST_FIXTURE(FTestFlushBeforeDestroyFixture, TestSubobjectConditionOnRootObjectTearOff)
+{
+	// Add client
+	FReplicationSystemTestClient* Client = CreateClient();
+
+	// Spawn object on server
+	UTestReplicatedIrisObject* ServerObject = Server->CreateObject();
+	UTestReplicatedIrisObject* ServerSubObject0 = Server->CreateSubObject(ServerObject->NetRefHandle, {.ConnectionFilteredComponentCount = 1U});
+	UTestReplicatedIrisObject* ServerSubObject1 = Server->CreateSubObject(ServerObject->NetRefHandle, {.ConnectionFilteredComponentCount = 1U});
+
+	Server->ReplicationSystem->SetOwningNetConnection(ServerObject->NetRefHandle, Client->ConnectionIdOnServer);
+
+	// Send update
+	Server->UpdateAndSend({Client});
+
+	const UTestReplicatedIrisObject* ClientSubObject0 = Cast<UTestReplicatedIrisObject>(Client->GetReplicationBridge()->GetReplicatedObject(ServerSubObject0->NetRefHandle));
+	const UTestReplicatedIrisObject* ClientSubObject1 = Cast<UTestReplicatedIrisObject>(Client->GetReplicationBridge()->GetReplicatedObject(ServerSubObject1->NetRefHandle));
+
+	const int SubObjectPrevIntA = ClientSubObject1->IntA;
+	const int SubObjectPrevOwnerA = ClientSubObject1->ConnectionFilteredComponents[0]->ToOwnerA;
+	const int SubObjectPrevSkipOwnerA = ClientSubObject1->ConnectionFilteredComponents[0]->SkipOwnerA;
+
+	ServerSubObject0->IntA += 11;
+	ServerSubObject0->ConnectionFilteredComponents[0]->ToOwnerA += 12;
+	ServerSubObject0->ConnectionFilteredComponents[0]->SkipOwnerA += 13;
+	ServerSubObject1->IntA += 14;
+	ServerSubObject1->ConnectionFilteredComponents[0]->ToOwnerA += 15;
+	ServerSubObject1->ConnectionFilteredComponents[0]->SkipOwnerA += 16;
+
+	// Set different subobject conditions
+	Server->ReplicationBridge->SetSubObjectNetCondition(ServerSubObject0->NetRefHandle, ELifetimeCondition::COND_OwnerOnly);
+	Server->ReplicationBridge->SetSubObjectNetCondition(ServerSubObject1->NetRefHandle, ELifetimeCondition::COND_SkipOwner);
+
+	// Request tear off of object
+	Server->ReplicationBridge->EndReplication(ServerObject, EEndReplicationFlags::TearOff);
+
+	// Update once and fail sending to cause torn off object to no longer be in scope
+	Server->UpdateAndSend({Client}, DoNotDeliverPacket);
+
+	// Send update
+	Server->UpdateAndSend({Client}, DeliverPacket);
+
+	// Verify we got the updated state only for the OwnerOnly subobject
+	UE_NET_ASSERT_EQ(ClientSubObject0->IntA, ServerSubObject0->IntA);
+	UE_NET_ASSERT_EQ(ClientSubObject0->ConnectionFilteredComponents[0]->ToOwnerA, ServerSubObject0->ConnectionFilteredComponents[0]->ToOwnerA);
+	UE_NET_ASSERT_EQ(ClientSubObject0->ConnectionFilteredComponents[0]->SkipOwnerA, SubObjectPrevSkipOwnerA);
+
+	UE_NET_ASSERT_EQ(ClientSubObject1->IntA, SubObjectPrevIntA);
+	UE_NET_ASSERT_EQ(ClientSubObject1->ConnectionFilteredComponents[0]->ToOwnerA, SubObjectPrevOwnerA);
+	UE_NET_ASSERT_EQ(ClientSubObject1->ConnectionFilteredComponents[0]->SkipOwnerA, SubObjectPrevSkipOwnerA);
+}
+
+// Test setting subobject conditions before subobject tear off
+UE_NET_TEST_FIXTURE(FTestFlushBeforeDestroyFixture, TestSubobjectConditionOnSubobjectTearOff)
+{
+	// Add client
+	FReplicationSystemTestClient* Client = CreateClient();
+
+	// Spawn object on server
+	UTestReplicatedIrisObject* ServerObject = Server->CreateObject();
+	UTestReplicatedIrisObject* ServerSubObject0 = Server->CreateSubObject(ServerObject->NetRefHandle, {.ConnectionFilteredComponentCount = 1U});
+	UTestReplicatedIrisObject* ServerSubObject1 = Server->CreateSubObject(ServerObject->NetRefHandle, {.ConnectionFilteredComponentCount = 1U});
+
+	Server->ReplicationSystem->SetOwningNetConnection(ServerObject->NetRefHandle, Client->ConnectionIdOnServer);
+
+	// Send update
+	Server->UpdateAndSend({Client});
+
+	const UTestReplicatedIrisObject* ClientSubObject0 = Cast<UTestReplicatedIrisObject>(Client->GetReplicationBridge()->GetReplicatedObject(ServerSubObject0->NetRefHandle));
+	const UTestReplicatedIrisObject* ClientSubObject1 = Cast<UTestReplicatedIrisObject>(Client->GetReplicationBridge()->GetReplicatedObject(ServerSubObject1->NetRefHandle));
+
+	const int SubObjectPrevIntA = ClientSubObject1->IntA;
+	const int SubObjectPrevOwnerA = ClientSubObject1->ConnectionFilteredComponents[0]->ToOwnerA;
+	const int SubObjectPrevSkipOwnerA = ClientSubObject1->ConnectionFilteredComponents[0]->SkipOwnerA;
+
+	ServerSubObject0->IntA += 11;
+	ServerSubObject0->ConnectionFilteredComponents[0]->ToOwnerA += 12;
+	ServerSubObject0->ConnectionFilteredComponents[0]->SkipOwnerA += 13;
+	ServerSubObject1->IntA += 14;
+	ServerSubObject1->ConnectionFilteredComponents[0]->ToOwnerA += 15;
+	ServerSubObject1->ConnectionFilteredComponents[0]->SkipOwnerA += 16;
+
+	// Set different subobject conditions
+	Server->ReplicationBridge->SetSubObjectNetCondition(ServerSubObject0->NetRefHandle, ELifetimeCondition::COND_OwnerOnly);
+	Server->ReplicationBridge->SetSubObjectNetCondition(ServerSubObject1->NetRefHandle, ELifetimeCondition::COND_SkipOwner);
+
+	// Request tear off of subobjects
+	Server->ReplicationBridge->EndReplication(ServerSubObject0, EEndReplicationFlags::TearOff);
+	Server->ReplicationBridge->EndReplication(ServerSubObject1, EEndReplicationFlags::TearOff);
+
+	// Update once and fail sending to cause torn off object to no longer be in scope
+	Server->UpdateAndSend({Client}, DoNotDeliverPacket);
+
+	// Send update
+	Server->UpdateAndSend({Client}, DeliverPacket);
+
+	// Verify we got the updated state only for the OwnerOnly subobject
+	UE_NET_ASSERT_EQ(ClientSubObject0->IntA, ServerSubObject0->IntA);
+	UE_NET_ASSERT_EQ(ClientSubObject0->ConnectionFilteredComponents[0]->ToOwnerA, ServerSubObject0->ConnectionFilteredComponents[0]->ToOwnerA);
+	UE_NET_ASSERT_EQ(ClientSubObject0->ConnectionFilteredComponents[0]->SkipOwnerA, SubObjectPrevSkipOwnerA);
+
+	UE_NET_ASSERT_EQ(ClientSubObject1->IntA, SubObjectPrevIntA);
+	UE_NET_ASSERT_EQ(ClientSubObject1->ConnectionFilteredComponents[0]->ToOwnerA, SubObjectPrevOwnerA);
+	UE_NET_ASSERT_EQ(ClientSubObject1->ConnectionFilteredComponents[0]->SkipOwnerA, SubObjectPrevSkipOwnerA);
+}
+
+// Test setting subobject condition on subobjects before subobject tear off and set owner before subobject state has flushed
+UE_NET_TEST_FIXTURE(FTestFlushBeforeDestroyFixture, TestSubobjectConditionOnSubobjectTearOffAndSetOwnerBeforeFlushed)
+{
+	// Add client
+	FReplicationSystemTestClient* Client = CreateClient();
+
+	// Spawn object on server
+	UTestReplicatedIrisObject* ServerObject = Server->CreateObject();
+	UTestReplicatedIrisObject* ServerSubObject0 = Server->CreateSubObject(ServerObject->NetRefHandle, {.ConnectionFilteredComponentCount = 1U});
+	UTestReplicatedIrisObject* ServerSubObject1 = Server->CreateSubObject(ServerObject->NetRefHandle, {.ConnectionFilteredComponentCount = 1U});
+
+	// Send update
+	Server->UpdateAndSend({Client});
+
+	const UTestReplicatedIrisObject* ClientSubObject0 = Cast<UTestReplicatedIrisObject>(Client->GetReplicationBridge()->GetReplicatedObject(ServerSubObject0->NetRefHandle));
+	const UTestReplicatedIrisObject* ClientSubObject1 = Cast<UTestReplicatedIrisObject>(Client->GetReplicationBridge()->GetReplicatedObject(ServerSubObject1->NetRefHandle));
+
+	const int SubObjectPrevIntA = ClientSubObject1->IntA;
+	const int SubObjectPrevOwnerA = ClientSubObject1->ConnectionFilteredComponents[0]->ToOwnerA;
+	const int SubObjectPrevSkipOwnerA = ClientSubObject1->ConnectionFilteredComponents[0]->SkipOwnerA;
+
+	ServerSubObject0->IntA += 11;
+	ServerSubObject0->ConnectionFilteredComponents[0]->ToOwnerA += 12;
+	ServerSubObject0->ConnectionFilteredComponents[0]->SkipOwnerA += 13;
+	ServerSubObject1->IntA += 14;
+	ServerSubObject1->ConnectionFilteredComponents[0]->ToOwnerA += 15;
+	ServerSubObject1->ConnectionFilteredComponents[0]->SkipOwnerA += 16;
+
+	// Set different subobject conditions
+	Server->ReplicationBridge->SetSubObjectNetCondition(ServerSubObject0->NetRefHandle, ELifetimeCondition::COND_OwnerOnly);
+	Server->ReplicationBridge->SetSubObjectNetCondition(ServerSubObject1->NetRefHandle, ELifetimeCondition::COND_SkipOwner);
+
+	// Request tear off of subobjects
+	Server->ReplicationBridge->EndReplication(ServerSubObject0, EEndReplicationFlags::TearOff);
+	Server->ReplicationBridge->EndReplication(ServerSubObject1, EEndReplicationFlags::TearOff);
+
+	// Update once and fail sending to cause torn off object to no longer be in scope
+	Server->UpdateAndSend({Client}, DoNotDeliverPacket);
+
+	// Set owner on root object
+	Server->ReplicationSystem->SetOwningNetConnection(ServerObject->NetRefHandle, Client->ConnectionIdOnServer);
+ 
+	// Send update
+	Server->UpdateAndSend({Client}, DeliverPacket);
+
+	// Verify we got the updated state only for the OwnerOnly subobject
+	UE_NET_ASSERT_EQ(ClientSubObject0->IntA, ServerSubObject0->IntA);
+	UE_NET_ASSERT_EQ(ClientSubObject0->ConnectionFilteredComponents[0]->ToOwnerA, ServerSubObject0->ConnectionFilteredComponents[0]->ToOwnerA);
+	UE_NET_ASSERT_EQ(ClientSubObject0->ConnectionFilteredComponents[0]->SkipOwnerA, SubObjectPrevSkipOwnerA);
+
+	UE_NET_ASSERT_EQ(ClientSubObject1->IntA, SubObjectPrevIntA);
+	UE_NET_ASSERT_EQ(ClientSubObject1->ConnectionFilteredComponents[0]->ToOwnerA, SubObjectPrevOwnerA);
+	UE_NET_ASSERT_EQ(ClientSubObject1->ConnectionFilteredComponents[0]->SkipOwnerA, SubObjectPrevSkipOwnerA);
+}
+
+// Test setting subobject condition on subobjects before subobject tear off and switch owners before subobject state has flushed
+UE_NET_TEST_FIXTURE(FTestFlushBeforeDestroyFixture, TestSubobjectConditionOnSubobjectTearOffAndSwitchOwnerBeforeFlushed)
+{
+	// Add client
+	FReplicationSystemTestClient* Client = CreateClient();
+
+	// Spawn object on server
+	UTestReplicatedIrisObject* ServerObject = Server->CreateObject();
+	UTestReplicatedIrisObject* ServerSubObject0 = Server->CreateSubObject(ServerObject->NetRefHandle, {.ConnectionFilteredComponentCount = 1U});
+	UTestReplicatedIrisObject* ServerSubObject1 = Server->CreateSubObject(ServerObject->NetRefHandle, {.ConnectionFilteredComponentCount = 1U});
+
+	Server->ReplicationSystem->SetOwningNetConnection(ServerObject->NetRefHandle, Client->ConnectionIdOnServer);
+
+	// Send update
+	Server->UpdateAndSend({Client});
+
+	const UTestReplicatedIrisObject* ClientSubObject0 = Cast<UTestReplicatedIrisObject>(Client->GetReplicationBridge()->GetReplicatedObject(ServerSubObject0->NetRefHandle));
+	const UTestReplicatedIrisObject* ClientSubObject1 = Cast<UTestReplicatedIrisObject>(Client->GetReplicationBridge()->GetReplicatedObject(ServerSubObject1->NetRefHandle));
+
+	const int SubObjectPrevIntA = ClientSubObject1->IntA;
+	const int SubObjectPrevOwnerA = ClientSubObject1->ConnectionFilteredComponents[0]->ToOwnerA;
+	const int SubObjectPrevSkipOwnerA = ClientSubObject1->ConnectionFilteredComponents[0]->SkipOwnerA;
+
+	ServerSubObject0->IntA += 11;
+	ServerSubObject0->ConnectionFilteredComponents[0]->ToOwnerA += 12;
+	ServerSubObject0->ConnectionFilteredComponents[0]->SkipOwnerA += 13;
+	ServerSubObject1->IntA += 14;
+	ServerSubObject1->ConnectionFilteredComponents[0]->ToOwnerA += 15;
+	ServerSubObject1->ConnectionFilteredComponents[0]->SkipOwnerA += 16;
+
+	// Set different subobject conditions
+	Server->ReplicationBridge->SetSubObjectNetCondition(ServerSubObject0->NetRefHandle, ELifetimeCondition::COND_OwnerOnly);
+	Server->ReplicationBridge->SetSubObjectNetCondition(ServerSubObject1->NetRefHandle, ELifetimeCondition::COND_SkipOwner);
+
+	// Request tear off of subobjects
+	Server->ReplicationBridge->EndReplication(ServerSubObject0, EEndReplicationFlags::TearOff);
+	Server->ReplicationBridge->EndReplication(ServerSubObject1, EEndReplicationFlags::TearOff);
+
+	// Update once and fail sending to cause torn off object to no longer be in scope
+	Server->UpdateAndSend({Client}, DoNotDeliverPacket);
+
+	// Switch owners on root object
+	Server->ReplicationSystem->SetOwningNetConnection(ServerObject->NetRefHandle, UE::Net::InvalidConnectionId);
+ 
+	// Send update
+	Server->UpdateAndSend({Client}, DeliverPacket);
+
+	// Verify we got the updated state only for the SkipOwner subobject
+	UE_NET_ASSERT_EQ(ClientSubObject0->IntA, SubObjectPrevIntA);
+	UE_NET_ASSERT_EQ(ClientSubObject0->ConnectionFilteredComponents[0]->ToOwnerA, SubObjectPrevOwnerA);
+	UE_NET_ASSERT_EQ(ClientSubObject0->ConnectionFilteredComponents[0]->SkipOwnerA, SubObjectPrevSkipOwnerA);
+
+	UE_NET_ASSERT_EQ(ClientSubObject1->IntA, ServerSubObject1->IntA);
+	UE_NET_ASSERT_EQ(ClientSubObject1->ConnectionFilteredComponents[0]->ToOwnerA, SubObjectPrevOwnerA);
+	UE_NET_ASSERT_EQ(ClientSubObject1->ConnectionFilteredComponents[0]->SkipOwnerA, ServerSubObject1->ConnectionFilteredComponents[0]->SkipOwnerA);
+}
+
+// Test sending attachment to connection filtered object and tear it off
+UE_NET_TEST_FIXTURE(FTestFlushBeforeDestroyFixture, TestTearOffOfConnectionFilteredObject)
+{
+	// Add clients
+	FReplicationSystemTestClient* ClientArray[] = {CreateClient(), CreateClient()};
+	TStrongObjectPtr<UMockNetObjectAttachmentHandler> ClientMockNetObjectAttachmentHandlers[UE_ARRAY_COUNT(ClientArray)];
+	for (FReplicationSystemTestClient*& Client : ClientArray)
+	{
+		RegisterNetBlobHandlers(Client);
+		ClientMockNetObjectAttachmentHandlers[&Client - ClientArray] = ClientMockNetObjectAttachmentHandler;
+	}
+
+	// Setup case where we have a new object for which we have an attachment which should execute a tearoff after we have confirmed creation
+	UReplicatedTestObject* ServerObject = Server->CreateObject();
+
+	// Apply connection filter
+	TBitArray<> AllowedConnections;
+	AllowedConnections.Init(false, Clients[0]->ConnectionIdOnServer + 1);
+	AllowedConnections[Clients[0]->ConnectionIdOnServer] = true;
+	Server->ReplicationSystem->SetConnectionFilter(ServerObject->NetRefHandle, AllowedConnections, ENetFilterStatus::Allow);
+
+	// Send update
+	Server->UpdateAndSend({ClientArray}, DeliverPacket);
+
+	// Create attachment
+	{
+		constexpr uint32 PayloadBitCount = 24;
+		const TRefCountPtr<FNetObjectAttachment>& Attachment = MockNetObjectAttachmentHandler->CreateReliableNetObjectAttachment(PayloadBitCount);
+		FNetObjectReference AttachmentTarget = FObjectReferenceCache::MakeNetObjectReference(ServerObject->NetRefHandle);
+		for (FReplicationSystemTestClient* Client : ClientArray)
+		{
+			Server->GetReplicationSystem()->QueueNetObjectAttachment(Client->ConnectionIdOnServer, AttachmentTarget, Attachment);
+		}
+	}
+
+	// Request tearoff
+	Server->ReplicationBridge->EndReplication(ServerObject, EEndReplicationFlags::TearOff);
+
+	// Update once and fail sending to cause torn off object to no longer be in scope
+	Server->UpdateAndSend({ClientArray}, DoNotDeliverPacket);
+
+	// Deliver a packet, this should flush the object and deliver the attachment to the allowed connection
+	Server->UpdateAndSend({ClientArray}, DeliverPacket);
+
+	// Verify that the attachment has been received on the allowed connection
+	UE_NET_ASSERT_EQ(ClientMockNetObjectAttachmentHandlers[0]->GetFunctionCallCounts().OnNetBlobReceived, 1U);
+
+	// Verify that the attachment was not received on the disallowed connection
+	UE_NET_ASSERT_EQ(ClientMockNetObjectAttachmentHandlers[1]->GetFunctionCallCounts().OnNetBlobReceived, 0U);
+}
+
+// Test sending attachment to owner filtered object and tear it off
+UE_NET_TEST_FIXTURE(FTestFlushBeforeDestroyFixture, TestTearOffOfOwnerFilteredObject)
+{
+	// Add clients
+	FReplicationSystemTestClient* ClientArray[] = {CreateClient(), CreateClient()};
+	TStrongObjectPtr<UMockNetObjectAttachmentHandler> ClientMockNetObjectAttachmentHandlers[UE_ARRAY_COUNT(ClientArray)];
+	for (FReplicationSystemTestClient*& Client : ClientArray)
+	{
+		RegisterNetBlobHandlers(Client);
+		ClientMockNetObjectAttachmentHandlers[&Client - ClientArray] = ClientMockNetObjectAttachmentHandler;
+	}
+
+	// Setup case where we have a new object for which we have an attachment which should execute a tearoff after we have confirmed creation
+	UReplicatedTestObject* ServerObject = Server->CreateObject();
+
+	// Apply owner and owner filter
+	Server->ReplicationSystem->SetOwningNetConnection(ServerObject->NetRefHandle, ClientArray[0]->ConnectionIdOnServer);
+	Server->ReplicationSystem->SetFilter(ServerObject->NetRefHandle, ToOwnerFilterHandle);
+
+	// Send update
+	Server->UpdateAndSend({ClientArray}, DeliverPacket);
+
+	// Create attachment
+	{
+		constexpr uint32 PayloadBitCount = 24;
+		const TRefCountPtr<FNetObjectAttachment>& Attachment = MockNetObjectAttachmentHandler->CreateReliableNetObjectAttachment(PayloadBitCount);
+		FNetObjectReference AttachmentTarget = FObjectReferenceCache::MakeNetObjectReference(ServerObject->NetRefHandle);
+		for (FReplicationSystemTestClient* Client : ClientArray)
+		{
+			Server->GetReplicationSystem()->QueueNetObjectAttachment(Client->ConnectionIdOnServer, AttachmentTarget, Attachment);
+		}
+	}
+
+	// Request tearoff
+	Server->ReplicationBridge->EndReplication(ServerObject, EEndReplicationFlags::TearOff);
+
+	// Update once and fail sending to cause torn off object to no longer be in scope
+	Server->UpdateAndSend({ClientArray}, DoNotDeliverPacket);
+
+	// Deliver a packet, this should flush the object and deliver the attachment to the owning connection
+	Server->UpdateAndSend({ClientArray}, DeliverPacket);
+
+	// Verify that the attachment has been received on the allowed connection
+	UE_NET_ASSERT_EQ(ClientMockNetObjectAttachmentHandlers[0]->GetFunctionCallCounts().OnNetBlobReceived, 1U);
+
+	// Verify that the attachment was not received on the disallowed connection
+	UE_NET_ASSERT_EQ(ClientMockNetObjectAttachmentHandlers[1]->GetFunctionCallCounts().OnNetBlobReceived, 0U);
+}
 }
diff --git a/Engine/Source/Runtime/Experimental/Iris/Core/Private/Iris/ReplicationSystem/Conditionals/ReplicationConditionals.cpp b/Engine/Source/Runtime/Experimental/Iris/Core/Private/Iris/ReplicationSystem/Conditionals/ReplicationConditionals.cpp
--- a/Engine/Source/Runtime/Experimental/Iris/Core/Private/Iris/ReplicationSystem/Conditionals/ReplicationConditionals.cpp
+++ b/Engine/Source/Runtime/Experimental/Iris/Core/Private/Iris/ReplicationSystem/Conditionals/ReplicationConditionals.cpp
@@ -78,6 +78,24 @@ void FReplicationConditionals::OnMaxInternalNetRefIndexIncreased(FInternalNetRef
 	}
 }
 
+void FReplicationConditionals::OnInternalNetRefIndicesFreed(const TConstArrayView<FInternalNetRefIndex>& FreedIndices)
+{
+	IRIS_PROFILER_SCOPE(FReplicationConditionals_OnInternalNetRefIndicesFreed);
+	for (const FInternalNetRefIndex ObjectIndex : FreedIndices)
+	{
+		ClearPerObjectInfo(ObjectIndex);
+	}
+
+	const FNetBitArray& ValidConnections = ReplicationConnections->GetValidConnections();
+	if (ValidConnections.FindLastOne() != FNetBitArrayBase::InvalidIndex)
+	{
+		for (const FInternalNetRefIndex ObjectIndex : FreedIndices)
+		{
+			ClearConnectionInfosForObject(ValidConnections, ObjectIndex);
+		}
+	}
+}
+
 void FReplicationConditionals::MarkLifeTimeConditionalsDirtyForObjectsInGroup(FNetObjectGroupHandle GroupHandle)
 {
 	IRIS_PROFILER_SCOPE(MarkLifeTimeConditionalsDirtyForObjectsInGroup)
@@ -538,7 +556,6 @@ bool FReplicationConditionals::SetPropertyDynamicCondition(FInternalNetRefIndex
 
 void FReplicationConditionals::Update()
 {
-	UpdateObjectsInScope();
 	if (bEnableUpdateObjectsWithDirtyConditionals)
 	{
 		UpdateAndResetObjectsWithDirtyConditionals();
@@ -783,68 +800,6 @@ bool FReplicationConditionals::ApplyConditionalsToChangeMask(uint32 ReplicatingC
 	return bMaskWasModified;
 }
 
-void FReplicationConditionals::UpdateObjectsInScope()
-{
-	IRIS_PROFILER_SCOPE(FReplicationConditionals_UpdateObjectsInScope);
-
-	const FNetBitArrayView ObjectsInScope = NetRefHandleManager->GetCurrentFrameScopableInternalIndices();
-	const FNetBitArrayView PrevObjectsInScope = NetRefHandleManager->GetPrevFrameScopableInternalIndices();
-
-	const uint32 WordCountForModifiedWords = Align(FPlatformMath::Max(MaxInternalNetRefIndex, 1U), 32U)/32U;
-	TArray<uint32> ModifiedWords;
-	ModifiedWords.SetNumUninitialized(WordCountForModifiedWords);
-	uint32* ModifiedWordsStorage = ModifiedWords.GetData();
-	uint32 ModifiedWordIndex = 0U;
-	const uint32* ObjectsInScopeStorage = ObjectsInScope.GetData();
-	const uint32* PrevObjectsInScopeStorage = PrevObjectsInScope.GetData();
-	{
-		for (uint32 WordIt = 0U, WordEndIt = WordCountForModifiedWords; WordIt != WordEndIt; ++WordIt)
-		{
-			const bool bWordDiffers = ObjectsInScopeStorage[WordIt] != PrevObjectsInScopeStorage[WordIt];
-			ModifiedWordsStorage[ModifiedWordIndex] = WordIt;
-			ModifiedWordIndex += bWordDiffers;
-		}
-	}
-
-	// If the scope didn't change there's nothing for us to do.
-	if (ModifiedWordIndex == 0U)
-	{
-		return;
-	}
-
-	{
-		const FNetBitArray& ValidConnections = ReplicationConnections->GetValidConnections();
-		const uint32 MaxValidConnectionId = ValidConnections.FindLastOne();
-
-		for (uint32 WordIt = 0U, WordEndIt = ModifiedWordIndex; WordIt != WordEndIt; ++WordIt)
-		{
-			const uint32 WordIndex = ModifiedWordsStorage[WordIt];
-			const uint32 PrevExistingObjects = PrevObjectsInScopeStorage[WordIndex];
-			const uint32 ExistingObjects = ObjectsInScopeStorage[WordIndex];
-
-			const uint32 BitOffset = WordIndex*32U;
-
-			// Clear info for deleted objects.
-			{
-				uint32 DeletedObjects = PrevExistingObjects & ~ExistingObjects;
-				for (; DeletedObjects; )
-				{
-					const uint32 LeastSignificantBit = DeletedObjects & uint32(-int32(DeletedObjects));
-					DeletedObjects ^= LeastSignificantBit;
-
-					const uint32 ObjectIndex = BitOffset + FPlatformMath::CountTrailingZeros(LeastSignificantBit);
-					ClearPerObjectInfo(ObjectIndex);
-					if (MaxValidConnectionId != ~0U)
-					{
-						ClearConnectionInfosForObject(ValidConnections, MaxValidConnectionId, ObjectIndex);
-					}
-				}
-			}
-		}
-
-	}
-}
-
 void FReplicationConditionals::UpdateAndResetObjectsWithDirtyConditionals()
 {
 	IRIS_PROFILER_SCOPE(FReplicationConditionals_UpdateAndResetObjectsWithDirtyConditionals);
@@ -913,7 +868,7 @@ void FReplicationConditionals::ClearPerObjectInfo(FInternalNetRefIndex ObjectInd
 	DynamicConditions.Remove(ObjectIndex);
 }
 
-void FReplicationConditionals::ClearConnectionInfosForObject(const FNetBitArray& ValidConnections, uint32 MaxConnectionId, FInternalNetRefIndex ObjectIndex)
+void FReplicationConditionals::ClearConnectionInfosForObject(const FNetBitArray& ValidConnections, FInternalNetRefIndex ObjectIndex)
 {
 	IRIS_PROFILER_SCOPE(FReplicationConditionals_ClearConnectionInfosForObject);
 
diff --git a/Engine/Source/Runtime/Experimental/Iris/Core/Private/Iris/ReplicationSystem/Conditionals/ReplicationConditionals.h b/Engine/Source/Runtime/Experimental/Iris/Core/Private/Iris/ReplicationSystem/Conditionals/ReplicationConditionals.h
--- a/Engine/Source/Runtime/Experimental/Iris/Core/Private/Iris/ReplicationSystem/Conditionals/ReplicationConditionals.h
+++ b/Engine/Source/Runtime/Experimental/Iris/Core/Private/Iris/ReplicationSystem/Conditionals/ReplicationConditionals.h
@@ -57,6 +57,9 @@ public:
 	/** Called when the maximum InternalNetRefIndex increased and we need to realloc our lists */
 	void OnMaxInternalNetRefIndexIncreased(FInternalNetRefIndex NewMaxInternalIndex);
 
+	/** Called when when one or more NetRefInternalIndices have been freed and can be re-assigned to new objects. */
+	void OnInternalNetRefIndicesFreed(const TConstArrayView<FInternalNetRefIndex>& FreedIndices);
+
 	void AddConnection(uint32 ConnectionId);
 	void RemoveConnection(uint32 ConnectionId);
 
@@ -117,7 +120,6 @@ private:
 	};
 
 private:
-	void UpdateObjectsInScope();
 	void UpdateAndResetObjectsWithDirtyConditionals();
 
 	FConditionalsMask GetLifetimeConditionals(uint32 ReplicatingConnectionId, FInternalNetRefIndex ParentObjectIndex, bool bInitialState) const;
@@ -125,7 +127,7 @@ private:
 	FPerObjectInfo* GetPerObjectInfo(FInternalNetRefIndex ObjectIndex);
 	const FPerObjectInfo* GetPerObjectInfo(FInternalNetRefIndex ObjectIndex) const;
 	void ClearPerObjectInfo(FInternalNetRefIndex ObjectIndex);
-	void ClearConnectionInfosForObject(const FNetBitArray& ValidConnections, uint32 MaxConnectionId, FInternalNetRefIndex ObjectIndex);
+	void ClearConnectionInfosForObject(const FNetBitArray& ValidConnections, FInternalNetRefIndex ObjectIndex);
 
 	void GetChildSubObjectsToReplicate(uint32 ReplicatingConnectionId, const FConditionalsMask& LifetimeConditionals, const FInternalNetRefIndex ParentObjectIndex, FSubObjectsToReplicateArray& OutSubObjectsToReplicate);
 
diff --git a/Engine/Source/Runtime/Experimental/Iris/Core/Private/Iris/ReplicationSystem/Filtering/ReplicationFiltering.cpp b/Engine/Source/Runtime/Experimental/Iris/Core/Private/Iris/ReplicationSystem/Filtering/ReplicationFiltering.cpp
--- a/Engine/Source/Runtime/Experimental/Iris/Core/Private/Iris/ReplicationSystem/Filtering/ReplicationFiltering.cpp
+++ b/Engine/Source/Runtime/Experimental/Iris/Core/Private/Iris/ReplicationSystem/Filtering/ReplicationFiltering.cpp
@@ -318,6 +318,15 @@ void FReplicationFiltering::OnMaxInternalNetRefIndexIncreased(FInternalNetRefInd
 	}
 }
 
+void FReplicationFiltering::OnInternalNetRefIndicesFreed(const TConstArrayView<FInternalNetRefIndex>& FreedIndices)
+{
+	// Clear owner info just as the index is freed. We want to keep the owner info even when an object goes out of scope so that state flushing works as expected.
+	for (const FInternalNetRefIndex ObjectIndex : FreedIndices)
+	{
+		ObjectIndexToOwningConnection[ObjectIndex] = UE::Net::InvalidConnectionId;
+	}
+}
+
 void FReplicationFiltering::Filter()
 {
 #if UE_NET_IRIS_CSV_STATS
@@ -978,7 +987,6 @@ void FReplicationFiltering::UpdateObjectsInScope()
 				DeletedObjects ^= LeastSignificantBit;
 
 				const uint32 ObjectIndex = BitOffset + FPlatformMath::CountTrailingZeros(LeastSignificantBit);
-				ObjectIndexToOwningConnection[ObjectIndex] = 0U;
 
 				// Free per object info if present
 				FreePerObjectInfoForObject(ObjectIndex);
diff --git a/Engine/Source/Runtime/Experimental/Iris/Core/Private/Iris/ReplicationSystem/Filtering/ReplicationFiltering.h b/Engine/Source/Runtime/Experimental/Iris/Core/Private/Iris/ReplicationSystem/Filtering/ReplicationFiltering.h
--- a/Engine/Source/Runtime/Experimental/Iris/Core/Private/Iris/ReplicationSystem/Filtering/ReplicationFiltering.h
+++ b/Engine/Source/Runtime/Experimental/Iris/Core/Private/Iris/ReplicationSystem/Filtering/ReplicationFiltering.h
@@ -63,6 +63,9 @@ public:
 	/** Called when the maximum InternalNetRefIndex increased and we need to realloc our lists */
 	void OnMaxInternalNetRefIndexIncreased(FInternalNetRefIndex NewMaxInternalIndex);
 
+	/** Called when when one or more NetRefInternalIndices have been freed and can be re-assigned to new objects. */
+	void OnInternalNetRefIndicesFreed(const TConstArrayView<FInternalNetRefIndex>& FreedIndices);
+
 	/**
 	 * Executes group, owner and connection filtering then any dynamic filters.
 	 * At the end any object that is not relevant to at least one connection will be removed from the scoped object list.
diff --git a/Engine/Source/Runtime/Experimental/Iris/Core/Private/Iris/ReplicationSystem/NetRefHandleManager.cpp b/Engine/Source/Runtime/Experimental/Iris/Core/Private/Iris/ReplicationSystem/NetRefHandleManager.cpp
--- a/Engine/Source/Runtime/Experimental/Iris/Core/Private/Iris/ReplicationSystem/NetRefHandleManager.cpp
+++ b/Engine/Source/Runtime/Experimental/Iris/Core/Private/Iris/ReplicationSystem/NetRefHandleManager.cpp
@@ -558,6 +558,9 @@ void FNetRefHandleManager::DestroyObjectsPendingDestroy()
 {
 	IRIS_PROFILER_SCOPE(FNetRefHandleManager_DestroyObjectsPendingDestroy);
 
+	TArray<FInternalNetRefIndex> FreedInternalIndices;
+	FreedInternalIndices.Reserve(PendingDestroyInternalIndices.Num());
+
 	// Destroy Objects pending destroy
 	for (int32 It = PendingDestroyInternalIndices.Num() - 1; It >= 0; --It)
 	{
@@ -565,11 +568,18 @@ void FNetRefHandleManager::DestroyObjectsPendingDestroy()
 		// If we have subobjects pending tear off and such then wait before destroying the parent.
 		if (ReplicatedObjectRefCount[InternalIndex] == 0 && GetSubObjects(InternalIndex).Num() <= 0)
 		{
+			FreedInternalIndices.Add(InternalIndex);
+
 			InternalDestroyNetObject(InternalIndex);
 			PendingDestroyInternalIndices.RemoveAtSwap(It);
 		}
 	}
 	CSV_CUSTOM_STAT(IrisCommon, PendingDestroyInternalIndicesCount, (float)PendingDestroyInternalIndices.Num(), ECsvCustomStatOp::Set);
+
+	if (OnInternalNetRefIndicesFreed.IsBound())
+	{
+		OnInternalNetRefIndicesFreed.Broadcast(MakeConstArrayView(FreedInternalIndices));
+	}
 }
 
 bool FNetRefHandleManager::AddSubObject(FNetRefHandle OwnerHandle, FNetRefHandle SubObjectHandle, FNetRefHandle RelativeOtherSubObjectHandle, EAddSubObjectFlags Flags)
diff --git a/Engine/Source/Runtime/Experimental/Iris/Core/Private/Iris/ReplicationSystem/NetRefHandleManager.h b/Engine/Source/Runtime/Experimental/Iris/Core/Private/Iris/ReplicationSystem/NetRefHandleManager.h
--- a/Engine/Source/Runtime/Experimental/Iris/Core/Private/Iris/ReplicationSystem/NetRefHandleManager.h
+++ b/Engine/Source/Runtime/Experimental/Iris/Core/Private/Iris/ReplicationSystem/NetRefHandleManager.h
@@ -242,6 +242,11 @@ public:
 
 	bool SetSubObjectNetCondition(FInternalNetRefIndex SubObjectInternalIndex, FLifeTimeConditionStorage SubObjectCondition);
 
+	bool IsSubObject(FInternalNetRefIndex InternalIndex) const
+	{
+		return SubObjectInternalIndices.GetBit(InternalIndex);
+	}
+
 	// DependentObjects
 	const FNetBitArrayView GetDependentObjectInternalIndices() const { return MakeNetBitArrayView(DependentObjectInternalIndices); }
 	const FNetBitArrayView GetObjectsWithDependentObjectsInternalIndices() const { return MakeNetBitArrayView(ObjectsWithDependentObjectsInternalIndices); }
@@ -320,15 +325,19 @@ public:
 public:
 
 	/** Delegate that will notify when the highest NetChunkedArray internal index has increased (Highest = Max-1)*/
-	DECLARE_MULTICAST_DELEGATE_OneParam(FOnNetChunkedArrayIncrease, FInternalNetRefIndex HighestInternalIndex);
-	FOnNetChunkedArrayIncrease& GetOnNetChunkedArrayIncreaseDelegate() const { return OnNetChunkedArrayIncrease; };
+	using FOnNetChunkedArrayIncrease = TMulticastDelegate<void(FInternalNetRefIndex HighestInternalIndex)>;
+	FOnNetChunkedArrayIncrease::RegistrationType& GetOnNetChunkedArrayIncreaseDelegate() const { return OnNetChunkedArrayIncrease; };
 
 	/** Return the highest internal index that NetChunkedArrays are currently allocated for. */
 	uint32 GetHighestNetChunkedArrayInternalIndex() const { return HighestNetChunkedArrayInternalIndex; };
 
 	/** Delegate that will notify when the NetObjectLists (eg. NetBitArray and TArray indexed via FInternalNetRefIndex) have a new maximum index to support */
-	DECLARE_MULTICAST_DELEGATE_OneParam(FOnMaxInternalNetRefIndexIncreased, FInternalNetRefIndex MaxInternalIndex);
-	FOnMaxInternalNetRefIndexIncreased& GetOnMaxInternalNetRefIndexIncreasedDelegate() const { return OnMaxInternalNetRefIndexIncreased; };
+	using FOnMaxInternalNetRefIndexIncreased = TMulticastDelegate<void(FInternalNetRefIndex MaxInternalIndex)>;
+	FOnMaxInternalNetRefIndexIncreased::RegistrationType& GetOnMaxInternalNetRefIndexIncreasedDelegate() const { return OnMaxInternalNetRefIndexIncreased; };
+
+	/** Delegate that will notify when one or more NetRefInternalIndices have been freed and can be re-assigned to new objects */
+	using FOnInternalNetRefIndicesFreed = TMulticastDelegate<void(const TConstArrayView<FInternalNetRefIndex>& FreedIndices)>;
+	FOnInternalNetRefIndicesFreed::RegistrationType& GetOnInternalNetRefIndicesFreedDelegate() const { return OnInternalNetRefIndicesFreed; };
 
 public:
 
@@ -492,10 +501,10 @@ private:
 	// Delegates
 	mutable FOnMaxInternalNetRefIndexIncreased OnMaxInternalNetRefIndexIncreased;
 	mutable FOnNetChunkedArrayIncrease OnNetChunkedArrayIncrease;
+	mutable FOnInternalNetRefIndicesFreed OnInternalNetRefIndicesFreed;
 
 	// Track NetBitArrays owned by the NetRefHandleManager so they can automatically be resized.
 	TArray<FNetBitArray*> OwnedNetBitArrays;
-
 };
 
 const FNetRefHandleManager::FReplicatedObjectData& FNetRefHandleManager::GetReplicatedObjectData(FInternalNetRefIndex InternalIndex) const
diff --git a/Engine/Source/Runtime/Experimental/Iris/Core/Private/Iris/ReplicationSystem/ReplicationSystem.cpp b/Engine/Source/Runtime/Experimental/Iris/Core/Private/Iris/ReplicationSystem/ReplicationSystem.cpp
--- a/Engine/Source/Runtime/Experimental/Iris/Core/Private/Iris/ReplicationSystem/ReplicationSystem.cpp
+++ b/Engine/Source/Runtime/Experimental/Iris/Core/Private/Iris/ReplicationSystem/ReplicationSystem.cpp
@@ -147,6 +147,7 @@ public:
 			NetRefHandleManager.Init(NetRefHandleManagerInitParams);
 
 			NetRefHandleManager.GetOnMaxInternalNetRefIndexIncreasedDelegate().AddRaw(this, &FReplicationSystemImpl::OnMaxInternalNetRefIndexIncreased);
+			NetRefHandleManager.GetOnInternalNetRefIndicesFreedDelegate().AddRaw(this, &FReplicationSystemImpl::OnInternalNetRefIndicesFreed);
 		}
 
 		// Note that Params.MaxReplicatedObjectCount was just a suggestion for the NetRefHandleManager.
@@ -324,6 +325,7 @@ public:
 		}
 
 		ReplicationSystemInternal.GetNetRefHandleManager().GetOnMaxInternalNetRefIndexIncreasedDelegate().RemoveAll(this);
+		ReplicationSystemInternal.GetNetRefHandleManager().GetOnInternalNetRefIndicesFreedDelegate().RemoveAll(this);
 		ReplicationSystemInternal.GetNetRefHandleManager().Deinit();
 	}
 
@@ -339,6 +341,12 @@ public:
 		ReplicationSystemInternal.GetPrioritization().OnMaxInternalNetRefIndexIncreased(NewMaxInternalIndex);
 	}
 
+	void OnInternalNetRefIndicesFreed(const TConstArrayView<FInternalNetRefIndex>& FreedIndices)
+	{
+		ReplicationSystemInternal.GetFiltering().OnInternalNetRefIndicesFreed(FreedIndices);
+		ReplicationSystemInternal.GetConditionals().OnInternalNetRefIndicesFreed(FreedIndices);
+	}
+
 	void StartPreSendUpdate()
 	{
 		// Block unsupported operations until SendUpdate is finished
