Iris - Allow subobjects to be easily created without components. Optional description can now be passed to FReplicationSystemTestServer and FReplicationSystemTestClient UpdateAndSend methods.

diff --git a/Engine/Plugins/Runtime/ReplicationSystemTestPlugin/Source/Private/Tests/ReplicationSystem/ReplicationSystemServerClientTestFixture.cpp b/Engine/Plugins/Runtime/ReplicationSystemTestPlugin/Source/Private/Tests/ReplicationSystem/ReplicationSystemServerClientTestFixture.cpp
--- a/Engine/Plugins/Runtime/ReplicationSystemTestPlugin/Source/Private/Tests/ReplicationSystem/ReplicationSystemServerClientTestFixture.cpp
+++ b/Engine/Plugins/Runtime/ReplicationSystemTestPlugin/Source/Private/Tests/ReplicationSystem/ReplicationSystemServerClientTestFixture.cpp
@@ -443,13 +443,13 @@ FReplicationSystemTestClient::FReplicationSystemTestClient(const TCHAR* Name)
 {
 }
 
-bool FReplicationSystemTestClient::UpdateAndSend(FReplicationSystemTestServer* Server, bool bDeliver)
+bool FReplicationSystemTestClient::UpdateAndSend(FReplicationSystemTestServer* Server, bool bDeliver, const TCHAR* Desc)
 {
 	bool bSuccess = false;
 
 	PreSendUpdate();
 
-	if (SendUpdate())
+	if (SendUpdate(Desc))
 	{
 		constexpr uint32 ServerRemoteConnectionId = 0x01;
 		DeliverTo(*Server, LocalConnectionId, ServerRemoteConnectionId, bDeliver);
@@ -495,7 +495,7 @@ void FReplicationSystemTestServer::DeliverTo(FReplicationSystemTestClient* Clien
 	FReplicationSystemTestNode::DeliverTo(*Client, Client->ConnectionIdOnServer, Client->LocalConnectionId, bDeliver);
 }
 
-bool FReplicationSystemTestServer::UpdateAndSend(const TArrayView<FReplicationSystemTestClient*const>& Clients, bool bDeliver /*= true*/)
+bool FReplicationSystemTestServer::UpdateAndSend(const TArrayView<FReplicationSystemTestClient*const>& Clients, bool bDeliver /*= true*/, const TCHAR* Desc)
 {
 	bool bSuccess = true;
 
@@ -503,7 +503,7 @@ bool FReplicationSystemTestServer::UpdateAndSend(const TArrayView<FReplicationSy
 
 	for (FReplicationSystemTestClient* Client : Clients)
 	{
-		bSuccess &= SendAndDeliverTo(Client, bDeliver);
+		bSuccess &= SendAndDeliverTo(Client, bDeliver, Desc);
 	}
 	
 	PostSendUpdate();
diff --git a/Engine/Plugins/Runtime/ReplicationSystemTestPlugin/Source/Private/Tests/ReplicationSystem/ReplicationSystemServerClientTestFixture.h b/Engine/Plugins/Runtime/ReplicationSystemTestPlugin/Source/Private/Tests/ReplicationSystem/ReplicationSystemServerClientTestFixture.h
--- a/Engine/Plugins/Runtime/ReplicationSystemTestPlugin/Source/Private/Tests/ReplicationSystem/ReplicationSystemServerClientTestFixture.h
+++ b/Engine/Plugins/Runtime/ReplicationSystemTestPlugin/Source/Private/Tests/ReplicationSystem/ReplicationSystemServerClientTestFixture.h
@@ -137,7 +137,7 @@ public:
 	UTestReplicatedIrisObject* CreateObject(uint32 NumComponents, uint32 NumIrisComponents);
 	UTestReplicatedIrisObject* CreateSubObject(FNetRefHandle Owner, uint32 NumComponents, uint32 NumIrisComponents);
 	UTestReplicatedIrisObject* CreateObject(const UTestReplicatedIrisObject::FComponents& Components = UTestReplicatedIrisObject::FComponents());
-	UTestReplicatedIrisObject* CreateSubObject(FNetRefHandle Owner, const UTestReplicatedIrisObject::FComponents& Components);
+	UTestReplicatedIrisObject* CreateSubObject(FNetRefHandle Owner, const UTestReplicatedIrisObject::FComponents& Components = UTestReplicatedIrisObject::FComponents());
 	UTestReplicatedIrisObject* CreateObjectWithDynamicState(uint32 NumComponents, uint32 NumIrisComponents, uint32 NumDynamicStateComponents);
 
 	void DestroyObject(UReplicatedTestObject*, EEndReplicationFlags EndReplicationFlags = EEndReplicationFlags::Destroy);
@@ -196,7 +196,7 @@ public:
 	explicit FReplicationSystemTestClient(FReplicationSystemTestNode::EDelaySetup Delay) : FReplicationSystemTestNode(Delay) {}
 
 	// Tick and send packets to the server
-	bool UpdateAndSend(class FReplicationSystemTestServer* Server, bool bDeliver = true);
+	bool UpdateAndSend(class FReplicationSystemTestServer* Server, bool bDeliver = true, const TCHAR* Desc = nullptr);
 
 	uint32 ConnectionIdOnServer;
 	uint32 LocalConnectionId;
@@ -218,7 +218,7 @@ public:
 	void DeliverTo(FReplicationSystemTestClient* Client, bool bDeliver);
 
 	// Tick and send packets to one or all clients
-	bool UpdateAndSend(const TArrayView<FReplicationSystemTestClient*const>& Clients, bool bDeliver=true);
+	bool UpdateAndSend(const TArrayView<FReplicationSystemTestClient*const>& Clients, bool bDeliver=true, const TCHAR* Desc = nullptr);
 };
 
 class FReplicationSystemServerClientTestFixture : public FNetworkAutomationTestSuiteFixture
