Iris - Added FTestEnsureScope to allow tests to fail on ensures. The ensure expression and message is printed, but without the file and line information which isn't available.

diff --git a/Engine/Plugins/Runtime/ReplicationSystemTestPlugin/Source/Private/NetworkAutomationTestMacros.cpp b/Engine/Plugins/Runtime/ReplicationSystemTestPlugin/Source/Private/NetworkAutomationTestMacros.cpp
--- a/Engine/Plugins/Runtime/ReplicationSystemTestPlugin/Source/Private/NetworkAutomationTestMacros.cpp
+++ b/Engine/Plugins/Runtime/ReplicationSystemTestPlugin/Source/Private/NetworkAutomationTestMacros.cpp
@@ -1,14 +1,26 @@
 // Copyright Epic Games, Inc. All Rights Reserved.
 
 #include "NetworkAutomationTestMacros.h"
-#include "ReplicationSystemTestPlugin/NetworkAutomationTest.h"
 #include "NetworkAutomationTest.h"
+#include "Misc/AssertionMacros.h"
+#include "ReplicationSystemTestPlugin/NetworkAutomationTest.h"
 
 DEFINE_LOG_CATEGORY(LogNetworkAutomationTest);
 
 namespace UE::Net
 {
 
+FTestEnsureScope::FTestEnsureScope()
+: FEnsureScope(
+	[](const FEnsureHandlerArgs& Args)
+	{
+		UE_LOG(LogNetworkAutomationTest, Display, TEXT("Ensure condition failed: %hs\n%s\n"), Args.Expression, Args.Message);
+		return true;
+	}
+  ) 
+{
+}
+
 FTestResult::FTestResult()
 : TestResult(ETestResult_Success)
 {
diff --git a/Engine/Plugins/Runtime/ReplicationSystemTestPlugin/Source/Private/NetworkAutomationTestMacros.h b/Engine/Plugins/Runtime/ReplicationSystemTestPlugin/Source/Private/NetworkAutomationTestMacros.h
--- a/Engine/Plugins/Runtime/ReplicationSystemTestPlugin/Source/Private/NetworkAutomationTestMacros.h
+++ b/Engine/Plugins/Runtime/ReplicationSystemTestPlugin/Source/Private/NetworkAutomationTestMacros.h
@@ -2,10 +2,11 @@
 
 #pragma once
 
-#include "CoreTypes.h"
+#include "HAL/Platform.h"
 #include "HAL/PreprocessorHelpers.h"
 #include "TestMessage.h"
 #include "Logging/LogVerbosity.h"
+#include "Tests/EnsureScope.h"
 
 #include <catch2/catch_test_macros.hpp>
 
@@ -17,6 +18,12 @@ namespace UE::Net
 namespace UE::Net
 {
 
+class FTestEnsureScope : public FEnsureScope
+{
+public:
+	FTestEnsureScope();
+};
+
 enum ETestResult
 {
 	ETestResult_Success,
