Iris - Fix various ensures and behavior related to tear off and object scope changes. Removed support for non-atomic object hierarchy destroy.

diff --git a/Engine/Plugins/Runtime/ReplicationSystemTestPlugin/Source/Private/Tests/ReplicationSystem/TestFlushBeforeDestroy.cpp b/Engine/Plugins/Runtime/ReplicationSystemTestPlugin/Source/Private/Tests/ReplicationSystem/TestFlushBeforeDestroy.cpp
--- a/Engine/Plugins/Runtime/ReplicationSystemTestPlugin/Source/Private/Tests/ReplicationSystem/TestFlushBeforeDestroy.cpp
+++ b/Engine/Plugins/Runtime/ReplicationSystemTestPlugin/Source/Private/Tests/ReplicationSystem/TestFlushBeforeDestroy.cpp
@@ -1,11 +1,13 @@
 // Copyright Epic Games, Inc. All Rights Reserved.
 
 #include "ReplicationSystemServerClientTestFixture.h"
+#include "Iris/Core/IrisLog.h"
 #include "Iris/ReplicationSystem/ReplicationSystem.h"
 #include "Iris/ReplicationSystem/Filtering/NetObjectFilter.h"
+#include "Logging/LogScopedVerbosityOverride.h"
+#include "Misc/ScopeExit.h"
 #include "NetBlob/PartialNetBlobTestFixture.h"
 #include "NetBlob/MockNetBlob.h"
-#include "Misc/ScopeExit.h"
 
 namespace UE::Net::Private
 {
@@ -1775,4 +1777,470 @@ UE_NET_TEST_FIXTURE(FTestFlushBeforeDestroyFixture, TestTearOffOfOwnerFilteredOb
 	// Verify that the attachment was not received on the disallowed connection
 	UE_NET_ASSERT_EQ(ClientMockNetObjectAttachmentHandlers[1]->GetFunctionCallCounts().OnNetBlobReceived, 0U);
 }
+
+UE_NET_TEST_FIXTURE(FTestFlushBeforeDestroyFixture, TestDelayedTearOffAndSubobjectDestroy)
+{
+	// Add a client
+	FReplicationSystemTestClient* Client = CreateClient();
+
+	// Spawn object on server
+	UTestReplicatedIrisObject* ServerObject = Server->CreateObject();
+	UTestReplicatedIrisObject* ServerSubObject0 = Server->CreateSubObject(ServerObject->NetRefHandle);
+	UTestReplicatedIrisObject* ServerSubObject1 = Server->CreateSubObject(ServerObject->NetRefHandle);
+
+	// Send and deliver packet
+	Server->UpdateAndSend({Client}, DeliverPacket, TEXT("Create objects"));
+
+	Server->ReplicationSystem->TearOffNextUpdate(ServerObject->NetRefHandle);
+
+	Server->UpdateAndSend({Client}, DoNotDeliverPacket, TEXT("Tear off"));
+
+	Server->DestroyObject(ServerSubObject1);
+
+	Server->UpdateAndSend({Client}, DoNotDeliverPacket, TEXT("Tear off and destroy"));
+}
+
+UE_NET_TEST_FIXTURE(FTestFlushBeforeDestroyFixture, TestTearOffAndFilterOutInSameFrameWhileSubobjectHasReliableAttachment)
+{
+	// Uncomment the verbosiy override scopes and enable UE_NET_ENABLE_REPLICATIONWRITER_LOG to track down issues if test fails.
+	//LOG_SCOPE_VERBOSITY_OVERRIDE(LogIris, ELogVerbosity::Verbose);
+	//LOG_SCOPE_VERBOSITY_OVERRIDE(LogIrisBridge, ELogVerbosity::Verbose);
+
+	FTestEnsureScope EnsureScope;
+
+	// Add a client
+	FReplicationSystemTestClient* Client = CreateClient();
+	RegisterNetBlobHandlers(Client);
+
+	// Spawn object on server
+	UTestReplicatedIrisObject* ServerObject = Server->CreateObject();
+	UTestReplicatedIrisObject* ServerSubObject0 = Server->CreateSubObject(ServerObject->NetRefHandle);
+
+	// Send and deliver packet
+	Server->UpdateAndSend({Client}, DeliverPacket, TEXT("Create objects"));
+
+	UTestReplicatedIrisObject* ClientObject = Cast<UTestReplicatedIrisObject>(Client->GetReplicationBridge()->GetReplicatedObject(ServerObject->NetRefHandle));
+	UTestReplicatedIrisObject* ClientSubObject0 = Cast<UTestReplicatedIrisObject>(Client->GetReplicationBridge()->GetReplicatedObject(ServerSubObject0->NetRefHandle));
+
+	UTestReplicatedIrisObject* ServerSubObject1 = Server->CreateSubObject(ServerObject->NetRefHandle);
+	// Create reliable attachment to ServerSubObject1
+	{
+		constexpr uint32 PayloadBitCount = 128U;
+		const TRefCountPtr<FNetObjectAttachment>& Attachment = MockNetObjectAttachmentHandler->CreateReliableNetObjectAttachment(PayloadBitCount);
+		FNetObjectReference AttachmentTarget = FObjectReferenceCache::MakeNetObjectReference(ServerSubObject1->NetRefHandle);
+		Server->GetReplicationSystem()->QueueNetObjectAttachment(Client->ConnectionIdOnServer, AttachmentTarget, Attachment);
+	}
+
+	Server->DestroyObject(ServerSubObject1);
+
+	// Put data in flight
+	Server->PreSendUpdate();
+	Server->SendTo(Client);
+	Server->PostSendUpdate();
+
+	// Filter out object
+	FNetObjectGroupHandle ExclusionGroupHandle = Server->ReplicationSystem->CreateGroup(NAME_None);
+	Server->ReplicationSystem->AddToGroup(ExclusionGroupHandle, ServerObject->NetRefHandle);
+	Server->ReplicationSystem->AddExclusionFilterGroup(ExclusionGroupHandle);
+
+	Server->PreSendUpdate();
+	Server->SendTo(Client);
+	Server->PostSendUpdate();
+	UE_NET_ASSERT_EQ(EnsureScope.GetCount(), 0);
+
+	ServerObject->IntA++;
+	ServerSubObject0->IntA++;
+
+	// Filter in object
+	Server->ReplicationSystem->RemoveFromGroup(ExclusionGroupHandle, ServerObject->NetRefHandle);
+
+	// Tear off 
+	Server->ReplicationSystem->TearOffNextUpdate(ServerObject->NetRefHandle);
+
+	Server->PreSendUpdate();
+	Server->SendTo(Client);
+	Server->PostSendUpdate();
+	UE_NET_ASSERT_EQ(EnsureScope.GetCount(), 0);
+
+	// Deliver remaining packets
+	{
+		const SIZE_T PacketCount = Server->GetConnectionInfo(Client->ConnectionIdOnServer).WrittenPackets.Count();
+		for (SIZE_T PacketIt = 0; PacketIt != PacketCount; ++PacketIt)
+		{
+			Server->DeliverTo(Client, DeliverPacket);
+		}
+	}
+
+	// Send update
+	Server->UpdateAndSend({Client});
+
+	// Verify all objects are unresolvable on the client
+	UE_NET_ASSERT_FALSE(Client->IsResolvableNetRefHandle(ServerObject->NetRefHandle));
+	UE_NET_ASSERT_FALSE(Client->IsResolvableNetRefHandle(ServerSubObject0->NetRefHandle));
+	UE_NET_ASSERT_FALSE(Client->IsResolvableNetRefHandle(ServerSubObject1->NetRefHandle));
+
+	// Verify we got the latest state and received the reliable attachment
+	UE_NET_ASSERT_EQ(ClientObject->IntA, ServerObject->IntA);
+	UE_NET_ASSERT_EQ(ClientSubObject0->IntA, ServerSubObject0->IntA);
+	UE_NET_ASSERT_EQ(ClientMockNetObjectAttachmentHandler->GetFunctionCallCounts().OnNetBlobReceived, 1U);
+}
+
+// This test is almost like the one above but with some minor differences to updates.
+UE_NET_TEST_FIXTURE(FTestFlushBeforeDestroyFixture, TestTearOffAndFilterOutInSameFrameWhileSubobjectHasReliableAttachment2)
+{
+	// Uncomment the verbosiy override scopes and enable UE_NET_ENABLE_REPLICATIONWRITER_LOG to track down issues if test fails.
+	//LOG_SCOPE_VERBOSITY_OVERRIDE(LogIris, ELogVerbosity::Verbose);
+	//LOG_SCOPE_VERBOSITY_OVERRIDE(LogIrisBridge, ELogVerbosity::Verbose);
+
+	FTestEnsureScope EnsureScope;
+
+	// Add a client
+	FReplicationSystemTestClient* Client = CreateClient();
+	RegisterNetBlobHandlers(Client);
+
+	// Spawn object on server
+	UTestReplicatedIrisObject* ServerObject = Server->CreateObject();
+	UTestReplicatedIrisObject* ServerSubObject0 = Server->CreateSubObject(ServerObject->NetRefHandle);
+
+	// Send and deliver packet
+	Server->UpdateAndSend({Client}, DeliverPacket, TEXT("Create objects"));
+
+	UTestReplicatedIrisObject* ClientObject = Cast<UTestReplicatedIrisObject>(Client->GetReplicationBridge()->GetReplicatedObject(ServerObject->NetRefHandle));
+	UTestReplicatedIrisObject* ClientSubObject0 = Cast<UTestReplicatedIrisObject>(Client->GetReplicationBridge()->GetReplicatedObject(ServerSubObject0->NetRefHandle));
+
+	UTestReplicatedIrisObject* ServerSubObject1 = Server->CreateSubObject(ServerObject->NetRefHandle);
+	// Create reliable attachment to ServerSubObject1
+	{
+		constexpr uint32 PayloadBitCount = 128U;
+		const TRefCountPtr<FNetObjectAttachment>& Attachment = MockNetObjectAttachmentHandler->CreateReliableNetObjectAttachment(PayloadBitCount);
+		FNetObjectReference AttachmentTarget = FObjectReferenceCache::MakeNetObjectReference(ServerSubObject1->NetRefHandle);
+		Server->GetReplicationSystem()->QueueNetObjectAttachment(Client->ConnectionIdOnServer, AttachmentTarget, Attachment);
+	}
+
+	Server->DestroyObject(ServerSubObject1);
+
+	// Put data in flight
+	Server->PreSendUpdate();
+	Server->SendTo(Client);
+	Server->PostSendUpdate();
+
+	// Filter out object
+	FNetObjectGroupHandle ExclusionGroupHandle = Server->ReplicationSystem->CreateGroup(NAME_None);
+	Server->ReplicationSystem->AddToGroup(ExclusionGroupHandle, ServerObject->NetRefHandle);
+	Server->ReplicationSystem->AddExclusionFilterGroup(ExclusionGroupHandle);
+
+	Server->PreSendUpdate();
+	Server->SendTo(Client);
+	Server->PostSendUpdate();
+	UE_NET_ASSERT_EQ(EnsureScope.GetCount(), 0);
+
+	ServerObject->IntA++;
+	ServerSubObject0->IntA++;
+
+	Server->PreSendUpdate();
+	Server->SendTo(Client);
+	Server->PostSendUpdate();
+	UE_NET_ASSERT_EQ(EnsureScope.GetCount(), 0);
+
+	Server->DeliverTo(Client, DeliverPacket);
+
+	// Filter in object
+	Server->ReplicationSystem->RemoveFromGroup(ExclusionGroupHandle, ServerObject->NetRefHandle);
+
+	// Tear off 
+	Server->ReplicationSystem->TearOffNextUpdate(ServerObject->NetRefHandle);
+
+	Server->PreSendUpdate();
+	Server->SendTo(Client);
+	Server->PostSendUpdate();
+	UE_NET_ASSERT_EQ(EnsureScope.GetCount(), 0);
+
+	// Deliver remaining packets
+	{
+		const SIZE_T PacketCount = Server->GetConnectionInfo(Client->ConnectionIdOnServer).WrittenPackets.Count();
+		for (SIZE_T PacketIt = 0; PacketIt != PacketCount; ++PacketIt)
+		{
+			Server->DeliverTo(Client, DeliverPacket);
+		}
+	}
+
+	// Send update
+	Server->UpdateAndSend({Client});
+
+	// Send update
+	Server->UpdateAndSend({Client});
+
+	// Verify all objects are unresolvable on the client
+	UE_NET_ASSERT_FALSE(Client->IsResolvableNetRefHandle(ServerObject->NetRefHandle));
+	UE_NET_ASSERT_FALSE(Client->IsResolvableNetRefHandle(ServerSubObject0->NetRefHandle));
+	UE_NET_ASSERT_FALSE(Client->IsResolvableNetRefHandle(ServerSubObject1->NetRefHandle));
+
+	// Verify we got the latest state and received the reliable attachment
+	UE_NET_ASSERT_EQ(ClientObject->IntA, ServerObject->IntA);
+	UE_NET_ASSERT_EQ(ClientSubObject0->IntA, ServerSubObject0->IntA);
+	UE_NET_ASSERT_EQ(ClientMockNetObjectAttachmentHandler->GetFunctionCallCounts().OnNetBlobReceived, 1U);
+}
+
+UE_NET_TEST_FIXTURE(FTestFlushBeforeDestroyFixture, TestTearOffOfObjectNotInScopeWithSubObjectPendingDestroy)
+{
+	// Uncomment the verbosiy override scopes and enable UE_NET_ENABLE_REPLICATIONWRITER_LOG to track down issues if test fails.
+	//LOG_SCOPE_VERBOSITY_OVERRIDE(LogIris, ELogVerbosity::Verbose);
+	//LOG_SCOPE_VERBOSITY_OVERRIDE(LogIrisBridge, ELogVerbosity::Verbose);
+
+	FTestEnsureScope EnsureScope;
+
+	// Add a client
+	FReplicationSystemTestClient* Client = CreateClient();
+	RegisterNetBlobHandlers(Client);
+
+	// Spawn object on server
+	UTestReplicatedIrisObject* ServerObject = Server->CreateObject();
+	UTestReplicatedIrisObject* ServerSubObject0 = Server->CreateSubObject(ServerObject->NetRefHandle);
+
+	// Send and deliver packet
+	Server->UpdateAndSend({Client}, DeliverPacket, TEXT("Create objects"));
+
+	UTestReplicatedIrisObject* ClientObject = Cast<UTestReplicatedIrisObject>(Client->GetReplicationBridge()->GetReplicatedObject(ServerObject->NetRefHandle));
+	UTestReplicatedIrisObject* ClientSubObject0 = Cast<UTestReplicatedIrisObject>(Client->GetReplicationBridge()->GetReplicatedObject(ServerSubObject0->NetRefHandle));
+
+	UTestReplicatedIrisObject* ServerSubObject1 = Server->CreateSubObject(ServerObject->NetRefHandle);
+	// Create reliable attachment to ServerSubObject1
+	{
+		constexpr uint32 PayloadBitCount = 128U;
+		const TRefCountPtr<FNetObjectAttachment>& Attachment = MockNetObjectAttachmentHandler->CreateReliableNetObjectAttachment(PayloadBitCount);
+		FNetObjectReference AttachmentTarget = FObjectReferenceCache::MakeNetObjectReference(ServerSubObject1->NetRefHandle);
+		Server->GetReplicationSystem()->QueueNetObjectAttachment(Client->ConnectionIdOnServer, AttachmentTarget, Attachment);
+	}
+
+	Server->DestroyObject(ServerSubObject1);
+
+	ServerObject->IntA++;
+	ServerSubObject0->IntA++;
+
+	// Put data in flight
+	Server->PreSendUpdate();
+	Server->SendTo(Client, TEXT("SubObject1 creation"));
+	Server->PostSendUpdate();
+
+	// Filter out object. This will cause root to end up in WaitOnFlush
+	FNetObjectGroupHandle ExclusionGroupHandle = Server->ReplicationSystem->CreateGroup(NAME_None);
+	Server->ReplicationSystem->AddToGroup(ExclusionGroupHandle, ServerObject->NetRefHandle);
+	Server->ReplicationSystem->AddExclusionFilterGroup(ExclusionGroupHandle);
+
+	Server->PreSendUpdate();
+	Server->SendTo(Client);
+	Server->PostSendUpdate();
+
+	UE_NET_ASSERT_EQ(EnsureScope.GetCount(), 0);
+
+	// Do not deliver SubObject1 creation
+	Server->DeliverTo(Client, DoNotDeliverPacket);
+
+	Server->PreSendUpdate();
+	Server->SendTo(Client);
+	Server->PostSendUpdate();
+
+	// At this point we're expecting Root in WaitOnFlush, SubObject0 in PendingDestroy and SubObject1 in WaitOnCreateConfirmation.
+	
+	// Tear off 
+	Server->ReplicationSystem->TearOffNextUpdate(ServerObject->NetRefHandle);
+
+	Server->PreSendUpdate();
+	Server->SendTo(Client);
+	Server->PostSendUpdate();
+	UE_NET_ASSERT_EQ(EnsureScope.GetCount(), 0);
+
+	// Do not deliver remaining packets
+	{
+		const SIZE_T PacketCount = Server->GetConnectionInfo(Client->ConnectionIdOnServer).WrittenPackets.Count();
+		for (SIZE_T PacketIt = 0; PacketIt != PacketCount; ++PacketIt)
+		{
+			Server->DeliverTo(Client, DoNotDeliverPacket);
+		}
+	}
+
+	// Send update
+	Server->UpdateAndSend({Client});
+
+	// Send update
+	Server->UpdateAndSend({Client});
+
+	// Verify all objects are unresolvable on the client
+	UE_NET_ASSERT_FALSE(Client->IsResolvableNetRefHandle(ServerObject->NetRefHandle));
+	UE_NET_ASSERT_FALSE(Client->IsResolvableNetRefHandle(ServerSubObject0->NetRefHandle));
+	UE_NET_ASSERT_FALSE(Client->IsResolvableNetRefHandle(ServerSubObject1->NetRefHandle));
+
+	// Once an object starts getting destroyed it will clear dirtiness. A subsequent tear off will not automatically get that latest state over. 
+	UE_NET_ASSERT_EQ(ClientObject->IntA, ServerObject->IntA);
+	UE_NET_ASSERT_EQ(ClientSubObject0->IntA, ServerSubObject0->IntA);
+
+	// Verify we got the latest state and received the reliable attachment
+	UE_NET_ASSERT_EQ(ClientMockNetObjectAttachmentHandler->GetFunctionCallCounts().OnNetBlobReceived, 1U);
+}
+
+UE_NET_TEST_FIXTURE(FTestFlushBeforeDestroyFixture, TestTearOffOfObjectNotInScopeWithSubObjectPendingDestroy2)
+{
+	// Uncomment the verbosiy override scopes and enable UE_NET_ENABLE_REPLICATIONWRITER_LOG to track down issues if test fails.
+	//LOG_SCOPE_VERBOSITY_OVERRIDE(LogIris, ELogVerbosity::Verbose);
+	//LOG_SCOPE_VERBOSITY_OVERRIDE(LogIrisBridge, ELogVerbosity::Verbose);
+
+	FTestEnsureScope EnsureScope;
+
+	// Add a client
+	FReplicationSystemTestClient* Client = CreateClient();
+	RegisterNetBlobHandlers(Client);
+
+	// Spawn object on server
+	UTestReplicatedIrisObject* ServerObject = Server->CreateObject();
+	UTestReplicatedIrisObject* ServerSubObject0 = Server->CreateSubObject(ServerObject->NetRefHandle);
+
+	// Send and deliver packet
+	Server->UpdateAndSend({Client}, DeliverPacket, TEXT("Create objects"));
+
+	UTestReplicatedIrisObject* ClientObject = Cast<UTestReplicatedIrisObject>(Client->GetReplicationBridge()->GetReplicatedObject(ServerObject->NetRefHandle));
+	UTestReplicatedIrisObject* ClientSubObject0 = Cast<UTestReplicatedIrisObject>(Client->GetReplicationBridge()->GetReplicatedObject(ServerSubObject0->NetRefHandle));
+
+	UTestReplicatedIrisObject* ServerSubObject1 = Server->CreateSubObject(ServerObject->NetRefHandle);
+	// Create reliable attachment to ServerSubObject1
+	{
+		constexpr uint32 PayloadBitCount = 128U;
+		const TRefCountPtr<FNetObjectAttachment>& Attachment = MockNetObjectAttachmentHandler->CreateReliableNetObjectAttachment(PayloadBitCount);
+		FNetObjectReference AttachmentTarget = FObjectReferenceCache::MakeNetObjectReference(ServerSubObject1->NetRefHandle);
+		Server->GetReplicationSystem()->QueueNetObjectAttachment(Client->ConnectionIdOnServer, AttachmentTarget, Attachment);
+	}
+
+	Server->DestroyObject(ServerSubObject1);
+
+	ServerObject->IntA++;
+	ServerSubObject0->IntA++;
+
+	// Put data in flight
+	Server->PreSendUpdate();
+	Server->SendTo(Client, TEXT("SubObject1 creation"));
+	Server->PostSendUpdate();
+
+	// Filter out object. This will cause root to end up in WaitOnFlush
+	FNetObjectGroupHandle ExclusionGroupHandle = Server->ReplicationSystem->CreateGroup(NAME_None);
+	Server->ReplicationSystem->AddToGroup(ExclusionGroupHandle, ServerObject->NetRefHandle);
+	Server->ReplicationSystem->AddExclusionFilterGroup(ExclusionGroupHandle);
+
+	Server->PreSendUpdate();
+	Server->SendTo(Client);
+	Server->PostSendUpdate();
+
+	UE_NET_ASSERT_EQ(EnsureScope.GetCount(), 0);
+
+	// Do not deliver SubObject1 creation
+	Server->DeliverTo(Client, DoNotDeliverPacket);
+
+	Server->PreSendUpdate();
+	Server->SendTo(Client);
+	Server->PostSendUpdate();
+
+	// Deliver SubObject1 creation
+	Server->DeliverTo(Client, DeliverPacket);
+
+	Server->PreSendUpdate();
+	Server->SendTo(Client);
+	Server->PostSendUpdate();
+
+	// At this stage we're expecting all objects to be in WaitOnDestroyConfirmation
+	
+	// Tear off 
+	Server->ReplicationSystem->TearOffNextUpdate(ServerObject->NetRefHandle);
+
+	Server->PreSendUpdate();
+	Server->SendTo(Client);
+	Server->PostSendUpdate();
+	UE_NET_ASSERT_EQ(EnsureScope.GetCount(), 0);
+
+	// Do not deliver remaining packets
+	{
+		const SIZE_T PacketCount = Server->GetConnectionInfo(Client->ConnectionIdOnServer).WrittenPackets.Count();
+		for (SIZE_T PacketIt = 0; PacketIt != PacketCount; ++PacketIt)
+		{
+			Server->DeliverTo(Client, DoNotDeliverPacket);
+		}
+	}
+
+	// Send update
+	Server->UpdateAndSend({Client});
+
+	// Send update
+	Server->UpdateAndSend({Client});
+
+	// Verify all objects are unresolvable on the client
+	UE_NET_ASSERT_FALSE(Client->IsResolvableNetRefHandle(ServerObject->NetRefHandle));
+	UE_NET_ASSERT_FALSE(Client->IsResolvableNetRefHandle(ServerSubObject0->NetRefHandle));
+	UE_NET_ASSERT_FALSE(Client->IsResolvableNetRefHandle(ServerSubObject1->NetRefHandle));
+
+	// Verify we received the reliable attachment
+	UE_NET_ASSERT_EQ(ClientMockNetObjectAttachmentHandler->GetFunctionCallCounts().OnNetBlobReceived, 1U);
+}
+
+UE_NET_TEST_FIXTURE(FTestFlushBeforeDestroyFixture, TestSubObjectIsNotDestroyedBeforeRoot)
+{
+	// Uncomment the verbosiy override scopes and enable UE_NET_ENABLE_REPLICATIONWRITER_LOG to track down issues if test fails.
+	//LOG_SCOPE_VERBOSITY_OVERRIDE(LogIris, ELogVerbosity::Verbose);
+	//LOG_SCOPE_VERBOSITY_OVERRIDE(LogIrisBridge, ELogVerbosity::Verbose);
+
+	FTestEnsureScope EnsureScope;
+
+	// Add a client
+	FReplicationSystemTestClient* Client = CreateClient();
+	RegisterNetBlobHandlers(Client);
+
+	// Spawn object on server
+	UTestReplicatedIrisObject* ServerObject = Server->CreateObject();
+	UTestReplicatedIrisObject* ServerSubObject0 = Server->CreateSubObject(ServerObject->NetRefHandle);
+
+	// Send and deliver packet
+	Server->UpdateAndSend({Client}, DeliverPacket, TEXT("Create objects"));
+
+	UTestReplicatedIrisObject* ClientObject = Cast<UTestReplicatedIrisObject>(Client->GetReplicationBridge()->GetReplicatedObject(ServerObject->NetRefHandle));
+	UTestReplicatedIrisObject* ClientSubObject0 = Cast<UTestReplicatedIrisObject>(Client->GetReplicationBridge()->GetReplicatedObject(ServerSubObject0->NetRefHandle));
+
+	UTestReplicatedIrisObject* ServerSubObject1 = Server->CreateSubObject(ServerObject->NetRefHandle);
+	// Create reliable attachment to ServerSubObject1
+	{
+		// Create a huge attachment such that it takes more than a couple of frames to deliver
+		constexpr uint32 PayloadBitCount = 4000U*8U;
+		const TRefCountPtr<FNetObjectAttachment>& Attachment = MockNetObjectAttachmentHandler->CreateReliableNetObjectAttachment(PayloadBitCount);
+		FNetObjectReference AttachmentTarget = FObjectReferenceCache::MakeNetObjectReference(ServerSubObject1->NetRefHandle);
+		Server->GetReplicationSystem()->QueueNetObjectAttachment(Client->ConnectionIdOnServer, AttachmentTarget, Attachment);
+	}
+
+	Server->DestroyObject(ServerSubObject1);
+
+	// Put data in flight
+	Server->PreSendUpdate();
+	Server->SendTo(Client, TEXT("SubObject1 creation"));
+	Server->PostSendUpdate();
+
+	// Filter out object. This will cause root to end up in WaitOnFlush
+	FNetObjectGroupHandle ExclusionGroupHandle = Server->ReplicationSystem->CreateGroup(NAME_None);
+	Server->ReplicationSystem->AddToGroup(ExclusionGroupHandle, ServerObject->NetRefHandle);
+	Server->ReplicationSystem->AddExclusionFilterGroup(ExclusionGroupHandle);
+
+	Server->PreSendUpdate();
+	Server->SendTo(Client);
+	Server->PostSendUpdate();
+
+	UE_NET_EXPECT_EQ(EnsureScope.GetCount(), 0);
+
+	// Do not deliver SubObject1 creation
+	Server->DeliverTo(Client, DoNotDeliverPacket);
+
+	Server->PreSendUpdate();
+	Server->SendTo(Client);
+	Server->PostSendUpdate();
+
+	// Update and make sure all objects are created
+	Server->DeliverTo(Client, DeliverPacket);
+
+	UE_NET_ASSERT_TRUE(Client->IsResolvableNetRefHandle(ServerObject->NetRefHandle));
+	UE_NET_ASSERT_TRUE(Client->IsResolvableNetRefHandle(ServerSubObject0->NetRefHandle));
+	UE_NET_ASSERT_TRUE(Client->IsResolvableNetRefHandle(ServerSubObject1->NetRefHandle));
+}
+
 }
diff --git a/Engine/Source/Runtime/Experimental/Iris/Core/Private/Iris/ReplicationSystem/ReplicationReader.cpp b/Engine/Source/Runtime/Experimental/Iris/Core/Private/Iris/ReplicationSystem/ReplicationReader.cpp
--- a/Engine/Source/Runtime/Experimental/Iris/Core/Private/Iris/ReplicationSystem/ReplicationReader.cpp
+++ b/Engine/Source/Runtime/Experimental/Iris/Core/Private/Iris/ReplicationSystem/ReplicationReader.cpp
@@ -394,16 +394,7 @@ void FReplicationReader::ReadStreamDebugFeatures(FNetSerializationContext& Conte
 
 uint32 FReplicationReader::ReadObjectsPendingDestroy(FNetSerializationContext& Context)
 {
-	FNetBitStreamReader& Reader = *Context.GetBitStreamReader();
-	const bool bReadbjectsAndSubObjectsIndividually = Reader.ReadBool();
-	if (bReadbjectsAndSubObjectsIndividually)
-	{
-		return ReadObjectsAndSubObjectsPendingDestroy(Context);
-	}
-	else
-	{
-		return ReadRootObjectsPendingDestroy(Context);
-	}
+	return ReadRootObjectsPendingDestroy(Context);
 }
 
 uint32 FReplicationReader::ReadObjectsAndSubObjectsPendingDestroy(FNetSerializationContext& Context)
diff --git a/Engine/Source/Runtime/Experimental/Iris/Core/Private/Iris/ReplicationSystem/ReplicationWriter.cpp b/Engine/Source/Runtime/Experimental/Iris/Core/Private/Iris/ReplicationSystem/ReplicationWriter.cpp
--- a/Engine/Source/Runtime/Experimental/Iris/Core/Private/Iris/ReplicationSystem/ReplicationWriter.cpp
+++ b/Engine/Source/Runtime/Experimental/Iris/Core/Private/Iris/ReplicationSystem/ReplicationWriter.cpp
@@ -81,10 +81,6 @@ static FAutoConsoleVariableRef CVarGReplicationWriterReplicationRecordStarvation
 static bool bValidateObjectsWithDirtyChanges = true;
 static FAutoConsoleVariableRef CvarValidateObjectsWithDirtyChanges(TEXT("net.Iris.ReplicationWriter.ValidateObjectsWithDirtyChanges"), bValidateObjectsWithDirtyChanges, TEXT("Ensure that we don't try to mark invalid objects as dirty when they shouldn't."));
 
-/* NOTE: This cvar is just temporary for testing purposes. Eventually only root objects will replicate destroy and the cvar will be removed. */
-static bool bDestroyRootAndSubObjectsIndividually = false;
-static FAutoConsoleVariableRef CvarDestroyRootAndSubObjectsIndividually(TEXT("net.Iris.DestroyRootAndSubObjectsIndividually"), bDestroyRootAndSubObjectsIndividually, TEXT("TEMPORARY. Whether to send destroy on both root and subobjects (Iris legacy behavior) or just root objects and have the receiving end destroy the subobjects with it. Default is false, only sending destroy for root objects."));
-
 // net.Iris.ReplicationWriter.WriteBatchSizePerObject
 static bool bDebugBatchSizePerObjectEnabled = false;
 static FAutoConsoleVariableRef CvarWriteBatchSizePerObject(TEXT("net.Iris.ReplicationWriter.WriteBatchSizePerObject"), bDebugBatchSizePerObjectEnabled, TEXT("Write batch size per object. Helps tracking down bitstream errors. Requires code to be compiled with UE_NET_REPLICATIONDATASTREAM_DEBUG to be enabled."));
@@ -680,7 +676,7 @@ void FReplicationWriter::SetPendingDestroyOrSubObjectPendingDestroyState(uint32
 			// If owner is not pending destroy we mark the state of the SubObject to SubObjectPendingDestroy and mark owner as having dirty subobjects which will 
 			// destroy the subobject using the replicated state path of the owner.
 			FReplicationInfo& OwnerInfo = ReplicatedObjects[ObjectData.SubObjectRootIndex];
-			if ((OwnerInfo.GetState() != EReplicatedObjectState::Invalid) && (OwnerInfo.GetState() < EReplicatedObjectState::PendingDestroy))
+			if ((OwnerInfo.GetState() != EReplicatedObjectState::Invalid) && !ObjectsPendingDestroy.GetBit(ObjectData.SubObjectRootIndex))
 			{
 				MarkObjectDirty(ObjectData.SubObjectRootIndex, "SetPendingDestroyOrSubObjectPendingDestroyState");
 				OwnerInfo.HasDirtySubObjects = 1U;
@@ -689,6 +685,7 @@ void FReplicationWriter::SetPendingDestroyOrSubObjectPendingDestroyState(uint32
 				MarkObjectDirty(InternalIndex, "SetPendingDestroyOrSubObjectPendingDestroyState2");
 				ObjectsPendingDestroy.SetBit(InternalIndex);
 				Info.SubObjectPendingDestroy = 1U;
+				ensure(!Info.TearOff || (GetFlushStatus(InternalIndex, Info, Info.FlushFlags) == FlushFlags_None));
 				return;
 			}
 		}
@@ -730,18 +727,15 @@ void FReplicationWriter::UpdateScope(const FNetBitArrayView& UpdatedScope)
 		}
 		else if (State == EReplicatedObjectState::WaitOnFlush)
 		{			
-			if (ensureMsgf(!Info.TearOff, TEXT("We cannot cancel flush for an object pending tearoff ( InternalIndex: %u )"), Index))
-			{
-				// If we are waiting on flush but are re-added to scope we reset flush flags to default.
-				ObjectsPendingDestroy.ClearBit(Index);
-				Info.FlushFlags = GetDefaultFlushFlags();
-				SetState(Index, EReplicatedObjectState::Created);
-
-				// If we have accumulated changes while WaitingOnFlush, we should send them now
-				Info.SubObjectPendingDestroy = 0U;
-				Info.HasDirtyChangeMask |= FNetBitArrayView(Info.GetChangeMaskStoragePointer(), Info.ChangeMaskBitCount).IsAnyBitSet();
-				ObjectsWithDirtyChanges.SetBitValue(Index, Info.HasDirtyChangeMask);
-			}
+			// If we are waiting on flush but are re-added to scope we reset flush flags to default.
+			ObjectsPendingDestroy.ClearBit(Index);
+			Info.FlushFlags = GetDefaultFlushFlags();
+			SetState(Index, EReplicatedObjectState::Created);
+
+			// If we have accumulated changes while WaitingOnFlush, we should send them now
+			Info.SubObjectPendingDestroy = 0U;
+			Info.HasDirtyChangeMask |= FNetBitArrayView(Info.GetChangeMaskStoragePointer(), Info.ChangeMaskBitCount).IsAnyBitSet();
+			ObjectsWithDirtyChanges.SetBitValue(Index, Info.HasDirtyChangeMask);
 		}
 		else if (State == EReplicatedObjectState::WaitOnCreateConfirmation)
 		{
@@ -957,7 +951,7 @@ void FReplicationWriter::InternalUpdateDirtyChangeMasks(const FChangeMaskCache&
 	const uint32 MarkForTearOff = bMarkForTearOff ? 1U : 0U;
 	const ChangeMaskStorageType* StoragePtr = CachedChangeMasks.Storage.GetData();
 
-	for (const auto& Entry : CachedChangeMasks.Indices)
+	for (const FChangeMaskCache::FCachedInfo& Entry : CachedChangeMasks.Indices)
 	{
 		FReplicationInfo& Info = ReplicatedObjects[Entry.InternalIndex];
 		if (Info.GetState() == EReplicatedObjectState::Invalid)
@@ -973,8 +967,20 @@ void FReplicationWriter::InternalUpdateDirtyChangeMasks(const FChangeMaskCache&
 		}
 
 		if (Entry.bMarkSubObjectOwnerDirty == 0U)
-		{		
+		{
+			ensure((!MarkForTearOff && !Info.TearOff) || !Info.SubObjectPendingDestroy);
 			// Mark object for TearOff, that is that we will stop replication as soon as the tear-off is acknowledged
+			if (MarkForTearOff)
+			{
+				const EReplicatedObjectState CurrentState = Info.GetState();
+				// If an object was recently removed from scope for a connection then many objects can be in PendingDestroy. We need to move to WaitOnFlush to progress with tear off.
+				if (CurrentState == EReplicatedObjectState::PendingDestroy)
+				{
+					SetState(Entry.InternalIndex, EReplicatedObjectState::WaitOnFlush);
+					Info.FlushFlags = GetFlushStatus(Entry.InternalIndex, Info, static_cast<EFlushFlags>(Info.FlushFlags | ExtraFlushFlags));
+				}
+			}
+
 			Info.TearOff |= MarkForTearOff;
 
 			// Update flush flags
@@ -994,6 +1000,27 @@ void FReplicationWriter::InternalUpdateDirtyChangeMasks(const FChangeMaskCache&
 				// Mark changemask as dirty
 				Info.HasDirtyChangeMask = bMarkScopedObjectDirty ? 1U : 0U;
 			}
+
+			if (bMarkForTearOff && !bMarkScopedObjectDirty)
+			{
+				if (Info.GetState() <= EReplicatedObjectState::PendingTearOff)
+				{
+					MarkObjectDirty(Entry.InternalIndex, "TearOffObjectNotInScope");
+					Info.HasDirtyChangeMask = FNetBitArrayView(Info.GetChangeMaskStoragePointer(), Info.ChangeMaskBitCount).IsAnyBitSet();
+
+					// If we're a subobject and forced dirtied ourself we want the root object to be dirty as well
+					if (Info.IsSubObject)
+					{
+						const FInternalNetRefIndex OwnerIndex = NetRefHandleManager->GetRootObjectInternalIndexOfSubObject(Entry.InternalIndex);
+						FReplicationInfo& OwnerInfo = ReplicatedObjects[OwnerIndex];
+						if (OwnerInfo.GetState() != EReplicatedObjectState::Invalid && OwnerInfo.GetState() < EReplicatedObjectState::PendingDestroy)
+						{
+							MarkObjectDirty(OwnerIndex, "TearOffObjectNotInScopeRoot");
+							OwnerInfo.HasDirtySubObjects = 1U;
+						}
+					}
+				}
+			}
 		}
 		else
 		{
@@ -1836,178 +1863,7 @@ void FReplicationWriter::WriteStreamDebugFeatures(FNetSerializationContext& Cont
 
 uint32 FReplicationWriter::WriteObjectsPendingDestroy(FNetSerializationContext& Context)
 {
-	FNetBitStreamWriter& Writer = *Context.GetBitStreamWriter();
-	const bool bWriteObjectsAndSubObjectsIndividually = Writer.WriteBool(bDestroyRootAndSubObjectsIndividually);
-	if (Writer.IsOverflown())
-	{
-		return 0;
-	}
-
-	if (bWriteObjectsAndSubObjectsIndividually)
-	{
-		return WriteObjectsAndSubObjectsPendingDestroy(Context);
-	}
-	else
-	{
-		return WriteRootObjectsPendingDestroy(Context);
-	}
-}
-
-uint32 FReplicationWriter::WriteObjectsAndSubObjectsPendingDestroy(FNetSerializationContext& Context)
-{
-	FNetBitStreamWriter& Writer = *Context.GetBitStreamWriter();
-
-	UE_NET_TRACE_SCOPE(ObjectsAndSubObjectsPendingDestroy, Writer, Context.GetTraceCollector(), ENetTraceVerbosity::Trace);
-
-	uint32 WrittenCount = 0;
-
-	// Write how many destroyed objects we have
-	const uint32 HeaderPos = Writer.GetPosBits();
-	constexpr uint32 DestroyObjectBitCount = 16U;
-	constexpr uint32 MaxDestroyObjectCount = (1U << DestroyObjectBitCount) - 1U;
-
-	Writer.WriteBits(WrittenCount, DestroyObjectBitCount);
-
-	// Can't write if bitstream is overflown. Shouldn't write if we're an OOB packet or in replication record starvation mode.
-	if (Writer.IsOverflown() || WriteContext.bIsOOBPacket || WriteContext.bIsInReplicationRecordStarvation)
-	{
-		return 0U;
-	}
-
-	bool bWroteAllDestroyedObjects = true;
-
-	// Note: This code is not honoring net.Iris.MaxDestroyObjectsPerPacket because it's bug prone to split root and subobjects in different packets, which is already a possibility with this code. Set net.Iris.DestroyRootAndSubObjectsIndividually 0 to get proper destruction throttling.
-	for (uint32 InternalIndex = 0U; ((InternalIndex = ObjectsPendingDestroy.FindFirstOne(InternalIndex + 1U)) != FNetBitArray::InvalidIndex) && (WrittenCount < MaxDestroyObjectCount); )
-	{
-		FReplicationInfo& Info = GetReplicationInfo(InternalIndex);
-		const FNetRefHandleManager::FReplicatedObjectData& ObjectData = NetRefHandleManager->GetReplicatedObjectDataNoCheck(InternalIndex);
-
-		// Don't send destroy until object creation has been acked.
-		if (!Info.IsCreationConfirmed)
-		{
-			continue;
-		}
-
-		// Already waiting on destroy confirmation
-		if (Info.GetState() == EReplicatedObjectState::WaitOnDestroyConfirmation)
-		{
-			continue;
-		}
-
-		if (Info.GetState() == EReplicatedObjectState::WaitOnFlush)
-		{
-			if (GetFlushStatus(InternalIndex, Info, Info.FlushFlags) != EFlushFlags::FlushFlags_None)
-			{
-				continue;
-			}
-			
-			// Object and subobjects are now flushed and can be destroyed
-			SetPendingDestroyOrSubObjectPendingDestroyState(InternalIndex, Info);
-		}
-		
-		if (ObjectData.IsSubObject())
-		{
-			if (Info.SubObjectPendingDestroy)
-			{
-				// If the owner is pending destroy or already destroyed then we need to replicate the destruction for the subobject here.
-				if (ObjectsPendingDestroy.GetBit(ObjectData.SubObjectRootIndex) || (ReplicatedObjects[ObjectData.SubObjectRootIndex].GetState() == EReplicatedObjectState::Invalid))
-				{
-					//			UE_LOG(LogTemp, Warning, TEXT("Object %s was a subobject but owner was destroyed so we need to destroy object normally"), ToCStr(ObjectData.RefHandle.ToString()));
-					SetState(InternalIndex, EReplicatedObjectState::PendingDestroy);
-					Info.SubObjectPendingDestroy = 0U;
-				}
-				else
-				{
-					// Replicate subobject destruction via state replication.
-					continue;
-				}
-			}
-			else
-			{
-				// If the owner is valid and isn't pending destroy we need to replicate destruction via state replication.
-				if (!ObjectsPendingDestroy.GetBit(ObjectData.SubObjectRootIndex))
-				{
-					FReplicationInfo& OwnerInfo = GetReplicationInfo(ObjectData.SubObjectRootIndex);
-					if (OwnerInfo.GetState() != EReplicatedObjectState::Invalid && OwnerInfo.GetState() < EReplicatedObjectState::PendingDestroy)
-					{
-						//			UE_LOG(LogTemp, Warning, TEXT("SubObject %s Should be destroyed by replicating owner"), ToCStr(ObjectData.RefHandle.ToString()));
-						SetState(InternalIndex, EReplicatedObjectState::SubObjectPendingDestroy);
-						Info.SubObjectPendingDestroy = 1U;
-
-						OwnerInfo.HasDirtySubObjects = 1U;
-
-						ObjectsWithDirtyChanges.SetBit(ObjectData.SubObjectRootIndex);
-						ObjectsWithDirtyChanges.SetBit(InternalIndex);
-						continue;
-					}
-				}
-			}
-		}
-
-		// Unexpected. Get more info.
-		if (Info.GetState() != EReplicatedObjectState::PendingDestroy)
-		{
-			ensureMsgf(Info.GetState() == EReplicatedObjectState::PendingDestroy, TEXT("Skipping writing destroy for object %s which is in unexpected state %s. IsSubObject: %u IsDestructionInfo: %u"), ToCStr(NetRefHandleManager->PrintObjectFromIndex(InternalIndex)), LexToString(Info.GetState()), Info.IsSubObject, Info.IsDestructionInfo);
-			continue;
-		}
-
-		// We do not support destroying an object that is currently being sent as a huge object.
-		if (IsObjectPartOfActiveHugeObject(InternalIndex))
-		{
-			UE_LOG(LogIris, Verbose, TEXT("Skipping writing destroy for object ( InternalIndex: %u ) which is part of active huge object."), InternalIndex);
-			bWroteAllDestroyedObjects = false;
-			continue;
-		}
-
-		UE_NET_TRACE_OBJECT_SCOPE(ObjectData.RefHandle,  Writer, Context.GetTraceCollector(), ENetTraceVerbosity::Trace);
-
-		FNetBitStreamRollbackScope RollbackScope(Writer);
-
-		// Write handle with the needed bitCount
-		WriteNetRefHandleId(Context, ObjectData.RefHandle);
-
-		// If this is a explicit subobject we must also send the owner handle in the case that the client is currently queued data due to async loading.
-		if (Writer.WriteBool(ObjectData.IsSubObject()))
-		{			
-			WriteNetRefHandleId(Context, NetRefHandleManager->GetReplicatedObjectDataNoCheck(ObjectData.SubObjectRootIndex).RefHandle);
-		}
-
-		// Write bit indicating if the static instance should be destroyed or not (could skip the bit for dynamic objects)
-		const bool bShouldDestroyInstance = ObjectData.RefHandle.IsDynamic() || NetRefHandleManager->GetIsDestroyedStartupObject(InternalIndex);
-		Writer.WriteBool(bShouldDestroyInstance);
-
-		if (!Writer.IsOverflown())
-		{
-			// Must update state before pushing record
-			SetState(InternalIndex, EReplicatedObjectState::WaitOnDestroyConfirmation);
-			Info.HasDirtyChangeMask = 0U;
-
-			// update transmission record
-			FBatchObjectInfo ObjectInfo = {};
-			FObjectRecord ObjectRecord;
-
-			ObjectInfo.InternalIndex = InternalIndex;
-			CreateObjectRecord(nullptr, Info, ObjectInfo, ObjectRecord);
-			CommitObjectRecord(InternalIndex, ObjectRecord);
-	
-			++WrittenCount;
-		}
-		else
-		{
-			break;
-		}
-	}
-
-	// Write Header
-	{
-		FNetBitStreamWriteScope WriteScope(Writer, HeaderPos);
-		Writer.WriteBits(WrittenCount, DestroyObjectBitCount);
-	}
-
-	bWroteAllDestroyedObjects = bWroteAllDestroyedObjects && !Writer.IsOverflown() && (WrittenCount < MaxDestroyObjectCount);
-	WriteContext.bHasDestroyedObjectsToSend = !bWroteAllDestroyedObjects;
-
-	return WrittenCount;
+	return WriteRootObjectsPendingDestroy(Context);
 }
 
 uint32 FReplicationWriter::WriteRootObjectsPendingDestroy(FNetSerializationContext& Context)
@@ -2048,7 +1904,7 @@ uint32 FReplicationWriter::WriteRootObjectsPendingDestroy(FNetSerializationConte
 	{
 		for (uint32 SubObjectIndex : NetRefHandleManager->GetSubObjects(InternalIndex))
 		{
-			const FReplicationInfo& SubObjectInfo = GetReplicationInfo(SubObjectIndex);
+			FReplicationInfo& SubObjectInfo = GetReplicationInfo(SubObjectIndex);
 			const EReplicatedObjectState SubObjectState = SubObjectInfo.GetState();
 			// A subobject may have stopped replicating immediately when going out of scope if it was never replicated to begin with.
 			if (SubObjectState == EReplicatedObjectState::Invalid)
@@ -2056,6 +1912,15 @@ uint32 FReplicationWriter::WriteRootObjectsPendingDestroy(FNetSerializationConte
 				continue;
 			}
 
+			if (SubObjectState == EReplicatedObjectState::WaitOnFlush)
+			{
+				if (GetFlushStatus(SubObjectIndex, SubObjectInfo, SubObjectInfo.FlushFlags) == EFlushFlags::FlushFlags_None)
+				{
+					ensureMsgf(!SubObjectInfo.TearOff, TEXT("Torn off subobjects should not be destroyed via the pending destroy. ( InternalIndex: %u )"), SubObjectIndex);
+					SetPendingDestroyOrSubObjectPendingDestroyState(SubObjectIndex, SubObjectInfo);
+				}
+			}
+
 			if (!(SubObjectState == EReplicatedObjectState::PendingDestroy || SubObjectState == EReplicatedObjectState::SubObjectPendingDestroy))
 			{
 				return false;
@@ -2088,8 +1953,8 @@ uint32 FReplicationWriter::WriteRootObjectsPendingDestroy(FNetSerializationConte
 			continue;
 		}
 
-		// Already waiting on destroy confirmation
-		if (Info.GetState() == EReplicatedObjectState::WaitOnDestroyConfirmation)
+		// Already waiting on destroy confirmation or need to replicate via the regular path
+		if (Info.GetState() == EReplicatedObjectState::WaitOnDestroyConfirmation || Info.GetState() == EReplicatedObjectState::PendingTearOff)
 		{
 			continue;
 		}
@@ -2100,19 +1965,45 @@ uint32 FReplicationWriter::WriteRootObjectsPendingDestroy(FNetSerializationConte
 			{
 				continue;
 			}
-			
-			// Object *and* subobjects are now flushed and can be destroyed
-			SetPendingDestroyOrSubObjectPendingDestroyState(InternalIndex, Info);
 
-			// Change state for subobjects waiting on flush
-			for (uint32 SubObjectIndex : NetRefHandleManager->GetSubObjects(InternalIndex))
+			// Tear off can happen while an object is out of scope. We should go through the state replication tear off path.
+			if (Info.TearOff)
 			{
-				FReplicationInfo& SubObjectInfo = GetReplicationInfo(SubObjectIndex);
-				if (SubObjectInfo.GetState() == EReplicatedObjectState::WaitOnFlush)
+				// Check that all subobjects are ready to be torn off too.
+				const bool bAllSubObjectsAreReadyToBeTornOff = [this]()
 				{
-					SetPendingDestroyOrSubObjectPendingDestroyState(SubObjectIndex, SubObjectInfo);
+					for (const FInternalNetRefIndex SubObjectIndex : NetRefHandleManager->GetSubObjectInternalIndicesView())
+					{
+						FReplicationInfo& SubObjectInfo = ReplicatedObjects[SubObjectIndex];
+						if (SubObjectInfo.GetState() == EReplicatedObjectState::WaitOnFlush)
+						{
+							if (GetFlushStatus(SubObjectIndex, SubObjectInfo, SubObjectInfo.FlushFlags) == EFlushFlags::FlushFlags_None)
+							{
+								SetState(SubObjectIndex, EReplicatedObjectState::PendingTearOff);
+								ObjectsWithDirtyChanges.SetBit(SubObjectIndex);
+							}
+							else
+							{
+								return false;
+							}
+						}
+					}
+
+					return true;
+				}();
+
+				if (bAllSubObjectsAreReadyToBeTornOff)
+				{
+					SetState(InternalIndex, EReplicatedObjectState::PendingTearOff);
+					ObjectsWithDirtyChanges.SetBit(InternalIndex);
+					Info.HasDirtySubObjects = 1;
 				}
 			}
+			else
+			{
+				// Object and subobjects are now flushed and can be destroyed
+				SetPendingDestroyOrSubObjectPendingDestroyState(InternalIndex, Info);
+			}
 		}
 
 		if (!AreSubObjectsReadyToBeDestroyed(InternalIndex))
@@ -2310,7 +2201,7 @@ FReplicationWriter::EWriteObjectStatus FReplicationWriter::WriteObjectAndSubObje
 	const EReplicatedObjectState State = Info.GetState();
 
 	// As an object might still have subobjects pending destroy in the list of subobjects 
-	if (State == EReplicatedObjectState::Invalid || !ensureMsgf(State > EReplicatedObjectState::Invalid && State < EReplicatedObjectState::PendingDestroy, TEXT("Unsupported state %s"), ToCStr(LexToString(State))))
+	if (State == EReplicatedObjectState::Invalid || !ensureMsgf(State > EReplicatedObjectState::Invalid && State < EReplicatedObjectState::PendingDestroy, TEXT("Unsupported state %s ( InternalIndex: %u )"), ToCStr(LexToString(State)), InternalIndex))
 	{
 		return EWriteObjectStatus::InvalidState;
 	}
