Iris
- Supressed intentional Errors from ReplicationSystemTest

diff --git a/Engine/Plugins/Runtime/ReplicationSystemTestPlugin/Source/Private/Tests/DataStream/TestChunkedDataStream.cpp b/Engine/Plugins/Runtime/ReplicationSystemTestPlugin/Source/Private/Tests/DataStream/TestChunkedDataStream.cpp
--- a/Engine/Plugins/Runtime/ReplicationSystemTestPlugin/Source/Private/Tests/DataStream/TestChunkedDataStream.cpp
+++ b/Engine/Plugins/Runtime/ReplicationSystemTestPlugin/Source/Private/Tests/DataStream/TestChunkedDataStream.cpp
@@ -9,6 +9,7 @@
 #include "Iris/DataStream/DataStreamDefinitions.h"
 #include "Iris/DataStream/DataStreamManager.h"
 #include "Iris/ReplicationSystem/ChunkedDataStream/ChunkedDataStream.h"
+#include "Iris/ReplicationSystem/ChunkedDataStream/ChunkedDataStreamCommon.h"
 
 #include "Iris/PacketControl/PacketNotification.h"
 #include "Iris/Serialization/NetBitStreamReader.h"
@@ -336,6 +337,9 @@ UE_NET_TEST_FIXTURE(FTestChunkedDataStreamFixture, TestChunkedDataStream_CannotE
 		UE_NET_ASSERT_TRUE(bResult);
 	}
 
+	// Suppress error, since we're intentionally overflowing
+	LOG_SCOPE_VERBOSITY_OVERRIDE(LogIrisChunkedDataStream, ELogVerbosity::Fatal);
+
 	// Try to enqueue some data that should not be allowed
 	{
 		TSharedRef<TArray64<uint8>> Payload = MakeShared<TArray64<uint8>>();
@@ -389,6 +393,10 @@ UE_NET_TEST_FIXTURE(FTestChunkedDataStreamFixture, TestChunkedDataStream_ClientW
 		ServerStream->EnqueuePayload(Payload);
 	}
 
+
+	// Suppress error, since we're intentionally overflowing
+	LOG_SCOPE_VERBOSITY_OVERRIDE(LogIrisChunkedDataStream, ELogVerbosity::Fatal);
+
 	// Tick send/receive
 	RoundTrip();
 
