Test that push model is enabled for ReplicationSystemTests and enable it for the ReplicationSystemLowLevelTests program.

diff --git a/Engine/Plugins/Runtime/ReplicationSystemTestPlugin/Source/Private/Tests/TestNetworkAutomationTestFramework.cpp b/Engine/Plugins/Runtime/ReplicationSystemTestPlugin/Source/Private/Tests/TestNetworkAutomationTestFramework.cpp
--- a/Engine/Plugins/Runtime/ReplicationSystemTestPlugin/Source/Private/Tests/TestNetworkAutomationTestFramework.cpp
+++ b/Engine/Plugins/Runtime/ReplicationSystemTestPlugin/Source/Private/Tests/TestNetworkAutomationTestFramework.cpp
@@ -2,6 +2,8 @@
 
 #include "NetworkAutomationTest.h"
 #include "NetworkAutomationTestMacros.h"
+#include "Net/Core/PushModel/PushModelMacros.h"
+#include "HAL/IConsoleManager.h"
 
 namespace UE::Net::Private
 {
@@ -229,4 +231,15 @@ UE_NET_TEST(NetworkAutomationTest, TestLogPrimitiveTypes)
 	UE_NET_LOG("nullptr_t ") << nullptr;
 }
 
+UE_NET_TEST(NetworkAutomationTest, TestPushModelEnabled)
+{
+#if !WITH_PUSH_MODEL
+	UE_NET_FAIL("Push model support not compiled in.");
+#endif
+
+	IConsoleVariable* PushModelCVar = IConsoleManager::Get().FindConsoleVariable(TEXT("Net.IsPushModelEnabled"));
+	UE_NET_ASSERT_NE(PushModelCVar, nullptr);
+	UE_NET_ASSERT_TRUE(PushModelCVar->GetBool());
+}
+
 }
diff --git a/Engine/Source/Programs/ReplicationSystemLowLevelTests/Private/ReplicationSystemLowLevelTests.cpp b/Engine/Source/Programs/ReplicationSystemLowLevelTests/Private/ReplicationSystemLowLevelTests.cpp
--- a/Engine/Source/Programs/ReplicationSystemLowLevelTests/Private/ReplicationSystemLowLevelTests.cpp
+++ b/Engine/Source/Programs/ReplicationSystemLowLevelTests/Private/ReplicationSystemLowLevelTests.cpp
@@ -1,6 +1,7 @@
 // Copyright Epic Games, Inc. All Rights Reserved.
 
 #include "HAL/PlatformMisc.h"
+#include "HAL/IConsoleManager.h"
 #include "Iris/IrisConfig.h"
 #include "Logging/LogScopedVerbosityOverride.h"
 #include "Misc/CommandLine.h"
@@ -57,6 +58,11 @@ GROUP_BEFORE_GLOBAL(Catch::DefaultGroup)
 
 	UE::Net::SetUseIrisReplication(true);
 
+	if (IConsoleVariable* PushModelCVar = IConsoleManager::Get().FindConsoleVariable(TEXT("Net.IsPushModelEnabled")))
+	{
+		PushModelCVar->Set(true);
+	}
+
 	{
 		// Find the engine directory before InitAll() enables the platform file stub. Prevents warnings.
 		FPlatformMisc::EngineDir();
