Update Iris global dirty net object tracker to support multiple coexisting replication systems.
-Add a delegate that's called on global dirty pollers that haven't gathered dirty state before the global dirty state is reset.
-Unlock the global dirty list when we would have reset it even if there are multiple replication systems.
-Add tests to cover the global dirty tracker & multiple replication systems.

diff --git a/Engine/Plugins/Runtime/ReplicationSystemTestPlugin/Source/Private/Tests/ReplicationSystem/MultiReplicationSystemsTestFixture.cpp b/Engine/Plugins/Runtime/ReplicationSystemTestPlugin/Source/Private/Tests/ReplicationSystem/MultiReplicationSystemsTestFixture.cpp
--- a/Engine/Plugins/Runtime/ReplicationSystemTestPlugin/Source/Private/Tests/ReplicationSystem/MultiReplicationSystemsTestFixture.cpp
+++ b/Engine/Plugins/Runtime/ReplicationSystemTestPlugin/Source/Private/Tests/ReplicationSystem/MultiReplicationSystemsTestFixture.cpp
@@ -126,12 +126,24 @@ void FMultiReplicationSystemsTestFixture::EndReplication(UTestReplicatedIrisObje
 	}
 }
 
-void FMultiReplicationSystemsTestFixture::BeginReplication(FReplicationSystemTestServer* Server, UTestReplicatedIrisObject* Object)
+FNetRefHandle FMultiReplicationSystemsTestFixture::BeginReplication(FReplicationSystemTestServer* Server, UTestReplicatedIrisObject* Object)
 {
 	if (UReplicatedTestObjectBridge* ReplicationBridge = Server->GetReplicationBridge())
 	{
-		ReplicationBridge->BeginReplication(Object);
+		return ReplicationBridge->BeginReplication(Object);
 	}
+
+	return FNetRefHandle::GetInvalid();
+}
+
+FNetRefHandle FMultiReplicationSystemsTestFixture::BeginReplication(FReplicationSystemTestServer* Server, UTestReplicatedIrisObject* Object, const UObjectReplicationBridge::FRootObjectReplicationParams& Params)
+{
+	if (UReplicatedTestObjectBridge* ReplicationBridge = Server->GetReplicationBridge())
+	{
+		return ReplicationBridge->BeginReplication(Object, Params);
+	}
+
+	return FNetRefHandle::GetInvalid();
 }
 
 void FMultiReplicationSystemsTestFixture::EndReplication(FReplicationSystemTestServer* Server, UTestReplicatedIrisObject* Object)
@@ -164,4 +176,42 @@ void FMultiReplicationSystemsTestFixture::FullSendAndDeliverUpdate()
 	}
 }
 
+void FMultiReplicationSystemsTestFixture::FullSendAndDeliverUpdateTwoPass()
+{
+	// PreSend & Send
+	for (FReplicationSystemTestServer* Server : Servers)
+	{
+		TArray<FReplicationSystemTestClient*>* FoundClients = ServerClients.Find(Server);
+		if (FoundClients && FoundClients->Num() > 0)
+		{
+			Server->PreSendUpdate();
+
+			for (FReplicationSystemTestClient* Client : *FoundClients)
+			{
+				Server->SendTo(Client);
+			}
+		}
+	}
+
+	// PostSend
+	for (FReplicationSystemTestServer* Server : Servers)
+	{
+		TArray<FReplicationSystemTestClient*>* FoundClients = ServerClients.Find(Server);
+		if (FoundClients && FoundClients->Num() > 0)
+		{
+			Server->PostSendUpdate();
+		}
+	}
+
+	// Deliver
+	for (TMap<FReplicationSystemTestServer*, TArray<FReplicationSystemTestClient*>>::ElementType& ServerWithClients : ServerClients)
+	{
+		FReplicationSystemTestServer* Server = ServerWithClients.Key;
+		for (FReplicationSystemTestClient* Client : ServerWithClients.Value)
+		{
+			Server->DeliverTo(Client, DeliverPacket);
+		}
+	}
+}
+
 }
diff --git a/Engine/Plugins/Runtime/ReplicationSystemTestPlugin/Source/Private/Tests/ReplicationSystem/MultiReplicationSystemsTestFixture.h b/Engine/Plugins/Runtime/ReplicationSystemTestPlugin/Source/Private/Tests/ReplicationSystem/MultiReplicationSystemsTestFixture.h
--- a/Engine/Plugins/Runtime/ReplicationSystemTestPlugin/Source/Private/Tests/ReplicationSystem/MultiReplicationSystemsTestFixture.h
+++ b/Engine/Plugins/Runtime/ReplicationSystemTestPlugin/Source/Private/Tests/ReplicationSystem/MultiReplicationSystemsTestFixture.h
@@ -58,8 +58,11 @@ protected:
 	/** Ends replication of the object on all replication systems. */
 	void EndReplication(UTestReplicatedIrisObject* Object);
 
-	/** Begins replicatiion of the object on a specific replication system. */
-	void BeginReplication(FReplicationSystemTestServer* Server, UTestReplicatedIrisObject* Object);
+	/** Begins replication of the object on a specific replication system. */
+	FNetRefHandle BeginReplication(FReplicationSystemTestServer* Server, UTestReplicatedIrisObject* Object);
+
+	/** Begins replication of the object on a specific replication system with the given parameters. */
+	FNetRefHandle BeginReplication(FReplicationSystemTestServer* Server, UTestReplicatedIrisObject* Object, const UObjectReplicationBridge::FRootObjectReplicationParams& Params);
 
 	/** Ends replication of the object on a specific replication system. */
 	void EndReplication(FReplicationSystemTestServer* Server, UTestReplicatedIrisObject* Object);
@@ -67,6 +70,14 @@ protected:
 	/** Performs PreSendUpdate, SendAndDeliverTo and PostSendUpdate on all servers. Packets are delivered to all clients connected to the servers. */
 	void FullSendAndDeliverUpdate();
 
+	/**
+	 * Performs PreSendUpdate and SendAndDeliverTo on each server one at a time.
+	 * Then in a separate loop performs PostSendUpdate on each server.
+	 * Then packets are delivered to all clients connected to the servers.
+	 * This more closely resembles the order that occurs with UWorld and TickFlush/PostTickFlush.
+	 */
+	void FullSendAndDeliverUpdateTwoPass();
+
 private:
 	FDataStreamTestUtil DataStreamUtil;
 	TArray<FReplicationSystemTestServer*> Servers;
diff --git a/Engine/Plugins/Runtime/ReplicationSystemTestPlugin/Source/Private/Tests/ReplicationSystem/TestDirtyNetObjectTracker.cpp b/Engine/Plugins/Runtime/ReplicationSystemTestPlugin/Source/Private/Tests/ReplicationSystem/TestDirtyNetObjectTracker.cpp
--- a/Engine/Plugins/Runtime/ReplicationSystemTestPlugin/Source/Private/Tests/ReplicationSystem/TestDirtyNetObjectTracker.cpp
+++ b/Engine/Plugins/Runtime/ReplicationSystemTestPlugin/Source/Private/Tests/ReplicationSystem/TestDirtyNetObjectTracker.cpp
@@ -171,7 +171,7 @@ UE_NET_TEST_FIXTURE(FReplicationSystemServerClientTestFixture, GlobalDirtyTracke
 	UE_NET_ASSERT_NE(ClientObject->IntA, ServerObject->IntA);
 
 	// Now mark the object dirty
-	Server->ReplicationSystem->ForceNetUpdate(ServerObject->NetRefHandle);
+	FGlobalDirtyNetObjectTracker::MarkNetObjectStateDirty(FNetHandleManager::GetOrCreateNetHandle(ServerObject));
 
 	// Send and deliver packet
 	Server->UpdateAndSend({ Client });
diff --git a/Engine/Plugins/Runtime/ReplicationSystemTestPlugin/Source/Private/Tests/ReplicationSystem/TestMultiReplicationSystems.cpp b/Engine/Plugins/Runtime/ReplicationSystemTestPlugin/Source/Private/Tests/ReplicationSystem/TestMultiReplicationSystems.cpp
--- a/Engine/Plugins/Runtime/ReplicationSystemTestPlugin/Source/Private/Tests/ReplicationSystem/TestMultiReplicationSystems.cpp
+++ b/Engine/Plugins/Runtime/ReplicationSystemTestPlugin/Source/Private/Tests/ReplicationSystem/TestMultiReplicationSystems.cpp
@@ -2,6 +2,7 @@
 
 #include "Tests/ReplicationSystem/MultiReplicationSystemsTestFixture.h"
 #include "Net/Core/NetHandle/NetHandleManager.h"
+#include "Containers/StaticArray.h"
 
 namespace UE::Net::Private
 {
@@ -247,4 +248,735 @@ UE_NET_TEST_FIXTURE(FMultiReplicationSystemsTestFixture, ModifiedObjectIsReplica
 	}
 }
 
+/** Test that independent objects in separate systems are dirtied properly */
+UE_NET_TEST_FIXTURE(FMultiReplicationSystemsTestFixture, GlobalDirtyTracker)
+{
+	constexpr int32 NumServers = 3;
+	CreateServers(NumServers);
+
+	TArrayView<FReplicationSystemTestServer*> AllServers = GetAllServers();
+	UE_NET_ASSERT_EQ(AllServers.Num(), NumServers);
+
+	TStaticArray<UTestReplicatedIrisObject*, NumServers> ServerObjects = {};
+
+	// Create clients and replicated objects for all systems
+	for (uint32 i = 0; i < NumServers; ++i)
+	{
+		constexpr unsigned ClientCountPerServer = 3U;
+		for (unsigned ClientIt = 0, ClientEndIt = ClientCountPerServer; ClientIt != ClientEndIt; ++ClientIt)
+		{
+			CreateClientForServer(AllServers[i]);
+		}
+
+		// Spawn object on server that won't be polled automatically multiple times during this test
+		UObjectReplicationBridge::FRootObjectReplicationParams Params;
+		Params.PollFrequency = 0.001f;
+		Params.bUseClassConfigDynamicFilter = true;
+		Params.bNeedsPreUpdate = true;
+
+		ServerObjects[i] = AllServers[i]->CreateObject(Params);
+	}
+
+	// Send and deliver packets
+	FullSendAndDeliverUpdateTwoPass();
+
+	// Objects should have been created on the clients
+	for (int32 i = 0; i < NumServers; ++i)
+	{
+		for (FReplicationSystemTestClient* Client : GetClients(AllServers[i]))
+		{
+			UTestReplicatedIrisObject* ClientObject = Cast<UTestReplicatedIrisObject>(Client->GetReplicationBridge()->GetReplicatedObject(ServerObjects[i]->NetRefHandle));
+			UE_NET_ASSERT_NE(ClientObject, nullptr);
+		}
+	}
+
+	// Set a replicated variable, but don't mark it dirty
+	for (UTestReplicatedIrisObject* ServerObject : ServerObjects)
+	{
+		ServerObject->IntA = 0xFF;
+	}
+
+	// Send and deliver packets
+	FullSendAndDeliverUpdateTwoPass();
+
+	// Client replicated property should not have changed
+	for (int32 i = 0; i < NumServers; ++i)
+	{
+		for (FReplicationSystemTestClient* Client : GetClients(AllServers[i]))
+		{
+			UTestReplicatedIrisObject* ClientObject = Cast<UTestReplicatedIrisObject>(Client->GetReplicationBridge()->GetReplicatedObject(ServerObjects[i]->NetRefHandle));
+			UE_NET_ASSERT_NE(ClientObject->IntA, ServerObjects[i]->IntA);
+		}
+	}
+
+	// Now mark the object dirty
+	for (uint32 i = 0; i < NumServers; ++i)
+	{
+		FGlobalDirtyNetObjectTracker::MarkNetObjectStateDirty(FNetHandleManager::GetOrCreateNetHandle(ServerObjects[i]));
+	}
+
+	// Send and deliver packets
+	FullSendAndDeliverUpdateTwoPass();
+
+	// Client replicated properties should have changed now
+	for (int32 i = 0; i < NumServers; ++i)
+	{
+		for (FReplicationSystemTestClient* Client : GetClients(AllServers[i]))
+		{
+			UTestReplicatedIrisObject* ClientObject = Cast<UTestReplicatedIrisObject>(Client->GetReplicationBridge()->GetReplicatedObject(ServerObjects[i]->NetRefHandle));
+			UE_NET_ASSERT_EQ(ClientObject->IntA, ServerObjects[i]->IntA);
+		}
+
+		AllServers[i]->DestroyObject(ServerObjects[i]);
+	}
+}
+
+/** Test the global dirty tracker with multiple systems where one has no clients */
+UE_NET_TEST_FIXTURE(FMultiReplicationSystemsTestFixture, GlobalDirtyTrackerWithNoClients)
+{
+	constexpr int32 NumServers = 3;
+	constexpr int32 NumServersWithClients = NumServers - 1;
+	CreateServers(NumServers);
+
+	TArrayView<FReplicationSystemTestServer*> AllServers = GetAllServers();
+	UE_NET_ASSERT_EQ(AllServers.Num(), NumServers);
+
+	TStaticArray<UTestReplicatedIrisObject*, NumServersWithClients> ServerObjects = {};
+
+	// Create clients and replicated objects for all but one of the servers
+	for (uint32 i = 0; i < NumServersWithClients; ++i)
+	{
+		constexpr unsigned ClientCountPerServer = 3U;
+		for (unsigned ClientIt = 0, ClientEndIt = ClientCountPerServer; ClientIt != ClientEndIt; ++ClientIt)
+		{
+			CreateClientForServer(AllServers[i]);
+		}
+
+		// Spawn object on server that won't be polled automatically multiple times during this test
+		UObjectReplicationBridge::FRootObjectReplicationParams Params;
+		Params.PollFrequency = 0.001f;
+		Params.bUseClassConfigDynamicFilter = true;
+		Params.bNeedsPreUpdate = true;
+
+		ServerObjects[i] = AllServers[i]->CreateObject(Params);
+	}
+
+	// Send and deliver packets
+	FullSendAndDeliverUpdateTwoPass();
+
+	// Objects should have been created on the clients
+	for (int32 i = 0; i < NumServersWithClients; ++i)
+	{
+		for (FReplicationSystemTestClient* Client : GetClients(AllServers[i]))
+		{
+			UTestReplicatedIrisObject* ClientObject = Cast<UTestReplicatedIrisObject>(Client->GetReplicationBridge()->GetReplicatedObject(ServerObjects[i]->NetRefHandle));
+			UE_NET_ASSERT_NE(ClientObject, nullptr);
+		}
+	}
+
+	// Set a replicated variable, but don't mark it dirty
+	for (UTestReplicatedIrisObject* ServerObject : ServerObjects)
+	{
+		ServerObject->IntA = 0xFF;
+	}
+
+	// Send and deliver packets
+	FullSendAndDeliverUpdateTwoPass();
+
+	// Client replicated property should not have changed
+	for (int32 i = 0; i < NumServersWithClients; ++i)
+	{
+		for (FReplicationSystemTestClient* Client : GetClients(AllServers[i]))
+		{
+			UTestReplicatedIrisObject* ClientObject = Cast<UTestReplicatedIrisObject>(Client->GetReplicationBridge()->GetReplicatedObject(ServerObjects[i]->NetRefHandle));
+			UE_NET_ASSERT_NE(ClientObject->IntA, ServerObjects[i]->IntA);
+		}
+	}
+
+	// Now mark the object dirty
+	for (uint32 i = 0; i < NumServersWithClients; ++i)
+	{
+		FGlobalDirtyNetObjectTracker::MarkNetObjectStateDirty(FNetHandleManager::GetOrCreateNetHandle(ServerObjects[i]));
+	}
+
+	// Send and deliver packets
+	FullSendAndDeliverUpdateTwoPass();
+
+	// Client replicated properties should have changed now
+	for (int32 i = 0; i < NumServersWithClients; ++i)
+	{
+		for (FReplicationSystemTestClient* Client : GetClients(AllServers[i]))
+		{
+			UTestReplicatedIrisObject* ClientObject = Cast<UTestReplicatedIrisObject>(Client->GetReplicationBridge()->GetReplicatedObject(ServerObjects[i]->NetRefHandle));
+			UE_NET_ASSERT_EQ(ClientObject->IntA, ServerObjects[i]->IntA);
+		}
+
+		AllServers[i]->DestroyObject(ServerObjects[i]);
+	}
+}
+
+/** Test that independent objects in separate systems are dirtied properly and independently */
+UE_NET_TEST_FIXTURE(FMultiReplicationSystemsTestFixture, GlobalDirtyTrackerOneObjectDirty)
+{
+	constexpr int32 NumServers = 3;
+	CreateServers(NumServers);
+
+	TArrayView<FReplicationSystemTestServer*> AllServers = GetAllServers();
+	UE_NET_ASSERT_EQ(AllServers.Num(), NumServers);
+
+	TStaticArray<UTestReplicatedIrisObject*, NumServers> ServerObjects = {};
+
+	// Create clients and replicated objects for all systems
+	for (uint32 i = 0; i < NumServers; ++i)
+	{
+		constexpr unsigned ClientCountPerServer = 3U;
+		for (unsigned ClientIt = 0, ClientEndIt = ClientCountPerServer; ClientIt != ClientEndIt; ++ClientIt)
+		{
+			CreateClientForServer(AllServers[i]);
+		}
+
+		// Spawn object on server that won't be polled automatically multiple times during this test
+		UObjectReplicationBridge::FRootObjectReplicationParams Params;
+		Params.PollFrequency = 0.001f;
+		Params.bUseClassConfigDynamicFilter = true;
+		Params.bNeedsPreUpdate = true;
+
+		ServerObjects[i] = AllServers[i]->CreateObject(Params);
+	}
+
+	// Send and deliver packets
+	FullSendAndDeliverUpdateTwoPass();
+
+	// Objects should have been created on the clients
+	for (int32 i = 0; i < NumServers; ++i)
+	{
+		for (FReplicationSystemTestClient* Client : GetClients(AllServers[i]))
+		{
+			UTestReplicatedIrisObject* ClientObject = Cast<UTestReplicatedIrisObject>(Client->GetReplicationBridge()->GetReplicatedObject(ServerObjects[i]->NetRefHandle));
+			UE_NET_ASSERT_NE(ClientObject, nullptr);
+		}
+	}
+
+	// Set a replicated variable, but don't mark it dirty
+	for (UTestReplicatedIrisObject* ServerObject : ServerObjects)
+	{
+		ServerObject->IntA = 0xFF;
+	}
+
+	// Send and deliver packets
+	FullSendAndDeliverUpdateTwoPass();
+
+	// Client replicated property should not have changed
+	for (int32 i = 0; i < NumServers; ++i)
+	{
+		for (FReplicationSystemTestClient* Client : GetClients(AllServers[i]))
+		{
+			UTestReplicatedIrisObject* ClientObject = Cast<UTestReplicatedIrisObject>(Client->GetReplicationBridge()->GetReplicatedObject(ServerObjects[i]->NetRefHandle));
+			UE_NET_ASSERT_NE(ClientObject->IntA, ServerObjects[i]->IntA);
+		}
+	}
+
+	// Now only mark server 1's object dirty
+	FGlobalDirtyNetObjectTracker::MarkNetObjectStateDirty(FNetHandleManager::GetOrCreateNetHandle(ServerObjects[1]));
+
+	// Send and deliver packets
+	FullSendAndDeliverUpdateTwoPass();
+
+	// Client replicated property should have changed now on server 1...
+	for (FReplicationSystemTestClient* Client : GetClients(AllServers[1]))
+	{
+		UTestReplicatedIrisObject* ClientObject = Cast<UTestReplicatedIrisObject>(Client->GetReplicationBridge()->GetReplicatedObject(ServerObjects[1]->NetRefHandle));
+		UE_NET_ASSERT_EQ(ClientObject->IntA, ServerObjects[1]->IntA);
+	}
+
+	// But not server 0 or 2.
+	for (FReplicationSystemTestClient* Client : GetClients(AllServers[0]))
+	{
+		UTestReplicatedIrisObject* ClientObject = Cast<UTestReplicatedIrisObject>(Client->GetReplicationBridge()->GetReplicatedObject(ServerObjects[0]->NetRefHandle));
+		UE_NET_ASSERT_NE(ClientObject->IntA, ServerObjects[0]->IntA);
+	}
+
+	for (FReplicationSystemTestClient* Client : GetClients(AllServers[2]))
+	{
+		UTestReplicatedIrisObject* ClientObject = Cast<UTestReplicatedIrisObject>(Client->GetReplicationBridge()->GetReplicatedObject(ServerObjects[2]->NetRefHandle));
+		UE_NET_ASSERT_NE(ClientObject->IntA, ServerObjects[2]->IntA);
+	}
+
+	for (int32 i = 0; i < NumServers; ++i)
+	{
+		AllServers[i]->DestroyObject(ServerObjects[i]);
+	}
+}
+
+/** Test global dirty tracking with multiple repsystems and a late joining client */
+UE_NET_TEST_FIXTURE(FMultiReplicationSystemsTestFixture, GlobalDirtyTrackerLateJoinClient)
+{
+	constexpr int32 NumServers = 2;
+	constexpr int32 NumServersWithClients = NumServers - 1;
+	CreateServers(NumServers);
+
+	TArrayView<FReplicationSystemTestServer*> AllServers = GetAllServers();
+	UE_NET_ASSERT_EQ(AllServers.Num(), NumServers);
+
+	TStaticArray<UTestReplicatedIrisObject*, NumServers> ServerObjects = {};
+
+	// Create client for one server and replicated objects for all servers
+	FReplicationSystemTestClient* Server0Client = CreateClientForServer(AllServers[0]);
+
+	for (uint32 i = 0; i < NumServers; ++i)
+	{
+		// Spawn object on server that won't be polled automatically multiple times during this test
+		UObjectReplicationBridge::FRootObjectReplicationParams Params;
+		Params.PollFrequency = 0.001f;
+		Params.bUseClassConfigDynamicFilter = true;
+		Params.bNeedsPreUpdate = true;
+
+		ServerObjects[i] = AllServers[i]->CreateObject(Params);
+	}
+
+	// Send and deliver packets
+	FullSendAndDeliverUpdateTwoPass();
+
+	// Object should have been created on the client
+	UTestReplicatedIrisObject* Client0Object = Cast<UTestReplicatedIrisObject>(Server0Client->GetReplicationBridge()->GetReplicatedObject(ServerObjects[0]->NetRefHandle));
+	UE_NET_ASSERT_NE(Client0Object, nullptr);
+
+	// Set a replicated variable, but don't mark it dirty
+	for (UTestReplicatedIrisObject* ServerObject : ServerObjects)
+	{
+		ServerObject->IntA = 0xFF;
+	}
+
+	// Send and deliver packets
+	FullSendAndDeliverUpdateTwoPass();
+
+	// Client replicated property should not have changed
+	UE_NET_ASSERT_NE(Client0Object->IntA, ServerObjects[0]->IntA);
+
+	// Now mark the object dirty
+	for (UTestReplicatedIrisObject* ServerObject : ServerObjects)
+	{
+		FGlobalDirtyNetObjectTracker::MarkNetObjectStateDirty(FNetHandleManager::GetOrCreateNetHandle(ServerObject));
+	}
+
+	// Send and deliver packets, and reset global dirty tracker
+	FullSendAndDeliverUpdateTwoPass();
+
+	// Add client to second server/repsystem
+	FReplicationSystemTestClient* Server1Client = CreateClientForServer(AllServers[1]);
+
+	// Send and deliver packets
+	FullSendAndDeliverUpdateTwoPass();
+
+	UTestReplicatedIrisObject* Client1Object = Cast<UTestReplicatedIrisObject>(Server1Client->GetReplicationBridge()->GetReplicatedObject(ServerObjects[1]->NetRefHandle));
+	UE_NET_ASSERT_NE(Client1Object, nullptr);
+	UE_NET_ASSERT_EQ(Client1Object->IntA, ServerObjects[1]->IntA);
+
+	// Client replicated properties should have changed now
+	for (int32 i = 0; i < NumServers; ++i)
+	{
+		for (FReplicationSystemTestClient* Client : GetClients(AllServers[i]))
+		{
+			UTestReplicatedIrisObject* ClientObject = Cast<UTestReplicatedIrisObject>(Client->GetReplicationBridge()->GetReplicatedObject(ServerObjects[i]->NetRefHandle));
+			UE_NET_ASSERT_EQ(ClientObject->IntA, ServerObjects[i]->IntA);
+		}
+
+		AllServers[i]->DestroyObject(ServerObjects[i]);
+	}
+}
+
+/** Test that validates that a push model enabled object is marked as dirty inside PreUpdate/PreReplication with multiple repsystems */
+UE_NET_TEST_FIXTURE(FMultiReplicationSystemsTestFixture, GlobalDirtyTrackerInsidePreUpdate)
+{
+	constexpr int32 NumServers = 2;
+	CreateServers(NumServers);
+
+	TArrayView<FReplicationSystemTestServer*> AllServers = GetAllServers();
+	UE_NET_ASSERT_EQ(AllServers.Num(), NumServers);
+
+	TStaticArray<UTestReplicatedIrisObject*, NumServers> ServerObjects = {};
+
+	// Create a client and an object for each server
+	for (uint32 i = 0; i < NumServers; ++i)
+	{
+		CreateClientForServer(AllServers[i]);
+
+		// Spawn object on server that is polled every frame
+		UObjectReplicationBridge::FRootObjectReplicationParams Params;
+		Params.bNeedsPreUpdate = true;
+		UTestReplicatedIrisObject::FComponents ComponentsToCreate = { .ObjectReferenceComponentCount = 1 };
+
+		ServerObjects[i] = AllServers[i]->CreateObject(Params, &ComponentsToCreate);
+	}
+
+	// Send and deliver packet
+	FullSendAndDeliverUpdateTwoPass();
+
+	// Objects should have been created on the clients
+	for (int32 i = 0; i < NumServers; ++i)
+	{
+		for (FReplicationSystemTestClient* Client : GetClients(AllServers[i]))
+		{
+			UTestReplicatedIrisObject* ClientObject = Cast<UTestReplicatedIrisObject>(Client->GetReplicationBridge()->GetReplicatedObject(ServerObjects[i]->NetRefHandle));
+			UE_NET_ASSERT_NE(ClientObject, nullptr);
+		}
+	}
+
+	auto PreUpdateObjectForServer = [&](int32 ServerIndex)
+		{
+			return [&](TArrayView<UObject*> Instances, const UReplicationBridge* Bridge)
+				{
+					for (UObject* InObject : Instances)
+					{
+						if (InObject == ServerObjects[ServerIndex])
+						{
+							// Dirty the object in the global dirty tracker
+							ServerObjects[ServerIndex]->ObjectReferenceComponents[0]->ModifyIntA();
+						}
+					}
+				};
+		};
+
+	// Mark a property dirty on the first server during PreUpdate. As the object is polled every frame we expect the property to be updated on the client.
+	AllServers[0]->GetReplicationBridge()->SetExternalPreUpdateFunctor(PreUpdateObjectForServer(0));
+
+	// Send and deliver packet
+	FullSendAndDeliverUpdateTwoPass();
+
+	// The property should be updated on the client
+	for (int32 i = 0; i < NumServers; ++i)
+	{
+		for (FReplicationSystemTestClient* Client : GetClients(AllServers[i]))
+		{
+			UTestReplicatedIrisObject* ClientObject = Cast<UTestReplicatedIrisObject>(Client->GetReplicationBridge()->GetReplicatedObject(ServerObjects[i]->NetRefHandle));
+			UE_NET_ASSERT_EQ(ClientObject->ObjectReferenceComponents[0]->IntA, ServerObjects[i]->ObjectReferenceComponents[0]->IntA);
+		}
+	}
+
+	// Now use the other server/repsystem to mark dirty
+	AllServers[0]->GetReplicationBridge()->SetExternalPreUpdateFunctor([](TArrayView<UObject*> Instances, const UReplicationBridge* Bridge) {});
+	AllServers[1]->GetReplicationBridge()->SetExternalPreUpdateFunctor(PreUpdateObjectForServer(1));
+
+	// Send and deliver packet
+	FullSendAndDeliverUpdateTwoPass();
+
+	// The property should still be equal on the client of server 0
+	FReplicationSystemTestClient* Client0 = GetClients(AllServers[0])[0];
+	UTestReplicatedIrisObject* ClientObject0 = Cast<UTestReplicatedIrisObject>(Client0->GetReplicationBridge()->GetReplicatedObject(ServerObjects[0]->NetRefHandle));
+	UE_NET_ASSERT_EQ(ClientObject0->ObjectReferenceComponents[0]->IntA, ServerObjects[0]->ObjectReferenceComponents[0]->IntA);
+
+	// The property should be updated on the client of server 1 now since it re-polls dirtiness after the PreUpdate
+	FReplicationSystemTestClient* Client1 = GetClients(AllServers[1])[0];
+	UTestReplicatedIrisObject* ClientObject1 = Cast<UTestReplicatedIrisObject>(Client1->GetReplicationBridge()->GetReplicatedObject(ServerObjects[1]->NetRefHandle));
+	UE_NET_ASSERT_EQ(ClientObject1->ObjectReferenceComponents[0]->IntA, ServerObjects[1]->ObjectReferenceComponents[0]->IntA);
+}
+
+/** Test that validates that a push model enabled object is marked as dirty inside PreUpdate/PreReplication of a different replication system */
+UE_NET_TEST_FIXTURE(FMultiReplicationSystemsTestFixture, GlobalDirtyTrackerInsidePreUpdateCrossSystems)
+{
+	constexpr int32 NumServers = 2;
+	CreateServers(NumServers);
+
+	TArrayView<FReplicationSystemTestServer*> AllServers = GetAllServers();
+	UE_NET_ASSERT_EQ(AllServers.Num(), NumServers);
+
+	TStaticArray<UTestReplicatedIrisObject*, NumServers> ServerObjects = {};
+
+	// Create a client and an object for each server
+	for (uint32 i = 0; i < NumServers; ++i)
+	{
+		CreateClientForServer(AllServers[i]);
+
+		// Spawn object on server that is polled every frame
+		UObjectReplicationBridge::FRootObjectReplicationParams Params;
+		Params.bNeedsPreUpdate = true;
+		UTestReplicatedIrisObject::FComponents ComponentsToCreate = { .ObjectReferenceComponentCount = 1 };
+
+		ServerObjects[i] = AllServers[i]->CreateObject(Params, &ComponentsToCreate);
+	}
+
+	// Send and deliver packet
+	FullSendAndDeliverUpdateTwoPass();
+
+	// Objects should have been created on the clients
+	for (int32 i = 0; i < NumServers; ++i)
+	{
+		for (FReplicationSystemTestClient* Client : GetClients(AllServers[i]))
+		{
+			UTestReplicatedIrisObject* ClientObject = Cast<UTestReplicatedIrisObject>(Client->GetReplicationBridge()->GetReplicatedObject(ServerObjects[i]->NetRefHandle));
+			UE_NET_ASSERT_NE(ClientObject, nullptr);
+		}
+	}
+
+	auto PreUpdateObjectForServer = [&](int32 ServerIndex)
+		{
+			return [&](TArrayView<UObject*> Instances, const UReplicationBridge* Bridge)
+				{
+					// Dirty the object in the global dirty tracker
+					ServerObjects[ServerIndex]->ObjectReferenceComponents[0]->ModifyIntA();
+				};
+		};
+
+	// On server 0's PreUpdate, mark server 1's object dirty.
+	AllServers[0]->GetReplicationBridge()->SetExternalPreUpdateFunctor(PreUpdateObjectForServer(1));
+
+	// Send and deliver packet
+	FullSendAndDeliverUpdateTwoPass();
+
+	// The property should be updated on the client
+	for (int32 i = 0; i < NumServers; ++i)
+	{
+		for (FReplicationSystemTestClient* Client : GetClients(AllServers[i]))
+		{
+			UTestReplicatedIrisObject* ClientObject = Cast<UTestReplicatedIrisObject>(Client->GetReplicationBridge()->GetReplicatedObject(ServerObjects[i]->NetRefHandle));
+			UE_NET_ASSERT_EQ(ClientObject->ObjectReferenceComponents[0]->IntA, ServerObjects[i]->ObjectReferenceComponents[0]->IntA);
+		}
+	}
+
+	// This time on server 1's PreUpdate, mark server 0's object dirty.
+	AllServers[0]->GetReplicationBridge()->SetExternalPreUpdateFunctor([](TArrayView<UObject*> Instances, const UReplicationBridge* Bridge) {});
+	AllServers[1]->GetReplicationBridge()->SetExternalPreUpdateFunctor(PreUpdateObjectForServer(0));
+
+	// Send and deliver packet
+	FullSendAndDeliverUpdateTwoPass();
+
+	// Since ServerObjects[0] was marked dirty after server 0 updated, the client's object should not have changed yet.
+	FReplicationSystemTestClient* Client0 = GetClients(AllServers[0])[0];
+	UTestReplicatedIrisObject* ClientObject0 = Cast<UTestReplicatedIrisObject>(Client0->GetReplicationBridge()->GetReplicatedObject(ServerObjects[0]->NetRefHandle));
+	UE_NET_ASSERT_NE(ClientObject0->ObjectReferenceComponents[0]->IntA, ServerObjects[0]->ObjectReferenceComponents[0]->IntA);
+
+	// Client object 1 should still match server 1
+	FReplicationSystemTestClient* Client1 = GetClients(AllServers[1])[0];
+	UTestReplicatedIrisObject* ClientObject1 = Cast<UTestReplicatedIrisObject>(Client1->GetReplicationBridge()->GetReplicatedObject(ServerObjects[1]->NetRefHandle));
+	UE_NET_ASSERT_EQ(ClientObject1->ObjectReferenceComponents[0]->IntA, ServerObjects[1]->ObjectReferenceComponents[0]->IntA);
+
+	// Clear the PreUpdate functors
+	AllServers[0]->GetReplicationBridge()->SetExternalPreUpdateFunctor([](TArrayView<UObject*> Instances, const UReplicationBridge* Bridge) {});
+	AllServers[1]->GetReplicationBridge()->SetExternalPreUpdateFunctor([](TArrayView<UObject*> Instances, const UReplicationBridge* Bridge) {});
+
+	// Send and deliver packet. This should update client 0 to have the latest server value.
+	FullSendAndDeliverUpdateTwoPass();
+
+	// Both clients should match their servers now
+	UE_NET_ASSERT_EQ(ClientObject0->ObjectReferenceComponents[0]->IntA, ServerObjects[0]->ObjectReferenceComponents[0]->IntA);
+	UE_NET_ASSERT_EQ(ClientObject1->ObjectReferenceComponents[0]->IntA, ServerObjects[1]->ObjectReferenceComponents[0]->IntA);
+
+
+	// This time both servers mark each other's object dirty.
+	AllServers[0]->GetReplicationBridge()->SetExternalPreUpdateFunctor(PreUpdateObjectForServer(1));
+	AllServers[1]->GetReplicationBridge()->SetExternalPreUpdateFunctor(PreUpdateObjectForServer(0));
+
+	// Send and deliver packet
+	FullSendAndDeliverUpdateTwoPass();
+
+	// Since ServerObjects[0] was marked dirty after server 0 updated, the client's object should not have changed yet.
+	UE_NET_ASSERT_NE(ClientObject0->ObjectReferenceComponents[0]->IntA, ServerObjects[0]->ObjectReferenceComponents[0]->IntA);
+
+	// Since ServerObjects[1] was marked dirty before server 1 updated, the client object 1 should be updated.
+	UE_NET_ASSERT_EQ(ClientObject1->ObjectReferenceComponents[0]->IntA, ServerObjects[1]->ObjectReferenceComponents[0]->IntA);
+
+	// Clear the PreUpdate functors
+	AllServers[0]->GetReplicationBridge()->SetExternalPreUpdateFunctor([](TArrayView<UObject*> Instances, const UReplicationBridge* Bridge) {});
+	AllServers[1]->GetReplicationBridge()->SetExternalPreUpdateFunctor([](TArrayView<UObject*> Instances, const UReplicationBridge* Bridge) {});
+
+	// Send and deliver packet. This should update client 0 to have the latest server value.
+	FullSendAndDeliverUpdateTwoPass();
+
+	// Both clients should match their servers now
+	UE_NET_ASSERT_EQ(ClientObject0->ObjectReferenceComponents[0]->IntA, ServerObjects[0]->ObjectReferenceComponents[0]->IntA);
+	UE_NET_ASSERT_EQ(ClientObject1->ObjectReferenceComponents[0]->IntA, ServerObjects[1]->ObjectReferenceComponents[0]->IntA);
+
+}
+
+/** Test that the same object in separate systems is dirtied & replicated properly */
+UE_NET_TEST_FIXTURE(FMultiReplicationSystemsTestFixture, GlobalDirtyTracker_SameObject)
+{
+	constexpr int32 NumServers = 3;
+	CreateServers(NumServers);
+
+	TArrayView<FReplicationSystemTestServer*> AllServers = GetAllServers();
+	UE_NET_ASSERT_EQ(AllServers.Num(), NumServers);
+
+	TStaticArray<FNetRefHandle, NumServers> RepSystemRefHandles = {};
+	
+	// Spawn object on server that won't be polled automatically multiple times during this test
+	UObjectReplicationBridge::FRootObjectReplicationParams Params;
+	Params.PollFrequency = 0.001f;
+	UTestReplicatedIrisObject* ServerObject = CreateObject({});
+
+	// Create clients for all systems
+	for (uint32 i = 0; i < NumServers; ++i)
+	{
+		constexpr unsigned ClientCountPerServer = 3U;
+		for (unsigned ClientIt = 0, ClientEndIt = ClientCountPerServer; ClientIt != ClientEndIt; ++ClientIt)
+		{
+			CreateClientForServer(AllServers[i]);
+		}
+
+		RepSystemRefHandles[i] = BeginReplication(AllServers[i], ServerObject, Params);
+	}
+
+	// Send and deliver packets
+	FullSendAndDeliverUpdateTwoPass();
+
+	// Objects should have been created on the clients
+	for (int32 i = 0; i < NumServers; ++i)
+	{
+		for (FReplicationSystemTestClient* Client : GetClients(AllServers[i]))
+		{
+			UTestReplicatedIrisObject* ClientObject = Cast<UTestReplicatedIrisObject>(Client->GetReplicationBridge()->GetReplicatedObject(RepSystemRefHandles[i]));
+			UE_NET_ASSERT_NE(ClientObject, nullptr);
+		}
+	}
+
+	// Set a replicated variable, but don't mark it dirty
+	ServerObject->IntA = 0xFF;
+
+	// Send and deliver packets
+	FullSendAndDeliverUpdateTwoPass();
+
+	// Client replicated property should not have changed
+	for (int32 i = 0; i < NumServers; ++i)
+	{
+		for (FReplicationSystemTestClient* Client : GetClients(AllServers[i]))
+		{
+			UTestReplicatedIrisObject* ClientObject = Cast<UTestReplicatedIrisObject>(Client->GetReplicationBridge()->GetReplicatedObject(RepSystemRefHandles[i]));
+			UE_NET_ASSERT_NE(ClientObject->IntA, ServerObject->IntA);
+		}
+	}
+
+	// Now mark the object dirty
+	FGlobalDirtyNetObjectTracker::MarkNetObjectStateDirty(FNetHandleManager::GetOrCreateNetHandle(ServerObject));
+
+	// Send and deliver packets
+	FullSendAndDeliverUpdateTwoPass();
+
+	// Client replicated properties should have changed now
+	for (int32 i = 0; i < NumServers; ++i)
+	{
+		for (FReplicationSystemTestClient* Client : GetClients(AllServers[i]))
+		{
+			UTestReplicatedIrisObject* ClientObject = Cast<UTestReplicatedIrisObject>(Client->GetReplicationBridge()->GetReplicatedObject(RepSystemRefHandles[i]));
+			UE_NET_ASSERT_EQ(ClientObject->IntA, ServerObject->IntA);
+		}
+	}
+}
+
+/** Test that validates that a single push model enabled object replicated by multiple systems marked dirty inside one PreUpdate/PreReplication replicates properly */
+UE_NET_TEST_FIXTURE(FMultiReplicationSystemsTestFixture, GlobalDirtyTrackerInsidePreUpdate_SameObject)
+{
+	constexpr int32 NumServers = 2;
+	CreateServers(NumServers);
+
+	TArrayView<FReplicationSystemTestServer*> AllServers = GetAllServers();
+	UE_NET_ASSERT_EQ(AllServers.Num(), NumServers);
+
+	TStaticArray<FNetRefHandle, NumServers> RepSystemRefHandles = {};
+
+	UTestReplicatedIrisObject* ServerObject = CreateObject({ .ObjectReferenceComponentCount = 1 });
+
+	// Create a client for each server
+	for (uint32 i = 0; i < NumServers; ++i)
+	{
+		CreateClientForServer(AllServers[i]);
+
+		// Spawn object on server that is polled every frame
+		UObjectReplicationBridge::FRootObjectReplicationParams Params;
+		Params.bNeedsPreUpdate = true;
+
+		RepSystemRefHandles[i] = BeginReplication(AllServers[i], ServerObject, Params);
+	}
+
+	// Send and deliver packet
+	FullSendAndDeliverUpdateTwoPass();
+
+	// Objects should have been created on the clients
+	for (int32 i = 0; i < NumServers; ++i)
+	{
+		for (FReplicationSystemTestClient* Client : GetClients(AllServers[i]))
+		{
+			UTestReplicatedIrisObject* ClientObject = Cast<UTestReplicatedIrisObject>(Client->GetReplicationBridge()->GetReplicatedObject(RepSystemRefHandles[i]));
+			UE_NET_ASSERT_NE(ClientObject, nullptr);
+		}
+	}
+
+	auto PreUpdateObject = [&](TArrayView<UObject*> Instances, const UReplicationBridge* Bridge)
+		{
+			for (UObject* InObject : Instances)
+			{
+				if (InObject == ServerObject)
+				{
+					// Dirty the object in the global dirty tracker
+					ServerObject->ObjectReferenceComponents[0]->ModifyIntA();
+				}
+			}
+		};
+
+	// Mark a property dirty on the first server during PreUpdate. As the object is polled every frame we expect the property to be updated on the client.
+	AllServers[0]->GetReplicationBridge()->SetExternalPreUpdateFunctor(PreUpdateObject);
+
+	// Send and deliver packet
+	FullSendAndDeliverUpdateTwoPass();
+
+	// The property should be updated on the clients
+	FReplicationSystemTestClient* Client0 = GetClients(AllServers[0])[0];
+	UTestReplicatedIrisObject* ClientObject0 = Cast<UTestReplicatedIrisObject>(Client0->GetReplicationBridge()->GetReplicatedObject(RepSystemRefHandles[0]));
+	UE_NET_ASSERT_EQ(ClientObject0->ObjectReferenceComponents[0]->IntA, ServerObject->ObjectReferenceComponents[0]->IntA);
+
+	FReplicationSystemTestClient* Client1 = GetClients(AllServers[1])[0];
+	UTestReplicatedIrisObject* ClientObject1 = Cast<UTestReplicatedIrisObject>(Client1->GetReplicationBridge()->GetReplicatedObject(RepSystemRefHandles[1]));
+	UE_NET_ASSERT_EQ(ClientObject1->ObjectReferenceComponents[0]->IntA, ServerObject->ObjectReferenceComponents[0]->IntA);
+
+
+	// Now use the other server/repsystem to mark dirty
+	AllServers[0]->GetReplicationBridge()->SetExternalPreUpdateFunctor([](TArrayView<UObject*> Instances, const UReplicationBridge* Bridge) {});
+	AllServers[1]->GetReplicationBridge()->SetExternalPreUpdateFunctor(PreUpdateObject);
+
+	// Send and deliver packet
+	FullSendAndDeliverUpdateTwoPass();
+
+	// The property on client 0 won't be updated yet because the server only updated the property when rep system 1 updated (after replication to client 0)
+	UE_NET_ASSERT_NE(ClientObject0->ObjectReferenceComponents[0]->IntA, ServerObject->ObjectReferenceComponents[0]->IntA);
+
+	// The property should be updated on the client of server 1 now since it re-polls dirtiness after the PreUpdate
+	UE_NET_ASSERT_EQ(ClientObject1->ObjectReferenceComponents[0]->IntA, ServerObject->ObjectReferenceComponents[0]->IntA);
+
+	// Clear the PreUpdate functors
+	AllServers[0]->GetReplicationBridge()->SetExternalPreUpdateFunctor([](TArrayView<UObject*> Instances, const UReplicationBridge* Bridge) {});
+	AllServers[1]->GetReplicationBridge()->SetExternalPreUpdateFunctor([](TArrayView<UObject*> Instances, const UReplicationBridge* Bridge) {});
+
+	// Send and deliver packet. This should update client 0 to have the server value.
+	FullSendAndDeliverUpdateTwoPass();
+
+	// The property should be updated on all clients now
+	UE_NET_ASSERT_EQ(ClientObject0->ObjectReferenceComponents[0]->IntA, ServerObject->ObjectReferenceComponents[0]->IntA);
+	UE_NET_ASSERT_EQ(ClientObject1->ObjectReferenceComponents[0]->IntA, ServerObject->ObjectReferenceComponents[0]->IntA);
+
+
+	// Now have both systems mark the object dirty in each PreUpdate
+	const int32 ServerValuePreUpdate = ServerObject->ObjectReferenceComponents[0]->IntA;
+	AllServers[0]->GetReplicationBridge()->SetExternalPreUpdateFunctor(PreUpdateObject);
+	AllServers[1]->GetReplicationBridge()->SetExternalPreUpdateFunctor(PreUpdateObject);
+
+	// Send and deliver packet
+	FullSendAndDeliverUpdateTwoPass();
+
+	// The property on client 0 should be the server value + 1 since one PreUpdate ran before replicating to this client.
+	UE_NET_ASSERT_EQ(ClientObject0->ObjectReferenceComponents[0]->IntA, ServerValuePreUpdate + 1);
+
+	// The property on client 1 should be the server since both PreUpdates ran before replicating to this client.
+	UE_NET_ASSERT_EQ(ClientObject1->ObjectReferenceComponents[0]->IntA, ServerObject->ObjectReferenceComponents[0]->IntA);
+
+	// Clear the PreUpdate functors
+	AllServers[0]->GetReplicationBridge()->SetExternalPreUpdateFunctor([](TArrayView<UObject*> Instances, const UReplicationBridge* Bridge) {});
+	AllServers[1]->GetReplicationBridge()->SetExternalPreUpdateFunctor([](TArrayView<UObject*> Instances, const UReplicationBridge* Bridge) {});
+
+	// Send and deliver packet. This should update client 0 to have the latest server value.
+	FullSendAndDeliverUpdateTwoPass();
+
+	// The property should be updated on all clients now
+	UE_NET_ASSERT_EQ(ClientObject0->ObjectReferenceComponents[0]->IntA, ServerObject->ObjectReferenceComponents[0]->IntA);
+	UE_NET_ASSERT_EQ(ClientObject1->ObjectReferenceComponents[0]->IntA, ServerObject->ObjectReferenceComponents[0]->IntA);
+}
+
 }
diff --git a/Engine/Source/Runtime/Experimental/Iris/Core/Private/Iris/ReplicationSystem/DirtyNetObjectTracker.cpp b/Engine/Source/Runtime/Experimental/Iris/Core/Private/Iris/ReplicationSystem/DirtyNetObjectTracker.cpp
--- a/Engine/Source/Runtime/Experimental/Iris/Core/Private/Iris/ReplicationSystem/DirtyNetObjectTracker.cpp
+++ b/Engine/Source/Runtime/Experimental/Iris/Core/Private/Iris/ReplicationSystem/DirtyNetObjectTracker.cpp
@@ -32,7 +32,7 @@ void FDirtyNetObjectTracker::Init(const FDirtyNetObjectTrackerInitParams& Params
 	
 	NetObjectIdCount = Params.MaxInternalNetRefIndex;
 
-	GlobalDirtyTrackerPollHandle = FGlobalDirtyNetObjectTracker::CreatePoller();
+	GlobalDirtyTrackerPollHandle = FGlobalDirtyNetObjectTracker::CreatePoller(FGlobalDirtyNetObjectTracker::FPreResetDelegate::CreateRaw(this, &FDirtyNetObjectTracker::ApplyGlobalDirtyObjectList));
 
 	SetNetObjectListsSize(Params.MaxInternalNetRefIndex);
 
@@ -63,7 +63,7 @@ void FDirtyNetObjectTracker::OnMaxInternalNetRefIndexIncreased(FInternalNetRefIn
 	NetObjectIdCount = NewMaxInternalIndex;
 }
 
-void FDirtyNetObjectTracker::GrabAndApplyGlobalDirtyObjectList()
+void FDirtyNetObjectTracker::ApplyGlobalDirtyObjectList()
 {
 	{
 		const TSet<FNetHandle>& GlobalDirtyNetObjects = FGlobalDirtyNetObjectTracker::GetDirtyNetObjects(GlobalDirtyTrackerPollHandle);
@@ -76,6 +76,11 @@ void FDirtyNetObjectTracker::GrabAndApplyGlobalDirtyObjectList()
 			}
 		}
 	}
+}
+
+void FDirtyNetObjectTracker::ApplyAndTryResetGlobalDirtyObjectList()
+{
+	ApplyGlobalDirtyObjectList();
 
 	FGlobalDirtyNetObjectTracker::ResetDirtyNetObjectsIfSinglePoller(GlobalDirtyTrackerPollHandle);
 
@@ -93,7 +98,7 @@ void FDirtyNetObjectTracker::UpdateDirtyNetObjects()
 
 	LockExternalAccess();
 
-	GrabAndApplyGlobalDirtyObjectList();
+	ApplyAndTryResetGlobalDirtyObjectList();
 
 	//$IRIS TODO:  We could look if any objects where actually in the global list and skip the array iteration if not needed.
 
diff --git a/Engine/Source/Runtime/Experimental/Iris/Core/Private/Iris/ReplicationSystem/DirtyNetObjectTracker.h b/Engine/Source/Runtime/Experimental/Iris/Core/Private/Iris/ReplicationSystem/DirtyNetObjectTracker.h
--- a/Engine/Source/Runtime/Experimental/Iris/Core/Private/Iris/ReplicationSystem/DirtyNetObjectTracker.h
+++ b/Engine/Source/Runtime/Experimental/Iris/Core/Private/Iris/ReplicationSystem/DirtyNetObjectTracker.h
@@ -97,7 +97,17 @@ private:
 
 	void MarkNetObjectDirty(FInternalNetRefIndex NetObjectIndex);
 	void ForceNetUpdate(FInternalNetRefIndex NetObjectIndex);
-	void GrabAndApplyGlobalDirtyObjectList();
+	void ApplyGlobalDirtyObjectList();
+
+	/**
+	 * Applies the dirty state from the global dirty object tracker.
+	 * Then if this is the only poller of global dirty state, resets the global dirty state.
+	 * If there are multiple pollers of global dirty state (multiple replication systems),
+	 * the global state can't be reset until all pollers have gathered it. So we set the
+	 * bShouldResetPolledGlobalDirtyTracker flag which will attempt another reset in ReconcilePolledList,
+	 * which is called in PostSendUpdate after other pollers have had a chance to gather.
+	 */
+	void ApplyAndTryResetGlobalDirtyObjectList();
 
 	/** Can only be accessed via FDirtyObjectsAccessor */
 	FNetBitArrayView GetDirtyNetObjectsThisFrame();
diff --git a/Engine/Source/Runtime/Net/Core/Private/Net/Core/DirtyNetObjectTracker/GlobalDirtyNetObjectTracker.cpp b/Engine/Source/Runtime/Net/Core/Private/Net/Core/DirtyNetObjectTracker/GlobalDirtyNetObjectTracker.cpp
--- a/Engine/Source/Runtime/Net/Core/Private/Net/Core/DirtyNetObjectTracker/GlobalDirtyNetObjectTracker.cpp
+++ b/Engine/Source/Runtime/Net/Core/Private/Net/Core/DirtyNetObjectTracker/GlobalDirtyNetObjectTracker.cpp
@@ -9,12 +9,6 @@ namespace UE::Net
 
 class FGlobalDirtyNetObjectTracker::FPimpl
 {
-private:
-	FPimpl();
-	~FPimpl();
-
-	void ResetDirtyNetObjects();
-
 private:
 	friend FGlobalDirtyNetObjectTracker;
 
@@ -24,13 +18,15 @@ private:
 
 	struct FPollerStatus
 	{
+		FPreResetDelegate PreResetDelegate;
+
 		// Is this status tied to an active registered poller
 		uint8 bIsActive:1;
 
-		// Has this poller read the dirty list this frame yet.
-		uint8 bHasGathered:1;
+		// Does this poller need to read the dirty list this frame.
+		uint8 bNeedsGather:1;
 
-		FPollerStatus() : bIsActive(false), bHasGathered(false) {}
+		FPollerStatus() : bIsActive(false), bNeedsGather(true) {}
 
 		void ClearStatus()
 		{
@@ -64,11 +60,22 @@ void FGlobalDirtyNetObjectTracker::MarkNetObjectStateDirty(FNetHandle NetHandle)
 		if (ensureMsgf(!Instance->bLockDirtyList, TEXT("MarkNetObjectStateDirty was called while the dirty list was set to read-only.")))
 		{
 			Instance->DirtyObjects.Add(NetHandle);
+
+			// If there are multiple pollers, any poller that has already gathered this frame will need to gather again
+			// to know about this newly dirtied object, since it may have happened in another replication system's PreUpdate.
+			// Setting bNeedsGather = true will cause it to gather during the last-chance delegate in ResetDirtyNetObjects.
+			if (Instance->PollerCount > 1)
+			{
+				Instance->Pollers.ForAllSetBits([](uint32 BitIndex)
+					{
+						Instance->PollerStatuses[BitIndex].bNeedsGather = true;
+					});
+			}
 		}
 	}
 }
 
-FGlobalDirtyNetObjectTracker::FPollHandle FGlobalDirtyNetObjectTracker::CreatePoller()
+FGlobalDirtyNetObjectTracker::FPollHandle FGlobalDirtyNetObjectTracker::CreatePoller(FPreResetDelegate InPreResetDelegate)
 {
 	if (Instance)
 	{
@@ -89,6 +96,7 @@ FGlobalDirtyNetObjectTracker::FPollHandle FGlobalDirtyNetObjectTracker::CreatePo
 
 		Instance->PollerStatuses.SetNum(Instance->PollerCount, EAllowShrinking::No);
 		Instance->PollerStatuses[HandleIndex].bIsActive = true;
+		Instance->PollerStatuses[HandleIndex].PreResetDelegate = InPreResetDelegate;
 
 		return FPollHandle(HandleIndex);
 	}
@@ -129,7 +137,7 @@ const TSet<FNetHandle>& FGlobalDirtyNetObjectTracker::GetDirtyNetObjects(const F
 		check(Instance->PollerStatuses[Handle.Index].bIsActive);
 
 		Instance->Pollers.SetBit(Handle.Index);
-		Instance->PollerStatuses[Handle.Index].bHasGathered = true;
+		Instance->PollerStatuses[Handle.Index].bNeedsGather = false;
 		return Instance->DirtyObjects;
 	}
 
@@ -157,19 +165,24 @@ void FGlobalDirtyNetObjectTracker::ResetDirtyNetObjects(const FPollHandle& Handl
 
 		if (Instance->Pollers.IsNoBitSet())
 		{
-
-#if DO_CHECK
-			bool bAllPollersGathered = true;
-			for (const FPimpl::FPollerStatus& PollerStatus : Instance->PollerStatuses)
+			for (uint32 PollerIndex = 0; PollerIndex < Instance->PollerCount; ++PollerIndex)
 			{
-				bAllPollersGathered &= !PollerStatus.bIsActive || PollerStatus.bHasGathered;
+				FPimpl::FPollerStatus& PollerStatus = Instance->PollerStatuses[PollerIndex];
+
+				if (PollerStatus.bIsActive && PollerStatus.bNeedsGather)
+				{
+					PollerStatus.PreResetDelegate.ExecuteIfBound();
+					Instance->Pollers.ClearBit(PollerIndex);
+				}
+
+				// Poller will need to call GetDirtyNetObjects next frame
+				PollerStatus.bNeedsGather = true;
 			}
-			ensureMsgf(bAllPollersGathered, TEXT("Not all pollers gathered the dirty list before the list got reset. Those pollers will never know those objects were dirty."));
-#endif
 
-			Instance->bLockDirtyList = false;
 			Instance->DirtyObjects.Reset();
 		}
+
+		Instance->bLockDirtyList = false;
 	}
 }
 
@@ -179,12 +192,13 @@ bool FGlobalDirtyNetObjectTracker::ResetDirtyNetObjectsIfSinglePoller(const FPol
 	{
 		check(Instance->PollerStatuses[Handle.Index].bIsActive);
 
+		Instance->bLockDirtyList = false;
+
 		// Are we the only poller registered to read and reset the list
 		if (Instance->AssignedHandleIndices.CountSetBits() == 1)
 		{
 			check(Instance->Pollers.IsBitSet(Handle.Index));
 
-			Instance->bLockDirtyList = false;
 			Instance->DirtyObjects.Reset();
 
 			Instance->Pollers.ClearBit(Handle.Index);
@@ -208,31 +222,4 @@ void FGlobalDirtyNetObjectTracker::Deinit()
 	Instance = nullptr;
 }
 
-FGlobalDirtyNetObjectTracker::FPimpl::FPimpl()
-{
-#if WITH_ENGINE
-	FCoreDelegates::OnEndFrame.AddRaw(this, &FPimpl::ResetDirtyNetObjects);
-#endif
-}
-
-FGlobalDirtyNetObjectTracker::FPimpl::~FPimpl()
-{
-#if WITH_ENGINE
-	FCoreDelegates::OnEndFrame.RemoveAll(this);
-#endif
-}
-
-void FGlobalDirtyNetObjectTracker::FPimpl::ResetDirtyNetObjects()
-{
-	if (!ensureMsgf(Pollers.IsNoBitSet(), TEXT("FGlobalDirtyNetObjectTracker poller %u forgot to call ResetDirtyNetObjects."), Pollers.FindFirstOne()))
-	{
-		Pollers.ClearAllBits();
-
-		// DirtyObjects should already be reset if there are no pollers.
-		DirtyObjects.Reset();
-
-		bLockDirtyList = false;
-	}
-}
-
 }
diff --git a/Engine/Source/Runtime/Net/Core/Public/Net/Core/DirtyNetObjectTracker/GlobalDirtyNetObjectTracker.h b/Engine/Source/Runtime/Net/Core/Public/Net/Core/DirtyNetObjectTracker/GlobalDirtyNetObjectTracker.h
--- a/Engine/Source/Runtime/Net/Core/Public/Net/Core/DirtyNetObjectTracker/GlobalDirtyNetObjectTracker.h
+++ b/Engine/Source/Runtime/Net/Core/Public/Net/Core/DirtyNetObjectTracker/GlobalDirtyNetObjectTracker.h
@@ -50,8 +50,11 @@ public:
 	/** Marks an object as dirty. Assumes Init() has been called. */
 	NETCORE_API static void MarkNetObjectStateDirty(FNetHandle);
 
+	/** Delegate called on pollers that haven't gathered dirty state before it resets */
+	DECLARE_DELEGATE(FPreResetDelegate);
+
 	/** Create a poller which is assumed to call PollDirtyNetObjects(). Assumes Init() has been called. */
-	NETCORE_API static FPollHandle CreatePoller();
+	NETCORE_API static FPollHandle CreatePoller(FPreResetDelegate InPreResetDelegate);
 
 	/**
 	 * Returns the objects with dirty state. The reference can only be assumed to be valid until the next call to ResetDirtyNetObjects().
