Iris - Fix compile error

#rnx

diff --git a/Engine/Plugins/Runtime/ReplicationSystemTestPlugin/Source/Private/Tests/ReplicationSystem/TestMultiReplicationSystems.cpp b/Engine/Plugins/Runtime/ReplicationSystemTestPlugin/Source/Private/Tests/ReplicationSystem/TestMultiReplicationSystems.cpp
--- a/Engine/Plugins/Runtime/ReplicationSystemTestPlugin/Source/Private/Tests/ReplicationSystem/TestMultiReplicationSystems.cpp
+++ b/Engine/Plugins/Runtime/ReplicationSystemTestPlugin/Source/Private/Tests/ReplicationSystem/TestMultiReplicationSystems.cpp
@@ -624,7 +624,7 @@ UE_NET_TEST_FIXTURE(FMultiReplicationSystemsTestFixture, GlobalDirtyTrackerInsid
 
 	auto PreUpdateObjectForServer = [&](int32 ServerIndex)
 		{
-			return [&](TArrayView<UObject*> Instances, const UReplicationBridge* Bridge)
+			return [&ServerObjects, ServerIndex](TArrayView<UObject*> Instances, const UReplicationBridge* Bridge)
 				{
 					for (UObject* InObject : Instances)
 					{
@@ -710,7 +710,7 @@ UE_NET_TEST_FIXTURE(FMultiReplicationSystemsTestFixture, GlobalDirtyTrackerInsid
 
 	auto PreUpdateObjectForServer = [&](int32 ServerIndex)
 		{
-			return [&](TArrayView<UObject*> Instances, const UReplicationBridge* Bridge)
+			return [&ServerObjects, ServerIndex](TArrayView<UObject*> Instances, const UReplicationBridge* Bridge)
 				{
 					// Dirty the object in the global dirty tracker
 					ServerObjects[ServerIndex]->ObjectReferenceComponents[0]->ModifyIntA();
