[Iris] Fix PCH issue for TestFlushBeforeDestroy.

diff --git a/Engine/Plugins/Runtime/ReplicationSystemTestPlugin/Source/Private/Tests/ReplicationSystem/TestFlushBeforeDestroy.cpp b/Engine/Plugins/Runtime/ReplicationSystemTestPlugin/Source/Private/Tests/ReplicationSystem/TestFlushBeforeDestroy.cpp
--- a/Engine/Plugins/Runtime/ReplicationSystemTestPlugin/Source/Private/Tests/ReplicationSystem/TestFlushBeforeDestroy.cpp
+++ b/Engine/Plugins/Runtime/ReplicationSystemTestPlugin/Source/Private/Tests/ReplicationSystem/TestFlushBeforeDestroy.cpp
@@ -8,6 +8,7 @@
 #include "Misc/ScopeExit.h"
 #include "NetBlob/PartialNetBlobTestFixture.h"
 #include "NetBlob/MockNetBlob.h"
+#include "UObject/CoreNetTypes.h"
 
 namespace UE::Net::Private
 {
