Iris - Added FTestMessage stream operator for TObjectPtr.

diff --git a/Engine/Plugins/Runtime/ReplicationSystemTestPlugin/Source/Private/TestMessage.h b/Engine/Plugins/Runtime/ReplicationSystemTestPlugin/Source/Private/TestMessage.h
--- a/Engine/Plugins/Runtime/ReplicationSystemTestPlugin/Source/Private/TestMessage.h
+++ b/Engine/Plugins/Runtime/ReplicationSystemTestPlugin/Source/Private/TestMessage.h
@@ -4,6 +4,7 @@
 
 #include "HAL/Platform.h"
 #include "UObject/NameTypes.h"
+#include "UObject/ObjectPtr.h"
 #include "Math/Vector.h"
 #include <cinttypes>
 
@@ -151,4 +152,10 @@ inline FTestMessage& operator<<(FTestMessage& Message, SIZE_T Value)
 }
 #endif
 
+template<typename T>
+inline FTestMessage& operator<<(FTestMessage& Message, const TObjectPtr<T>& ObjectPtr)
+{
+	return Message << TStringBuilder<512>().Appendf(TEXT("%s (%p)"), ToCStr(ObjectPtr.GetName()), ObjectPtr.Get());
+}
+
 }
