Iris
* When RootObjects are added for replication, initialize their current location and cull distance in the WorldLocations cache
* This allows Iris systems to be able to immediately poll the location and culldistance of new objects. Ex: Inside NetObjectFilter::AddObject callback

ReplicationSystemTest
* Updated grid filter test to ensure they are created with their intented location instead of starting at Zero and being assigned a location post-creation.

#rnx

diff --git a/Engine/Plugins/Runtime/ReplicationSystemTestPlugin/Source/Private/Tests/ReplicationSystem/Filtering/TestGridFilter.cpp b/Engine/Plugins/Runtime/ReplicationSystemTestPlugin/Source/Private/Tests/ReplicationSystem/Filtering/TestGridFilter.cpp
--- a/Engine/Plugins/Runtime/ReplicationSystemTestPlugin/Source/Private/Tests/ReplicationSystem/Filtering/TestGridFilter.cpp
+++ b/Engine/Plugins/Runtime/ReplicationSystemTestPlugin/Source/Private/Tests/ReplicationSystem/Filtering/TestGridFilter.cpp
@@ -92,6 +92,37 @@ private:
 		WorldLocFilterHandle = Server->GetReplicationSystem()->GetFilterHandle("NetObjectGridWorldLocFilter");
 	}
 
+
+public:
+
+	// Array that preemptively registers the world info so it is assigned to an object during his creation.
+	struct FObjectWorldInfo
+	{
+		FVector Loc = FVector::ZeroVector;
+		float CullDistance = 1500.f;
+	};
+	TArray<FObjectWorldInfo> WorldInfoToBeAssigned;
+	TMap<const UObject*, FObjectWorldInfo> ObjectWorldInfoMap;
+
+	FObjectWorldInfo GetWorldInfo(const UObject* ReplicatedObject)
+	{
+		if (ObjectWorldInfoMap.Find(ReplicatedObject) == nullptr)
+		{
+			checkf(!WorldInfoToBeAssigned.IsEmpty(), TEXT("No info was pushed for assignation"));
+			FObjectWorldInfo NewInfo = WorldInfoToBeAssigned[0];
+			WorldInfoToBeAssigned.RemoveAt(0);
+
+			ObjectWorldInfoMap.Add(ReplicatedObject, NewInfo);
+		}
+
+		return ObjectWorldInfoMap[ReplicatedObject];
+	}
+
+	void SetWorldInfo(const UObject* ReplicatedObject, const FObjectWorldInfo& WorldInfo)
+	{
+		ObjectWorldInfoMap[ReplicatedObject] = WorldInfo;
+	}
+
 protected:
 	UNetObjectGridWorldLocFilter* WorldLocFilter;
 	FNetObjectFilterHandle WorldLocFilterHandle;
@@ -102,49 +133,52 @@ private:
 
 UE_NET_TEST_FIXTURE(FTestGridFilterFixture, TestWorldLocGridFilter)
 {
+	Server->GetReplicationBridge()->SetExternalWorldLocationUpdateFunctor([&](FNetRefHandle NetHandle, const UObject* ReplicatedObject, FVector& OutLocation, float& OutCullDistance)
+	{
+		const FObjectWorldInfo& Info = GetWorldInfo(ReplicatedObject);
+		OutLocation = Info.Loc;
+		OutCullDistance = Info.CullDistance;
+	});
+
+	const UNetObjectGridFilterConfig* DefaultGridConfig = GetDefault<UNetObjectGridFilterConfig>();
+
 	// Spawn object with WorldLocation's on server
 	UObjectReplicationBridge::FRootObjectReplicationParams Params;
 	Params.bNeedsWorldLocationUpdate = true;
 	Params.bUseClassConfigDynamicFilter = true;
-
+	
+	// Relevant objects
+	WorldInfoToBeAssigned.Add(FObjectWorldInfo{ .Loc = FVector::ZeroVector, .CullDistance = 1500.f });
 	UReplicatedTestObject* ServerObjectZero = Server->CreateObject(Params);
-	UReplicatedTestObject* ServerObjectNear = Server->CreateObject(Params);
-	UReplicatedTestObject* ServerObjectLimit = Server->CreateObject(Params);
-	UReplicatedTestObject* ServerObjectCulled = Server->CreateObject(Params);
-	UReplicatedTestObject* ServerObjectVeryFar = Server->CreateObject(Params);
-
-	const UNetObjectGridFilterConfig* DefaultGridConfig = GetDefault<UNetObjectGridFilterConfig>();
-
-	struct FObjectLoc 
-	{
-		FVector Loc;
-		float CullDistance;
-	};
-	TMap<UReplicatedTestObject*, FObjectLoc> ObjectLocs;
 
-	// Visible objects
-	ObjectLocs.Add(ServerObjectZero,	FObjectLoc{ .Loc=FVector::ZeroVector, .CullDistance=1500.f} );
-	ObjectLocs.Add(ServerObjectNear,	FObjectLoc{ .Loc=FVector(100.f, 100.f, 100.f), .CullDistance=1500.f} );
+	WorldInfoToBeAssigned.Add(FObjectWorldInfo{ .Loc = FVector(100.f, 100.f, 100.f), .CullDistance = 1500.f });
+	UReplicatedTestObject* ServerObjectNear = Server->CreateObject(Params);
 
-	// Make an object that sits right in the limit of a grid cell
+	// Make a location for an object that sits right in the limit of a grid cell
 	const FVector CellLimitPos(DefaultGridConfig->CellSizeX, DefaultGridConfig->CellSizeY, 0.0f);
-	ObjectLocs.Add(ServerObjectLimit,	FObjectLoc{ .Loc=CellLimitPos, .CullDistance=(float)CellLimitPos.Size()});
+	WorldInfoToBeAssigned.Add(FObjectWorldInfo{ .Loc = CellLimitPos, .CullDistance = (float)CellLimitPos.Size() });
+	UReplicatedTestObject* ServerObjectLimit = Server->CreateObject(Params);
 
 	// Culled objects
-	ObjectLocs.Add(ServerObjectCulled,	FObjectLoc{ .Loc=FVector(DefaultGridConfig->CellSizeX + 100.f, DefaultGridConfig->CellSizeY + 100.f, 100.f), .CullDistance=1500.f });
-	ObjectLocs.Add(ServerObjectVeryFar, FObjectLoc{ .Loc=FVector(DefaultGridConfig->CellSizeX + 99999.f, DefaultGridConfig->CellSizeY + 99999.f, 99999.f), .CullDistance=1500.f });
+	WorldInfoToBeAssigned.Add(FObjectWorldInfo{ .Loc = FVector(DefaultGridConfig->CellSizeX + 100.f, DefaultGridConfig->CellSizeY + 100.f, 100.f), .CullDistance = 1500.f });
+	UReplicatedTestObject* ServerObjectCulled = Server->CreateObject(Params);
 
-	Server->GetReplicationBridge()->SetExternalWorldLocationUpdateFunctor([&](FNetRefHandle NetHandle, const UObject* ReplicatedObject, FVector& OutLocation, float& OutCullDistance)
+	WorldInfoToBeAssigned.Add(FObjectWorldInfo{ .Loc = FVector(DefaultGridConfig->CellSizeX + 99999.f, DefaultGridConfig->CellSizeY + 99999.f, 99999.f), .CullDistance = 1500.f });
+	UReplicatedTestObject* ServerObjectVeryFar = Server->CreateObject(Params);
+
+	TArray<UReplicatedTestObject*> ServerReplicatedObjects;
 	{
-		const FObjectLoc& ObjectLoc = ObjectLocs[Cast<const UReplicatedTestObject>(ReplicatedObject)];
-		OutLocation = ObjectLoc.Loc;
-		OutCullDistance = ObjectLoc.CullDistance;
-	});
+		ServerReplicatedObjects.Add(ServerObjectZero);
+		ServerReplicatedObjects.Add(ServerObjectNear);
+		ServerReplicatedObjects.Add(ServerObjectLimit);
+		ServerReplicatedObjects.Add(ServerObjectCulled);
+		ServerReplicatedObjects.Add(ServerObjectVeryFar);
+	}
 
 	// Apply grid filter
-	for (auto It = ObjectLocs.CreateConstIterator(); It; ++It)
+	for (auto It : ServerReplicatedObjects)
 	{
-		Server->ReplicationSystem->SetFilter(It.Key()->NetRefHandle, WorldLocFilterHandle);
+		Server->ReplicationSystem->SetFilter(It->NetRefHandle, WorldLocFilterHandle);
 	}
 
 	// Add client
@@ -170,30 +204,20 @@ UE_NET_TEST_FIXTURE(FTestGridFilterFixture, TestWorldLocGridFilter)
 	UE_NET_ASSERT_EQ(Client->GetReplicationBridge()->GetReplicatedObject(ServerObjectCulled->NetRefHandle), nullptr);
 	UE_NET_ASSERT_EQ(Client->GetReplicationBridge()->GetReplicatedObject(ServerObjectVeryFar->NetRefHandle), nullptr);
 
-	for (auto It = ObjectLocs.CreateConstIterator(); It; ++It)
+	for (auto It : ServerReplicatedObjects)
 	{
-		Server->DestroyObject(It.Key());
+		Server->DestroyObject(It);
 	}
 }
 
 UE_NET_TEST_FIXTURE(FTestGridFilterFixture, TestWorldLocationIsFrequentlyUpdatedForNonDormantObject)
 {
-	struct FObjectLoc
-	{
-		FVector Loc;
-		float CullDistance;
-	};
-	
-	TMap<const UObject*, FObjectLoc> ObjectLocs;
-
 	Server->GetReplicationBridge()->SetExternalWorldLocationUpdateFunctor([&](FNetRefHandle NetHandle, const UObject* ReplicatedObject, FVector& OutLocation, float& OutCullDistance)
-		{
-			if (const FObjectLoc* ObjectLoc = ObjectLocs.Find(ReplicatedObject))
-			{
-				OutLocation = ObjectLoc->Loc;
-				OutCullDistance = ObjectLoc->CullDistance;
-			}
-		});
+	{
+		const FObjectWorldInfo& Info = GetWorldInfo(ReplicatedObject);
+		OutLocation = Info.Loc;
+		OutCullDistance = Info.CullDistance;
+	});
 
 	// Add client
 	FReplicationSystemTestClient* Client = CreateClient();
@@ -205,6 +229,9 @@ UE_NET_TEST_FIXTURE(FTestGridFilterFixture, TestWorldLocationIsFrequentlyUpdated
 	Server->ReplicationSystem->SetReplicationView(Client->ConnectionIdOnServer, ReplicationView);
 
 	// Spawn object with WorldLocation's on server
+	const UNetObjectGridFilterConfig* DefaultGridConfig = GetDefault<UNetObjectGridFilterConfig>();
+	WorldInfoToBeAssigned.Add(FObjectWorldInfo{ .Loc = FVector::ZeroVector, .CullDistance = 1500.f });
+
 	UObjectReplicationBridge::FRootObjectReplicationParams Params;
 	Params.bNeedsWorldLocationUpdate = true;
 	Params.bIsDormant = false;
@@ -214,19 +241,14 @@ UE_NET_TEST_FIXTURE(FTestGridFilterFixture, TestWorldLocationIsFrequentlyUpdated
 
 	UReplicatedTestObject* ServerObject = Server->CreateObject(Params);
 
-	// Visible objects
-	ObjectLocs.Add(ServerObject, FObjectLoc{ FVector::ZeroVector, 1500.f });
-
 	// Send and deliver packet
 	Server->UpdateAndSend({ Client });
 
 	// Verify objects has been created
 	UE_NET_ASSERT_NE(Client->GetReplicationBridge()->GetReplicatedObject(ServerObject->NetRefHandle), nullptr);
 
-	const UNetObjectGridFilterConfig* DefaultGridConfig = GetDefault<UNetObjectGridFilterConfig>();
-
 	// Update the world location to a cell the object shouldn't previously have been touching, without marking the object dirty
-	ObjectLocs.Emplace(ServerObject, FObjectLoc{ FVector(DefaultGridConfig->CellSizeX + 1600.0f,DefaultGridConfig->CellSizeY + 1600.0f, 0), 1500.f });
+	SetWorldInfo(ServerObject, FObjectWorldInfo{ .Loc = FVector(DefaultGridConfig->CellSizeX + 1600.0f,DefaultGridConfig->CellSizeY + 1600.0f, 0), .CullDistance = 1500.f});
 
 	// Send and deliver packet
 	for (uint32 LoopIt = 0, LoopEndIt = DefaultGridConfig->ViewPosRelevancyFrameCount + DefaultGridConfig->DefaultFrameCountBeforeCulling; LoopIt <= LoopEndIt; ++LoopIt)
@@ -238,7 +260,7 @@ UE_NET_TEST_FIXTURE(FTestGridFilterFixture, TestWorldLocationIsFrequentlyUpdated
 	UE_NET_ASSERT_EQ(Client->GetReplicationBridge()->GetReplicatedObject(ServerObject->NetRefHandle), nullptr);
 
 	// Return object to origin.
-	ObjectLocs.Emplace(ServerObject, FObjectLoc{ FVector::ZeroVector, 1500.f });
+	SetWorldInfo(ServerObject, FObjectWorldInfo{ .Loc = FVector::ZeroVector, .CullDistance = 1500.f });
 
 	// Send and deliver packet
 	Server->UpdateAndSend({ Client });
@@ -249,40 +271,27 @@ UE_NET_TEST_FIXTURE(FTestGridFilterFixture, TestWorldLocationIsFrequentlyUpdated
 
 UE_NET_TEST_FIXTURE(FTestGridFilterFixture, TestWorldLocCullDistance)
 {
-	// Spawn object with WorldLocation's on server
-	UObjectReplicationBridge::FRootObjectReplicationParams Params;
-	Params.bNeedsWorldLocationUpdate = true;
-	Params.bUseClassConfigDynamicFilter = true;
-
-	UReplicatedTestObject* ServerObjectTest = Server->CreateObject(Params);
-
-	const UNetObjectGridFilterConfig* DefaultGridConfig = GetDefault<UNetObjectGridFilterConfig>();
-
-	struct FObjectLoc
+	Server->GetReplicationBridge()->SetExternalWorldLocationUpdateFunctor([&](FNetRefHandle NetHandle, const UObject* ReplicatedObject, FVector& OutLocation, float& OutCullDistance)
 	{
-		FVector Loc;
-		float CullDistance;
-	};
-	TMap<UReplicatedTestObject*, FObjectLoc> ObjectLocs;
+		const FObjectWorldInfo& Info = GetWorldInfo(ReplicatedObject);
+		OutLocation = Info.Loc;
+		OutCullDistance = Info.CullDistance;
+	});
 
+	// Spawn an object that resides at the exact limit of it's culldistance and at the boundary of a grid cell.
+	const UNetObjectGridFilterConfig* DefaultGridConfig = GetDefault<UNetObjectGridFilterConfig>();
 	const float DefaultCullDistance = DefaultGridConfig->CellSizeX;
+	FObjectWorldInfo ObjectWorldInfo{ .Loc = FVector(DefaultCullDistance, 0.0f, 0.0f), .CullDistance = DefaultCullDistance };
+	WorldInfoToBeAssigned.Add(ObjectWorldInfo);
 
-	// Spawn an object that resides at the exact limit of his culldistance and at the boundary of a grid cell.
-	ObjectLocs.Add(ServerObjectTest, FObjectLoc{ .Loc=FVector(DefaultCullDistance, 0.0f, 0.0f), .CullDistance=DefaultCullDistance });
-
-	Server->GetReplicationBridge()->SetExternalWorldLocationUpdateFunctor([&](FNetRefHandle NetHandle, const UObject* ReplicatedObject, FVector& OutLocation, float& OutCullDistance)
-		{
-			const FObjectLoc& ObjectLoc = ObjectLocs[Cast<const UReplicatedTestObject>(ReplicatedObject)];
-			OutLocation = ObjectLoc.Loc;
-			OutCullDistance = ObjectLoc.CullDistance;
-		});
+	UObjectReplicationBridge::FRootObjectReplicationParams Params;
+	Params.bNeedsWorldLocationUpdate = true;
+	Params.bUseClassConfigDynamicFilter = true;
+	UReplicatedTestObject* ServerObjectTest = Server->CreateObject(Params);
 
 	// Apply grid filter
-	for (auto It = ObjectLocs.CreateConstIterator(); It; ++It)
-	{
-		// Use the CullDistanceTest profile that disables the FrameCountBeforeCulling feature
-		Server->ReplicationSystem->SetFilter(It.Key()->NetRefHandle, WorldLocFilterHandle, TEXT("CullDistanceTest"));
-	}
+	// Use the CullDistanceTest profile that disables the FrameCountBeforeCulling feature
+	Server->ReplicationSystem->SetFilter(ServerObjectTest->NetRefHandle, WorldLocFilterHandle, TEXT("CullDistanceTest"));
 
 	// Add client
 	FReplicationSystemTestClient* Client = CreateClient();
@@ -318,7 +327,8 @@ UE_NET_TEST_FIXTURE(FTestGridFilterFixture, TestWorldLocCullDistance)
 	UE_NET_ASSERT_NE(Client->GetReplicationBridge()->GetReplicatedObject(ServerObjectTest->NetRefHandle), nullptr);
 
 	// Change the object's internal cull distance
-	ObjectLocs[ServerObjectTest].CullDistance = DefaultCullDistance-100.f;
+	ObjectWorldInfo.CullDistance = DefaultCullDistance - 100.f;
+	SetWorldInfo(ServerObjectTest, ObjectWorldInfo);
 
 	// The object needs to be dirty for the cull distance update to be picked up
 	Server->ReplicationSystem->MarkDirty(ServerObjectTest->NetRefHandle);
@@ -329,10 +339,7 @@ UE_NET_TEST_FIXTURE(FTestGridFilterFixture, TestWorldLocCullDistance)
 	// Object should be filtered out
 	UE_NET_ASSERT_EQ(Client->GetReplicationBridge()->GetReplicatedObject(ServerObjectTest->NetRefHandle), nullptr);
 
-	for(auto It=ObjectLocs.CreateConstIterator(); It; ++It)
-	{
-		Server->DestroyObject(It.Key());
-	}
+	Server->DestroyObject(ServerObjectTest);
 }
 
 #if 0
diff --git a/Engine/Source/Runtime/Experimental/Iris/Core/Private/Iris/ReplicationSystem/ObjectReplicationBridge.cpp b/Engine/Source/Runtime/Experimental/Iris/Core/Private/Iris/ReplicationSystem/ObjectReplicationBridge.cpp
--- a/Engine/Source/Runtime/Experimental/Iris/Core/Private/Iris/ReplicationSystem/ObjectReplicationBridge.cpp
+++ b/Engine/Source/Runtime/Experimental/Iris/Core/Private/Iris/ReplicationSystem/ObjectReplicationBridge.cpp
@@ -486,6 +486,7 @@ UE::Net::FNetRefHandle UObjectReplicationBridge::StartReplicatingRootObject(UObj
 	{
 		FWorldLocations& WorldLocations = ReplicationSystem->GetReplicationSystemInternal()->GetWorldLocations();
 		WorldLocations.InitObjectInfoCache(ObjectIndex);
+		ForceUpdateWorldLocation(RefHandle, ObjectIndex);
 	}
 
 	// Set prioritizer
