Fixed test failure in TestStructNetTokenDataStore

diff --git a/Engine/Plugins/Runtime/ReplicationSystemTestPlugin/Source/Private/Tests/ReplicationSystem/TestStructNetTokenDataStore.cpp b/Engine/Plugins/Runtime/ReplicationSystemTestPlugin/Source/Private/Tests/ReplicationSystem/TestStructNetTokenDataStore.cpp
--- a/Engine/Plugins/Runtime/ReplicationSystemTestPlugin/Source/Private/Tests/ReplicationSystem/TestStructNetTokenDataStore.cpp
+++ b/Engine/Plugins/Runtime/ReplicationSystemTestPlugin/Source/Private/Tests/ReplicationSystem/TestStructNetTokenDataStore.cpp
@@ -28,14 +28,20 @@ UE_NET_IMPLEMENT_NAMED_STRUCT_LASTRESORT_NETSERIALIZER_AND_REGISTRY_DELEGATES(Te
 ////////////////////////////////////////////////
 // Tests
 ////////////////////////////////////////////////
+
+namespace UE::Net
+{
+	// Force instantiate template to force template statics to be instantiated.
+	template class TStructNetTokenDataStore<FTestStructNetTokenStoreStruct>;
+	using FTestStructNetTokenDataStore = TStructNetTokenDataStore<FTestStructNetTokenStoreStruct>;
+};
+
 namespace UE::Net::Private
 {
 
 class FTestStructNetTokensFixture : public FReplicationSystemServerClientTestFixture
 {
 public:
-
-	typedef UE::Net::TStructNetTokenDataStore<FTestStructNetTokenStoreStruct> FTestStructNetTokenDataStore;
 	using DataType = FTestStructNetTokenDataStore::DataType;
 
 	FTestStructNetTokenDataStore* ServerStructTokenStore = nullptr;
@@ -46,7 +52,6 @@ public:
 	const FNetTokenStoreState* ClientRemoteNetTokenState;
 	const FNetTokenStoreState* ServerRemoteNetTokenState;
 
-
 	FNetToken CreateAndExportNetToken(const DataType& Struct)
 	{
 		FNetToken Token = ServerStructTokenStore->GetOrCreateToken(Struct);
