Iris
- Subobjects replicating destroy is now dispatched immediately instead of deffered to end of dispatch
- Added cvar to enable/disable net.Iris.ImmediateDispatchEndReplicationForSubObjects, default is true

diff --git a/Engine/Plugins/Runtime/ReplicationSystemTestPlugin/Source/Private/Tests/ReplicationSystem/ReplicatedTestObject.cpp b/Engine/Plugins/Runtime/ReplicationSystemTestPlugin/Source/Private/Tests/ReplicationSystem/ReplicatedTestObject.cpp
--- a/Engine/Plugins/Runtime/ReplicationSystemTestPlugin/Source/Private/Tests/ReplicationSystem/ReplicatedTestObject.cpp
+++ b/Engine/Plugins/Runtime/ReplicationSystemTestPlugin/Source/Private/Tests/ReplicationSystem/ReplicatedTestObject.cpp
@@ -416,6 +416,30 @@ void UReplicatedSubObjectOrderObject::RegisterReplicationFragments(UE::Net::FFra
 	}
 }
 
+void UReplicatedSubObjectDestroyOrderObject::SetObjectExpectedToBeDestroyed(UReplicatedSubObjectDestroyOrderObject* OtherObject)
+{
+	ObjectToWatch = MakeWeakObjectPtr(OtherObject);
+}
+
+void UReplicatedSubObjectDestroyOrderObject::BeginDestroy()
+{
+	Super::BeginDestroy();
+}
+
+void UReplicatedSubObjectDestroyOrderObject::PreNetReceive()
+{
+	Super::PreNetReceive();
+	
+	bObjectExistedInPreNetReceive = ObjectToWatch.IsSet() && ObjectToWatch.GetValue().IsValid();
+}
+
+void UReplicatedSubObjectDestroyOrderObject::PostNetReceive()
+{
+	Super::PostNetReceive();
+	
+	bObjectExistedInPostNetReceive = ObjectToWatch.IsSet() && ObjectToWatch.GetValue().IsValid();
+}
+
 //////////////////////////////////////////////////////////////////////////
 // Implementation for UTestReplicatedObjectWithRepNotifies
 //////////////////////////////////////////////////////////////////////////
diff --git a/Engine/Plugins/Runtime/ReplicationSystemTestPlugin/Source/Private/Tests/ReplicationSystem/ReplicatedTestObject.h b/Engine/Plugins/Runtime/ReplicationSystemTestPlugin/Source/Private/Tests/ReplicationSystem/ReplicatedTestObject.h
--- a/Engine/Plugins/Runtime/ReplicationSystemTestPlugin/Source/Private/Tests/ReplicationSystem/ReplicatedTestObject.h
+++ b/Engine/Plugins/Runtime/ReplicationSystemTestPlugin/Source/Private/Tests/ReplicationSystem/ReplicatedTestObject.h
@@ -442,6 +442,28 @@ public:
 	TArray<UE::Net::FReplicationFragment*> ReplicationFragments;
 };
 
+UCLASS()
+class UReplicatedSubObjectDestroyOrderObject : public UReplicatedSubObjectOrderObject
+{
+	GENERATED_BODY()
+public:
+
+public:
+	void SetObjectExpectedToBeDestroyed(UReplicatedSubObjectDestroyOrderObject* OtherObject);
+
+	virtual void BeginDestroy() override;
+
+	virtual void PreNetReceive() override;
+
+	virtual void PostNetReceive() override;
+
+	TOptional<TWeakObjectPtr<UReplicatedSubObjectDestroyOrderObject>> ObjectToWatch;
+
+	bool bObjectExistedInPreNetReceive = false;
+	bool bObjectExistedInPostNetReceive = false;
+};
+
+
 UCLASS()
 class UTestReplicatedObjectWithRepNotifies : public UReplicatedTestObject
 {
diff --git a/Engine/Plugins/Runtime/ReplicationSystemTestPlugin/Source/Private/Tests/ReplicationSystem/TestReplicationSystem.cpp b/Engine/Plugins/Runtime/ReplicationSystemTestPlugin/Source/Private/Tests/ReplicationSystem/TestReplicationSystem.cpp
--- a/Engine/Plugins/Runtime/ReplicationSystemTestPlugin/Source/Private/Tests/ReplicationSystem/TestReplicationSystem.cpp
+++ b/Engine/Plugins/Runtime/ReplicationSystemTestPlugin/Source/Private/Tests/ReplicationSystem/TestReplicationSystem.cpp
@@ -2248,6 +2248,126 @@ UE_NET_TEST_FIXTURE(FReplicationSystemServerClientTestFixture, TestReplicationRe
 	UE_NET_ASSERT_FALSE(Client->IsValidNetRefHandle(FirstObjectRefHandle));
 }
 
+UE_NET_TEST_FIXTURE(FReplicationSystemServerClientTestFixture, TestSubObjectReplicatedDestroyBeforePostNetReceive)
+{
+	// Make sure that bOldImmediateDispatchEndReplicationForSubObjects is set to true
+	IConsoleVariable* LocalCVar = IConsoleManager::Get().FindConsoleVariable(TEXT("net.Iris.ImmediateDispatchEndReplicationForSubObjects"));
+	check(LocalCVar != nullptr && LocalCVar->IsVariableBool());
+	const bool bOldImmediateDispatchEndReplicationForSubObjects = LocalCVar->GetBool();
+	LocalCVar->Set(true, ECVF_SetByCode);
+
+	ON_SCOPE_EXIT
+	{
+		// Restore cvars
+		LocalCVar->Set(bOldImmediateDispatchEndReplicationForSubObjects, ECVF_SetByCode);
+	};
+
+	// Add a client
+	FReplicationSystemTestClient* Client = CreateClient();
+
+	// Spawn object on server
+	UReplicatedTestObject* ServerObject = Server->CreateObject(0, 0);
+
+	// Spawn some subobjects
+	UReplicatedSubObjectDestroyOrderObject* ServerSubObject0 = Server->CreateSubObject<UReplicatedSubObjectDestroyOrderObject>(ServerObject->NetRefHandle);
+	UReplicatedSubObjectDestroyOrderObject* ServerSubObject1 = Server->CreateSubObject<UReplicatedSubObjectDestroyOrderObject>(ServerObject->NetRefHandle);
+	UReplicatedSubObjectDestroyOrderObject* ServerSubObject2 = Server->CreateSubObject<UReplicatedSubObjectDestroyOrderObject>(ServerObject->NetRefHandle);
+
+	UReplicatedSubObjectDestroyOrderObject::RepOrderCounter = 0U;
+
+	// Send and deliver packet
+	Server->UpdateAndSend({Client});
+
+	// Verify that objects have replicated
+	UReplicatedSubObjectDestroyOrderObject* ClientSubObject0 = Cast<UReplicatedSubObjectDestroyOrderObject>(Client->GetReplicationBridge()->GetReplicatedObject(ServerSubObject0->NetRefHandle));
+	UReplicatedSubObjectDestroyOrderObject* ClientSubObject1 = Cast<UReplicatedSubObjectDestroyOrderObject>(Client->GetReplicationBridge()->GetReplicatedObject(ServerSubObject1->NetRefHandle));
+	UReplicatedSubObjectDestroyOrderObject* ClientSubObject2 = Cast<UReplicatedSubObjectDestroyOrderObject>(Client->GetReplicationBridge()->GetReplicatedObject(ServerSubObject2->NetRefHandle));
+
+	UE_NET_ASSERT_NE(ClientSubObject0, nullptr);
+	UE_NET_ASSERT_NE(ClientSubObject1, nullptr);
+	UE_NET_ASSERT_NE(ClientSubObject2, nullptr);
+
+	// Verify that they have replicated in expected order
+	UE_NET_ASSERT_EQ(ClientSubObject0->LastRepOrderCounter, 1U);
+	UE_NET_ASSERT_GT(ClientSubObject1->LastRepOrderCounter, ClientSubObject0->LastRepOrderCounter);
+	UE_NET_ASSERT_GT(ClientSubObject2->LastRepOrderCounter, ClientSubObject1->LastRepOrderCounter);	
+
+	// setup a watch on the client
+	ClientSubObject2->SetObjectExpectedToBeDestroyed(ClientSubObject1);
+
+	// Dirty some data on server and destroy SubObject1
+	++ServerSubObject2->IntA;
+	Server->ReplicationBridge->EndReplication(ServerSubObject1, EEndReplicationFlags::Destroy);
+
+	// Send and deliver packet
+	Server->UpdateAndSend({Client});
+
+	// Verify assumptions, replicated subobject destroy should have been issued before ClientSubObjects2 s call to PostNetReceive
+	UE_NET_ASSERT_TRUE(ClientSubObject2->bObjectExistedInPreNetReceive);
+	UE_NET_ASSERT_FALSE(ClientSubObject2->bObjectExistedInPostNetReceive);
+	UE_NET_ASSERT_FALSE(ClientSubObject2->ObjectToWatch->IsValid());
+}
+
+UE_NET_TEST_FIXTURE(FReplicationSystemServerClientTestFixture, TestSubObjectReplicatedDestroyAfterPostNetReceive)
+{
+	// Make sure that bOldImmediateDispatchEndReplicationForSubObjects is set to false
+	IConsoleVariable* LocalCVar = IConsoleManager::Get().FindConsoleVariable(TEXT("net.Iris.ImmediateDispatchEndReplicationForSubObjects"));
+	check(LocalCVar != nullptr && LocalCVar->IsVariableBool());
+	const bool bOldImmediateDispatchEndReplicationForSubObjects = LocalCVar->GetBool();
+	LocalCVar->Set(false, ECVF_SetByCode);
+
+	ON_SCOPE_EXIT
+	{
+		// Restore cvars
+		LocalCVar->Set(bOldImmediateDispatchEndReplicationForSubObjects, ECVF_SetByCode);
+	};
+
+	// Add a client
+	FReplicationSystemTestClient* Client = CreateClient();
+
+	// Spawn object on server
+	UReplicatedTestObject* ServerObject = Server->CreateObject(0, 0);
+
+	// Spawn some subobjects
+	UReplicatedSubObjectDestroyOrderObject* ServerSubObject0 = Server->CreateSubObject<UReplicatedSubObjectDestroyOrderObject>(ServerObject->NetRefHandle);
+	UReplicatedSubObjectDestroyOrderObject* ServerSubObject1 = Server->CreateSubObject<UReplicatedSubObjectDestroyOrderObject>(ServerObject->NetRefHandle);
+	UReplicatedSubObjectDestroyOrderObject* ServerSubObject2 = Server->CreateSubObject<UReplicatedSubObjectDestroyOrderObject>(ServerObject->NetRefHandle);
+
+	UReplicatedSubObjectDestroyOrderObject::RepOrderCounter = 0U;
+
+	// Send and deliver packet
+	Server->UpdateAndSend({Client});
+
+	// Verify that objects have replicated
+	UReplicatedSubObjectDestroyOrderObject* ClientSubObject0 = Cast<UReplicatedSubObjectDestroyOrderObject>(Client->GetReplicationBridge()->GetReplicatedObject(ServerSubObject0->NetRefHandle));
+	UReplicatedSubObjectDestroyOrderObject* ClientSubObject1 = Cast<UReplicatedSubObjectDestroyOrderObject>(Client->GetReplicationBridge()->GetReplicatedObject(ServerSubObject1->NetRefHandle));
+	UReplicatedSubObjectDestroyOrderObject* ClientSubObject2 = Cast<UReplicatedSubObjectDestroyOrderObject>(Client->GetReplicationBridge()->GetReplicatedObject(ServerSubObject2->NetRefHandle));
+
+	UE_NET_ASSERT_NE(ClientSubObject0, nullptr);
+	UE_NET_ASSERT_NE(ClientSubObject1, nullptr);
+	UE_NET_ASSERT_NE(ClientSubObject2, nullptr);
+
+	// Verify that they have replicated in expected order
+	UE_NET_ASSERT_EQ(ClientSubObject0->LastRepOrderCounter, 1U);
+	UE_NET_ASSERT_GT(ClientSubObject1->LastRepOrderCounter, ClientSubObject0->LastRepOrderCounter);
+	UE_NET_ASSERT_GT(ClientSubObject2->LastRepOrderCounter, ClientSubObject1->LastRepOrderCounter);	
+
+	// setup a watch on the client
+	ClientSubObject2->SetObjectExpectedToBeDestroyed(ClientSubObject1);
+
+	// Dirty some data on server and destroy SubObject1
+	++ServerSubObject2->IntA;
+	Server->ReplicationBridge->EndReplication(ServerSubObject1, EEndReplicationFlags::Destroy);
+
+	// Send and deliver packet
+	Server->UpdateAndSend({Client});
+
+	// Verify assumptions, replicated subobject destroy should have been issued after ClientSubObjects2 s call to PostNetReceive
+	UE_NET_ASSERT_TRUE(ClientSubObject2->bObjectExistedInPreNetReceive);
+	UE_NET_ASSERT_TRUE(ClientSubObject2->bObjectExistedInPostNetReceive);
+	UE_NET_ASSERT_FALSE(ClientSubObject2->ObjectToWatch->IsValid());
+}
+
 } // end namespace UE::Net::Private
 
 
