Iris
- Added parameter to LastResortNetSerializerConfig to exclude property from defaultstatehash
- FText now serialzied using LastResortNetSerializer excluding defaultstatehash
- Fixes issue with protocolid mismatch due to FText::Serialize modifying value

diff --git a/Engine/Plugins/Runtime/ReplicationSystemTestPlugin/Source/Private/Tests/Serialization/TestNameNetSerializer.cpp b/Engine/Plugins/Runtime/ReplicationSystemTestPlugin/Source/Private/Tests/Serialization/TestNameNetSerializer.cpp
--- a/Engine/Plugins/Runtime/ReplicationSystemTestPlugin/Source/Private/Tests/Serialization/TestNameNetSerializer.cpp
+++ b/Engine/Plugins/Runtime/ReplicationSystemTestPlugin/Source/Private/Tests/Serialization/TestNameNetSerializer.cpp
@@ -170,7 +170,7 @@ UE_NET_TEST_FIXTURE(FReplicationSystemServerClientTestFixture, TestFName)
 	Server->UpdateAndSend({Client});
 
 	// Verify that object has been spawned on client
-	auto ClientObject = Client->GetObjectAs<UTestNameNetSerializer_TestObject>(ServerObject->NetRefHandle);
+	UTestNameNetSerializer_TestObject* ClientObject = Client->GetObjectAs<UTestNameNetSerializer_TestObject>(ServerObject->NetRefHandle);
 
 	// Verify that created server handle now also exists on client
 	UE_NET_ASSERT_NE(ClientObject, nullptr);	
@@ -199,7 +199,7 @@ UE_NET_TEST_FIXTURE(FReplicationSystemServerClientTestFixture, TestFName_Replica
 	Server->UpdateAndSend({Client});
 
 	// Verify that object has been spawned on client
-	auto ClientObject = Client->GetObjectAs<UTestNameNetSerializer_TestObject>(ServerObject->NetRefHandle);
+	UTestNameNetSerializer_TestObject* ClientObject = Client->GetObjectAs<UTestNameNetSerializer_TestObject>(ServerObject->NetRefHandle);
 
 	// Verify that created server handle now also exists on client
 	UE_NET_ASSERT_NE(ClientObject, nullptr);	
@@ -282,6 +282,34 @@ UE_NET_TEST_FIXTURE(FRPCTestFixture, TestFNameServerRPCCanExportsName)
 	UE_NET_ASSERT_EQ(ExpectedName, ServerObject->NameFromServerRPC);
 }
 
+UE_NET_TEST_FIXTURE(FReplicationSystemServerClientTestFixture, TestFText)
+{
+	// Add a client
+	FReplicationSystemTestClient* Client = CreateClient();
+
+	// Spawn objects on server
+	UTextProperty_TestObject* ServerObject = Server->CreateObject<UTextProperty_TestObject>();
+
+	Server->UpdateAndSend({Client});
+
+	// Verify that object has been spawned on client
+	UTextProperty_TestObject* ClientObject = Client->GetObjectAs<UTextProperty_TestObject>(ServerObject->NetRefHandle);
+	UE_NET_ASSERT_NE(ClientObject, nullptr);	
+
+	// Verify that we managed to replicate the expected name
+	UE_NET_ASSERT_TRUE(ServerObject->TextProperty.EqualTo(ClientObject->TextProperty));
+
+	FText Modified = FText::FromString(TEXT("ModifiedText"));
+
+	ServerObject->TextProperty = Modified;
+
+	Server->UpdateAndSend({Client});
+
+	// Verify that we managed to replicate the expected name
+	UE_NET_ASSERT_TRUE(ServerObject->TextProperty.EqualTo(ClientObject->TextProperty));	
+}
+
+
 }
 
 UTestNameNetSerializer_TestObject::UTestNameNetSerializer_TestObject()
@@ -407,4 +435,27 @@ void UTestNameNetSerializer_TestObjectWithRPC::ServerRPCWithName_Implementation(
 	NameFromServerRPC = Name;
 }
 
+//////////////////////////////////////////////////////////////////////////
+// Implementation for UTextProperty_TestObject
+//////////////////////////////////////////////////////////////////////////
+UTextProperty_TestObject::UTextProperty_TestObject()
+: UReplicatedTestObject()
+{
+	TextProperty = FText::FromString("DefaultText");
+}
+
+void UTextProperty_TestObject::GetLifetimeReplicatedProps(TArray<class FLifetimeProperty>& OutLifetimeProps) const
+{
+	FDoRepLifetimeParams Params;
+	Params.bIsPushBased = false;
+
+	DOREPLIFETIME_WITH_PARAMS_FAST(UTextProperty_TestObject, TextProperty, Params);
+}
+
+void UTextProperty_TestObject::RegisterReplicationFragments(UE::Net::FFragmentRegistrationContext& Context, UE::Net::EFragmentRegistrationFlags RegistrationFlags)
+{
+	UE::Net::FReplicationFragmentUtil::CreateAndRegisterFragmentsForObject(this, Context, RegistrationFlags);
+}
+
+
 
diff --git a/Engine/Plugins/Runtime/ReplicationSystemTestPlugin/Source/Private/Tests/Serialization/TestNameNetSerializer.h b/Engine/Plugins/Runtime/ReplicationSystemTestPlugin/Source/Private/Tests/Serialization/TestNameNetSerializer.h
--- a/Engine/Plugins/Runtime/ReplicationSystemTestPlugin/Source/Private/Tests/Serialization/TestNameNetSerializer.h
+++ b/Engine/Plugins/Runtime/ReplicationSystemTestPlugin/Source/Private/Tests/Serialization/TestNameNetSerializer.h
@@ -59,4 +59,15 @@ public:
 	FName NameFromServerRPC;
 };
 
+UCLASS()
+class UTextProperty_TestObject : public UReplicatedTestObject
+{
+	GENERATED_BODY()
+public:
+	UTextProperty_TestObject();
 
+	void RegisterReplicationFragments(UE::Net::FFragmentRegistrationContext& Fragments, UE::Net::EFragmentRegistrationFlags RegistrationFlags) override;
+
+	UPROPERTY(Replicated)
+	FText TextProperty;
+};
diff --git a/Engine/Source/Runtime/Experimental/Iris/Core/Private/Iris/ReplicationState/DefaultPropertyNetSerializerInfos.cpp b/Engine/Source/Runtime/Experimental/Iris/Core/Private/Iris/ReplicationState/DefaultPropertyNetSerializerInfos.cpp
--- a/Engine/Source/Runtime/Experimental/Iris/Core/Private/Iris/ReplicationState/DefaultPropertyNetSerializerInfos.cpp
+++ b/Engine/Source/Runtime/Experimental/Iris/Core/Private/Iris/ReplicationState/DefaultPropertyNetSerializerInfos.cpp
@@ -6,6 +6,7 @@
 #include "Iris/Serialization/NetSerializers.h"
 #include "Iris/Serialization/InternalNetSerializers.h"
 #include "Iris/Serialization/DateTimeNetSerializer.h"
+#include "UObject/TextProperty.h"
 
 namespace UE::Net::Private
 {
@@ -138,6 +139,23 @@ struct FScriptInterfacePropertyNetSerializerInfo : public FPropertyNetSerializer
 };
 UE_NET_IMPLEMENT_NETSERIALIZER_INFO(FScriptInterfacePropertyNetSerializerInfo);
 
+/**
+ * FText NetSerializerInfo, special forwarding to LastResortNetSerializer excluding default state hash.
+ */
+struct FTextPropertyNetSerializerInfo : public FLastResortPropertyNetSerializerInfo
+{
+	FTextPropertyNetSerializerInfo() : FLastResortPropertyNetSerializerInfo()
+	{
+		bExcludeFromDefaultStateHash = true;
+	}
+
+	virtual const FFieldClass* GetPropertyTypeClass() const override
+	{
+		return FTextProperty::StaticClass();
+	}
+};
+UE_NET_IMPLEMENT_NETSERIALIZER_INFO(FTextPropertyNetSerializerInfo);
+
 // $TODO: Investigate automatic registration by binding all infos in static chain and auto-register when we freeze the registry
 /**
  * Apart from the specialized serializer we register here we support the following:
@@ -180,6 +198,9 @@ void RegisterDefaultPropertyNetSerializerInfos()
 	UE_NET_REGISTER_NETSERIALIZER_INFO(FNameProperty);
 	UE_NET_REGISTER_NETSERIALIZER_INFO(FStrProperty);
 
+	// FTextProperty is using a special variant of LastResortNetSerializer that does not serialize defaultstatehash.
+	UE_NET_REGISTER_NETSERIALIZER_INFO(FTextPropertyNetSerializerInfo);
+
 	// Special types
 	UE_NET_REGISTER_NETSERIALIZER_INFO(FNativeBoolPropertyNetSerializerInfo);
 	UE_NET_REGISTER_NETSERIALIZER_INFO(FBitFieldPropertyNetSerializerInfo);
diff --git a/Engine/Source/Runtime/Experimental/Iris/Core/Private/Iris/ReplicationState/PropertyNetSerializerInfoRegistry.cpp b/Engine/Source/Runtime/Experimental/Iris/Core/Private/Iris/ReplicationState/PropertyNetSerializerInfoRegistry.cpp
--- a/Engine/Source/Runtime/Experimental/Iris/Core/Private/Iris/ReplicationState/PropertyNetSerializerInfoRegistry.cpp
+++ b/Engine/Source/Runtime/Experimental/Iris/Core/Private/Iris/ReplicationState/PropertyNetSerializerInfoRegistry.cpp
@@ -109,6 +109,8 @@ FNetSerializerConfig* FLastResortPropertyNetSerializerInfo::BuildNetSerializerCo
 {
 	FLastResortPropertyNetSerializerConfig* Config = new (NetSerializerConfigBuffer) FLastResortPropertyNetSerializerConfig();
 	InitLastResortPropertyNetSerializerConfigFromProperty(*Config, Property);
+	Config->bExcludeFromDefaultStateHash = bExcludeFromDefaultStateHash;
+
 	return Config;
 }
 UE_NET_IMPLEMENT_NETSERIALIZER_INFO(FLastResortPropertyNetSerializerInfo);
diff --git a/Engine/Source/Runtime/Experimental/Iris/Core/Private/Iris/Serialization/InternalNetSerializers.h b/Engine/Source/Runtime/Experimental/Iris/Core/Private/Iris/Serialization/InternalNetSerializers.h
--- a/Engine/Source/Runtime/Experimental/Iris/Core/Private/Iris/Serialization/InternalNetSerializers.h
+++ b/Engine/Source/Runtime/Experimental/Iris/Core/Private/Iris/Serialization/InternalNetSerializers.h
@@ -73,6 +73,9 @@ struct FLastResortPropertyNetSerializerConfig : public FNetSerializerConfig
 public:
 	UPROPERTY()
 	TFieldPath<FProperty> Property;
+
+	UPROPERTY()
+	bool bExcludeFromDefaultStateHash = false;
 };
 
 // ENetRole. With role swapping at deserialization
diff --git a/Engine/Source/Runtime/Experimental/Iris/Core/Private/Iris/Serialization/LastResortPropertyNetSerializer.cpp b/Engine/Source/Runtime/Experimental/Iris/Core/Private/Iris/Serialization/LastResortPropertyNetSerializer.cpp
--- a/Engine/Source/Runtime/Experimental/Iris/Core/Private/Iris/Serialization/LastResortPropertyNetSerializer.cpp
+++ b/Engine/Source/Runtime/Experimental/Iris/Core/Private/Iris/Serialization/LastResortPropertyNetSerializer.cpp
@@ -81,11 +81,20 @@ void FLastResortPropertyNetSerializer::Serialize(FNetSerializationContext& Conte
 	const ConfigType* Config = static_cast<const ConfigType*>(Args.NetSerializerConfig);
 	const QuantizedType& Value = *reinterpret_cast<const QuantizedType*>(Args.Source);
 	FNetBitStreamWriter* Writer = Context.GetBitStreamWriter();
-
-	// For now we ignore this in default state hash if it has exported NetTokens as they will differ
-	if (Context.IsInitializingDefaultState() && (Value.QuantizedExports.NetTokenStorage.Num() > 0U))
+	
+	if (Context.IsInitializingDefaultState())
 	{
-		return;
+		// If the config indicates that we should not be included in the default state hash write nothing
+		if (Config->bExcludeFromDefaultStateHash)
+		{
+			return;
+		}
+
+		// For now we ignore this in default state hash if it has exported NetTokens as they will differ
+		if (Value.QuantizedExports.NetTokenStorage.Num() > 0U)
+		{
+			return;
+		}
 	}
 
 	UE_NET_TRACE_DYNAMIC_NAME_SCOPE(Config->Property.Get()->GetName(), *Writer, Context.GetTraceCollector(), ENetTraceVerbosity::VeryVerbose);
diff --git a/Engine/Source/Runtime/Experimental/Iris/Core/Public/Iris/ReplicationState/PropertyNetSerializerInfoRegistry.h b/Engine/Source/Runtime/Experimental/Iris/Core/Public/Iris/ReplicationState/PropertyNetSerializerInfoRegistry.h
--- a/Engine/Source/Runtime/Experimental/Iris/Core/Public/Iris/ReplicationState/PropertyNetSerializerInfoRegistry.h
+++ b/Engine/Source/Runtime/Experimental/Iris/Core/Public/Iris/ReplicationState/PropertyNetSerializerInfoRegistry.h
@@ -155,12 +155,16 @@ private:
 };
 
 struct FLastResortPropertyNetSerializerInfo : public FPropertyNetSerializerInfo
-{		
+{
 	IRISCORE_API virtual const FFieldClass* GetPropertyTypeClass() const override;
 	IRISCORE_API virtual bool IsSupported(const FProperty* Property) const override;
 	IRISCORE_API virtual const FNetSerializer* GetNetSerializer(const FProperty* Property) const override final;
 	IRISCORE_API virtual bool CanUseDefaultConfig(const FProperty* Property) const override final;
 	IRISCORE_API virtual FNetSerializerConfig* BuildNetSerializerConfig(void* NetSerializerConfigBuffer, const FProperty* Property) const override final;
+
+protected:
+	// If set this class should not be included in the default state hash.
+	bool bExcludeFromDefaultStateHash = false;
 };
 
 struct FNamedStructLastResortPropertyNetSerializerInfo final : public FLastResortPropertyNetSerializerInfo
