Iris
- Fixed Crash in NetObjectGridFilter if culldistance was large enough to not add object to grid and explictly remove it.

diff --git a/Engine/Plugins/Runtime/ReplicationSystemTestPlugin/Source/Private/Tests/ReplicationSystem/Filtering/TestGridFilter.cpp b/Engine/Plugins/Runtime/ReplicationSystemTestPlugin/Source/Private/Tests/ReplicationSystem/Filtering/TestGridFilter.cpp
--- a/Engine/Plugins/Runtime/ReplicationSystemTestPlugin/Source/Private/Tests/ReplicationSystem/Filtering/TestGridFilter.cpp
+++ b/Engine/Plugins/Runtime/ReplicationSystemTestPlugin/Source/Private/Tests/ReplicationSystem/Filtering/TestGridFilter.cpp
@@ -210,6 +210,50 @@ UE_NET_TEST_FIXTURE(FTestGridFilterFixture, TestWorldLocGridFilter)
 	}
 }
 
+UE_NET_TEST_FIXTURE(FTestGridFilterFixture, TestWorldLocGridFilterWithHugeCullDistance)
+{
+	Server->GetReplicationBridge()->SetExternalWorldLocationUpdateFunctor([&](FNetRefHandle NetHandle, const UObject* ReplicatedObject, FVector& OutLocation, float& OutCullDistance)
+	{
+		const FObjectWorldInfo& Info = GetWorldInfo(ReplicatedObject);
+		OutLocation = Info.Loc;
+		OutCullDistance = Info.CullDistance;
+	});
+
+	const UNetObjectGridFilterConfig* DefaultGridConfig = GetDefault<UNetObjectGridFilterConfig>();
+
+	// Spawn object with WorldLocation's on server
+	UObjectReplicationBridge::FRootObjectReplicationParams Params;
+	Params.bNeedsWorldLocationUpdate = true;
+	Params.bUseClassConfigDynamicFilter = true;
+
+	// We do expect and ensure
+	FTestEnsureScope SuppressEnsureScope;
+	
+	// Relevant objects
+	WorldInfoToBeAssigned.Add(FObjectWorldInfo{ .Loc = FVector::ZeroVector, .CullDistance = 200000.f });
+	UReplicatedTestObject* ServerObjectHugeCullDistance = Server->CreateObject(Params);
+
+	// Apply grid filter
+	Server->ReplicationSystem->SetFilter(ServerObjectHugeCullDistance->NetRefHandle, WorldLocFilterHandle);
+
+	// Add client
+	FReplicationSystemTestClient* Client = CreateClient();
+
+	// Set the view location of the client to (0,0,0)
+	FReplicationView ReplicationView;
+	FReplicationView::FView View = ReplicationView.Views.Emplace_GetRef();
+	View.Pos = FVector(0.0f, 0.0f, 0.0f);
+	Server->ReplicationSystem->SetReplicationView(Client->ConnectionIdOnServer, ReplicationView);
+
+	// Send and deliver packet
+	Server->PreSendUpdate();
+	Server->SendAndDeliverTo(Client, DeliverPacket);
+	Server->PostSendUpdate();
+
+	// Test that object is not filtered
+	UE_NET_ASSERT_NE(Client->GetReplicationBridge()->GetReplicatedObject(ServerObjectHugeCullDistance->NetRefHandle), nullptr);
+}
+
 UE_NET_TEST_FIXTURE(FTestGridFilterFixture, TestWorldLocationIsFrequentlyUpdatedForNonDormantObject)
 {
 	Server->GetReplicationBridge()->SetExternalWorldLocationUpdateFunctor([&](FNetRefHandle NetHandle, const UObject* ReplicatedObject, FVector& OutLocation, float& OutCullDistance)
diff --git a/Engine/Source/Runtime/Experimental/Iris/Core/Private/Iris/ReplicationSystem/Filtering/NetObjectGridFilter.cpp b/Engine/Source/Runtime/Experimental/Iris/Core/Private/Iris/ReplicationSystem/Filtering/NetObjectGridFilter.cpp
--- a/Engine/Source/Runtime/Experimental/Iris/Core/Private/Iris/ReplicationSystem/Filtering/NetObjectGridFilter.cpp
+++ b/Engine/Source/Runtime/Experimental/Iris/Core/Private/Iris/ReplicationSystem/Filtering/NetObjectGridFilter.cpp
@@ -344,7 +344,10 @@ void UNetObjectGridFilter::RemoveCellInfoForObject(const FObjectLocationInfo& Ob
 				Coord.Y = Y;
 				FCellObjects* Cell = Cells.Find(Coord);
 				checkSlow(Cell != nullptr);
-				Cell->ObjectIndices.Remove(ObjectIndex);
+				if (Cell)
+				{
+					Cell->ObjectIndices.Remove(ObjectIndex);
+				}
 			}
 		}
 	}
diff --git a/Engine/Source/Runtime/Experimental/Iris/Core/Public/Iris/ReplicationSystem/Filtering/NetObjectGridFilter.h b/Engine/Source/Runtime/Experimental/Iris/Core/Public/Iris/ReplicationSystem/Filtering/NetObjectGridFilter.h
--- a/Engine/Source/Runtime/Experimental/Iris/Core/Public/Iris/ReplicationSystem/Filtering/NetObjectGridFilter.h
+++ b/Engine/Source/Runtime/Experimental/Iris/Core/Public/Iris/ReplicationSystem/Filtering/NetObjectGridFilter.h
@@ -137,9 +137,9 @@ protected:
 	struct FCellBox
 	{
 		int32 MinX = 0;
-		int32 MaxX = 0;
+		int32 MaxX = -1;
 		int32 MinY = 0;
-		int32 MaxY = 0;
+		int32 MaxY = -1;
 
 		bool operator==(const FCellBox&) const;
 		bool operator!=(const FCellBox&) const;
@@ -149,7 +149,7 @@ protected:
 	struct FPerObjectInfo
 	{
 		FVector Position = FVector::ZeroVector;
-		FCellBox CellBox = {};
+		FCellBox CellBox;
 		uint32 ObjectIndex = 0U;
 		uint16 FrameCountBeforeCulling = 0U;
 
