Fix ensure in GridFilter test due to invalid cull distance usage

#rnx

diff --git a/Engine/Plugins/Runtime/ReplicationSystemTestPlugin/Source/Private/Tests/ReplicationSystem/Filtering/TestGridFilter.cpp b/Engine/Plugins/Runtime/ReplicationSystemTestPlugin/Source/Private/Tests/ReplicationSystem/Filtering/TestGridFilter.cpp
--- a/Engine/Plugins/Runtime/ReplicationSystemTestPlugin/Source/Private/Tests/ReplicationSystem/Filtering/TestGridFilter.cpp
+++ b/Engine/Plugins/Runtime/ReplicationSystemTestPlugin/Source/Private/Tests/ReplicationSystem/Filtering/TestGridFilter.cpp
@@ -6,6 +6,7 @@
 #include "Iris/ReplicationSystem/Filtering/NetObjectFilter.h"
 #include "Iris/ReplicationSystem/Filtering/NetObjectFilterDefinitions.h"
 #include "Iris/ReplicationSystem/Filtering/NetObjectGridFilter.h"
+#include "Iris/ReplicationSystem/WorldLocations.h"
 
 #include "Tests/ReplicationSystem/Filtering/MockNetObjectFilter.h"
 #include "Tests/ReplicationSystem/Filtering/TestFilteringObject.h"
@@ -220,6 +221,7 @@ UE_NET_TEST_FIXTURE(FTestGridFilterFixture, TestWorldLocGridFilterWithHugeCullDi
 	});
 
 	const UNetObjectGridFilterConfig* DefaultGridConfig = GetDefault<UNetObjectGridFilterConfig>();
+	const UWorldLocationsConfig* DefaultWorldLocFonfig = GetDefault<UWorldLocationsConfig>();
 
 	// Spawn object with WorldLocation's on server
 	UObjectReplicationBridge::FRootObjectReplicationParams Params;
@@ -230,7 +232,7 @@ UE_NET_TEST_FIXTURE(FTestGridFilterFixture, TestWorldLocGridFilterWithHugeCullDi
 	FTestEnsureScope SuppressEnsureScope;
 	
 	// Relevant objects
-	WorldInfoToBeAssigned.Add(FObjectWorldInfo{ .Loc = FVector::ZeroVector, .CullDistance = 200000.f });
+	WorldInfoToBeAssigned.Add(FObjectWorldInfo{ .Loc = FVector::ZeroVector, .CullDistance = DefaultWorldLocFonfig->MaxNetCullDistance-SMALL_NUMBER });
 	UReplicatedTestObject* ServerObjectHugeCullDistance = Server->CreateObject(Params);
 
 	// Apply grid filter
@@ -250,7 +252,7 @@ UE_NET_TEST_FIXTURE(FTestGridFilterFixture, TestWorldLocGridFilterWithHugeCullDi
 	Server->SendAndDeliverTo(Client, DeliverPacket);
 	Server->PostSendUpdate();
 
-	// Test that object is not filtered
+	// Test that the object is relevant
 	UE_NET_ASSERT_NE(Client->GetReplicationBridge()->GetReplicatedObject(ServerObjectHugeCullDistance->NetRefHandle), nullptr);
 }
 
