Iris
- Added pruning of unsent unreliable rpcs on the OOB channel after each tick to avoid accumulating a backlog of unreliable attachments if constantly overcommiting
- Disable by setting cvar net.Iris.ReplicationWriter.MaxUnsentOrderedUnreliableAttachmentAtEndOfTick -1
- Disable warning by setting cvar net.Iris.ReplicationWriter.MaxUnsentOrderedUnreliableAttachmentAtEndOfTick 0

diff --git a/Engine/Plugins/Runtime/ReplicationSystemTestPlugin/Source/Private/Tests/ReplicationSystem/RPC/RPCTestFixture.cpp b/Engine/Plugins/Runtime/ReplicationSystemTestPlugin/Source/Private/Tests/ReplicationSystem/RPC/RPCTestFixture.cpp
--- a/Engine/Plugins/Runtime/ReplicationSystemTestPlugin/Source/Private/Tests/ReplicationSystem/RPC/RPCTestFixture.cpp
+++ b/Engine/Plugins/Runtime/ReplicationSystemTestPlugin/Source/Private/Tests/ReplicationSystem/RPC/RPCTestFixture.cpp
@@ -12,6 +12,7 @@ void FRPCTestFixture::SetUp()
 
 	HandlerDefinitions.Add({ TEXT("NetRPCHandler") });
 	HandlerDefinitions.Add({ TEXT("PartialNetObjectAttachmentHandler") });
+	HandlerDefinitions.Add({ TEXT("NetObjectBlobHandler") });
 
 	BlobHandlerDefinitions->ReadWriteHandlerDefinitions() = HandlerDefinitions;
 
diff --git a/Engine/Plugins/Runtime/ReplicationSystemTestPlugin/Source/Private/Tests/ReplicationSystem/RPC/ReplicatedTestObjectWithRPC.cpp b/Engine/Plugins/Runtime/ReplicationSystemTestPlugin/Source/Private/Tests/ReplicationSystem/RPC/ReplicatedTestObjectWithRPC.cpp
--- a/Engine/Plugins/Runtime/ReplicationSystemTestPlugin/Source/Private/Tests/ReplicationSystem/RPC/ReplicatedTestObjectWithRPC.cpp
+++ b/Engine/Plugins/Runtime/ReplicationSystemTestPlugin/Source/Private/Tests/ReplicationSystem/RPC/ReplicatedTestObjectWithRPC.cpp
@@ -144,6 +144,13 @@ void UTestReplicatedObjectWithRPC::ServerUnreliableRPC_Implementation()
 	ServerUnreliableRPCCallOrder = ++CallOrder;
 }
 
+void UTestReplicatedObjectWithRPC::ServerUnreliableWithExportRPC_Implementation(const TArray<uint8>& Payload, UObject* Ref)
+{
+	bServerUnreliableWithExportRPCCalled = true;
+	ServerUnreliableWithExportRPCCallOrder = ++CallOrder;
+	++ServerUnreliableWithExportRPCCallCounter;
+};
+
 void UTestReplicatedObjectWithRPC::NetMulticast_MultiCastRPCSendImmediate_Implementation()
 {
 	NetMulticast_MultiCastRPCSendImmediateCallOrder = ++CallOrder;
diff --git a/Engine/Plugins/Runtime/ReplicationSystemTestPlugin/Source/Private/Tests/ReplicationSystem/RPC/ReplicatedTestObjectWithRPC.h b/Engine/Plugins/Runtime/ReplicationSystemTestPlugin/Source/Private/Tests/ReplicationSystem/RPC/ReplicatedTestObjectWithRPC.h
--- a/Engine/Plugins/Runtime/ReplicationSystemTestPlugin/Source/Private/Tests/ReplicationSystem/RPC/ReplicatedTestObjectWithRPC.h
+++ b/Engine/Plugins/Runtime/ReplicationSystemTestPlugin/Source/Private/Tests/ReplicationSystem/RPC/ReplicatedTestObjectWithRPC.h
@@ -97,6 +97,13 @@ public:
 	bool bServerUnreliableRPCCalled = false;
 	int32 ServerUnreliableRPCCallOrder = 0;
 
+	/** Server Unicast unreliable with reference */
+	UFUNCTION(Unreliable, Server)
+	void ServerUnreliableWithExportRPC(const TArray<uint8>& Payload, UObject* Ref);
+	bool bServerUnreliableWithExportRPCCalled = false;
+	int32 ServerUnreliableWithExportRPCCallOrder = 0;
+	int32 ServerUnreliableWithExportRPCCallCounter = 0;
+
 	UFUNCTION(NetMulticast, unreliable)
 	void NetMulticast_MultiCastRPCSendImmediate();
 	int32 NetMulticast_MultiCastRPCSendImmediateCallOrder = 0;
diff --git a/Engine/Plugins/Runtime/ReplicationSystemTestPlugin/Source/Private/Tests/ReplicationSystem/RPC/TestObjectRPC.cpp b/Engine/Plugins/Runtime/ReplicationSystemTestPlugin/Source/Private/Tests/ReplicationSystem/RPC/TestObjectRPC.cpp
--- a/Engine/Plugins/Runtime/ReplicationSystemTestPlugin/Source/Private/Tests/ReplicationSystem/RPC/TestObjectRPC.cpp
+++ b/Engine/Plugins/Runtime/ReplicationSystemTestPlugin/Source/Private/Tests/ReplicationSystem/RPC/TestObjectRPC.cpp
@@ -3,6 +3,7 @@
 #include "RPCTestFixture.h"
 #include "ReplicatedTestObjectWithRPC.h"
 #include "Containers/BitArray.h"
+#include "HAL/IConsoleManager.h"
 #include "Iris/Serialization/NetBitStreamReader.h"
 #include "Iris/Serialization/NetBitStreamWriter.h"
 #include "Iris/ReplicationSystem/ReplicationSystem.h"
@@ -475,4 +476,118 @@ namespace UE::Net::Private
 		}
 	}
 
+	// This test is written to test that we can control if ordered unreliable attachments are pruned or not after each tick.
+	UE_NET_TEST_FIXTURE(FRPCTestFixture, TestOvercommitOrderedUnreliableRPC)
+	{
+		// This tests tests the behavior with and without DropUnsentOrderedUnreliableAttachmentAtEndOfTick enabled
+		IConsoleVariable* CVarOverride = IConsoleManager::Get().FindConsoleVariable(TEXT("net.Iris.ReplicationWriter.MaxUnsentOrderedUnreliableAttachmentAtEndOfTick"));
+		check(CVarOverride != nullptr && CVarOverride->IsVariableInt());
+		const int32 OldCVarValue = CVarOverride->GetInt();
+
+		// Force disable dropping at unsent unreliable attachments
+		CVarOverride->Set(-1, ECVF_SetByCode);
+
+		ON_SCOPE_EXIT
+		{
+			// Restore cvar
+			CVarOverride->Set(OldCVarValue, ECVF_SetByCode);
+		};
+
+		// Add a client
+		FReplicationSystemTestClient* Client = CreateClient();
+
+		// Spawn object on server
+		UTestReplicatedObjectWithRPC* ServerObject = Server->CreateObject<UTestReplicatedObjectWithRPC>();
+
+		ServerObject->bIsServerObject = true;
+		ServerObject->ReplicationSystem = Server->GetReplicationSystem();
+		Server->ReplicationSystem->SetOwningNetConnection(ServerObject->NetRefHandle, Client->ConnectionIdOnServer);
+
+		// Send and deliver packet
+		Server->UpdateAndSend({Client});
+
+		UTestReplicatedObjectWithRPC* ClientObject = Cast<UTestReplicatedObjectWithRPC>(Client->GetReplicationBridge()->GetReplicatedObject(ServerObject->NetRefHandle));
+		
+		// Verify that created server handle now also exists on client
+		UE_NET_ASSERT_NE(ClientObject, nullptr);
+
+		ClientObject->ReplicationSystem = Client->GetReplicationSystem();
+
+		auto SendRpcs = [&ClientObject](uint32 RPCCount, uint32 PayloadSize, UObject* Reference)
+		{
+			for (uint32 i = 0; i<RPCCount; ++i)
+			{
+				TArray<uint8> Payload;
+				Payload.SetNumZeroed(PayloadSize);
+				ClientObject->ServerUnreliableWithExportRPC(Payload, Reference);
+			}
+		};
+
+		const int32 PacketSize = 500;
+		const int32 IssueRPCCount = 8;
+		const int32 RPCPayloadSize = 128;
+		const int32 ExpectedPerPacket = PacketSize / RPCPayloadSize;
+
+		Client->SetMaxSendPacketSize(PacketSize);
+	
+		// Run test allowing backlog of ordered unreliable attachments
+		{
+			CVarOverride->Set(-1, ECVF_SetByCode);
+	
+			// Post a bunch of unreliable RPCs that will overcommit more unreliable attachments than we should be able to send.
+			SendRpcs(IssueRPCCount, 128, nullptr);
+	
+			// We send a reliable RPC last so that we can query when all data has been transmitted (or not)
+			ClientObject->ServerRPC();
+
+			// Keep on sending until we have sent and acknowledged the reliable rpc.
+			UDataStream* ClientDataStream = Client->ReplicationSystem->GetDataStream(Client->LocalConnectionId, "Replication");
+			while (!ClientDataStream->HasAcknowledgedAllReliableData())
+			{
+				// Send and deliver packet
+				Client->UpdateAndSend(Server);
+				Server->UpdateAndSend({Client});
+			}
+
+			// Verify RPC reception
+			UE_NET_ASSERT_TRUE(ServerObject->bServerRPCCalled);
+
+			UE_NET_ASSERT_TRUE(ServerObject->bServerUnreliableWithExportRPCCalled);
+
+			// We expect all unreliables to be be delivered
+			UE_NET_ASSERT_EQ(ServerObject->ServerUnreliableWithExportRPCCallCounter, IssueRPCCount);
+		}
+
+		// Reset callcounters
+		ServerObject->bServerUnreliableWithExportRPCCalled = false;
+		ServerObject->ServerUnreliableWithExportRPCCallCounter = 0;
+
+		// Run test where we discard unsent unreliable attachments after each tick/frame
+		{
+			CVarOverride->Set(0, ECVF_SetByCode);
+	
+			// Post a bunch of unreliable RPCs that will overcommit more unreliable attachments than we should be able to send.
+			SendRpcs(IssueRPCCount, 128, nullptr);
+	
+			// We send a reliable RPC last so that we can query when all data has been transmitted (or not)
+			ClientObject->ServerRPC();
+
+			// Keep on sending until we have sent and acknowledged the reliable rpc.
+			UDataStream* ClientDataStream = Client->ReplicationSystem->GetDataStream(Client->LocalConnectionId, "Replication");
+			while (!ClientDataStream->HasAcknowledgedAllReliableData())
+			{
+				// Send and deliver packet
+				Client->UpdateAndSend(Server);
+				Server->UpdateAndSend({Client});
+			}
+
+			// Verify RPC reception
+			UE_NET_ASSERT_TRUE(ServerObject->bServerRPCCalled);
+
+			UE_NET_ASSERT_TRUE(ServerObject->bServerUnreliableWithExportRPCCalled);
+
+			// We expect some unreliable attachments to be delivered but not all.
+			UE_NET_ASSERT_EQ(ServerObject->ServerUnreliableWithExportRPCCallCounter, ExpectedPerPacket);
+		}
+	}
 }
\ No newline at end of file
diff --git a/Engine/Source/Runtime/Experimental/Iris/Core/Private/Iris/ReplicationSystem/AttachmentReplication.cpp b/Engine/Source/Runtime/Experimental/Iris/Core/Private/Iris/ReplicationSystem/AttachmentReplication.cpp
--- a/Engine/Source/Runtime/Experimental/Iris/Core/Private/Iris/ReplicationSystem/AttachmentReplication.cpp
+++ b/Engine/Source/Runtime/Experimental/Iris/Core/Private/Iris/ReplicationSystem/AttachmentReplication.cpp
@@ -137,6 +137,34 @@ public:
 		ReliableQueue.CommitReplicationRecord(Record);
 	}
 
+	SIZE_T GetUnreliableCount() const
+	{
+		SIZE_T UnreliableCount = 0;
+
+		// Count all unreliable attachments in PreQueue and ReliableQueue
+		for (const TRefCountPtr<FNetBlob>& BlobRef : PreQueue)
+		{
+			const FNetBlob* Blob = BlobRef.GetReference();
+			UnreliableCount += Blob && !Blob->IsReliable() ? 1 : 0;
+		}
+	
+		UnreliableCount += ReliableQueue.GetUnreliableCount();
+
+		return UnreliableCount;
+	}
+
+	void DropUnreliable()
+	{
+		// Drop all unreliable from pre-queue
+		PreQueue.RemoveAll([](const TRefCountPtr<FNetBlob>& BlobRef)
+		{
+			const FNetBlob* Blob = BlobRef.GetReference();
+			return Blob == nullptr || !Blob->IsReliable();
+		});
+
+		ReliableQueue.DropUnreliable();
+	}
+
 private:
 	void PopulateQueueFromPreQueue()
 	{
@@ -215,10 +243,31 @@ bool FNetObjectAttachmentSendQueue::Enqueue(TArrayView<const TRefCountPtr<FNetBl
 	return true;
 }
 
+SIZE_T FNetObjectAttachmentSendQueue::GetUnreliableCount() const
+{
+	SIZE_T UnreliableCount = 0;
+
+	UnreliableCount += UnreliableQueue.Count();
+	if (ReliableQueue)
+	{
+		UnreliableCount += ReliableQueue->GetUnreliableCount();
+	}
+
+	return UnreliableCount;
+}
+
 void FNetObjectAttachmentSendQueue::DropUnreliable(bool& bOutHasUnsent)
 {
 	UnreliableQueue.Empty();
-	bOutHasUnsent = (ReliableQueue != nullptr && ReliableQueue->HasUnsentBlobs());
+	if (ReliableQueue)
+	{
+		ReliableQueue->DropUnreliable();
+		bOutHasUnsent = ReliableQueue->HasUnsentBlobs();
+	}
+	else
+	{
+		bOutHasUnsent = false;
+	}
 }
 
 bool FNetObjectAttachmentSendQueue::HasUnsent() const
@@ -556,9 +605,20 @@ void FNetObjectAttachmentsWriter::DropAllAttachments(ENetObjectAttachmentType Ty
 	ObjectToQueue.Remove(ObjectIndex);
 }
 
+SIZE_T FNetObjectAttachmentsWriter::GetUnreliableAttachmentCount(ENetObjectAttachmentType Type, uint32 ObjectIndex) const
+{
+	const FNetObjectAttachmentSendQueue* Queue = GetQueue(Type, ObjectIndex);
+	if (Queue == nullptr)
+	{
+		return 0;
+	}
+
+	return Queue->GetUnreliableCount();
+}
+
 void FNetObjectAttachmentsWriter::DropUnreliableAttachments(ENetObjectAttachmentType Type, uint32 ObjectIndex, bool &bOutHasUnsentAttachments)
 {
-	FNetObjectAttachmentSendQueue* Queue = GetQueue(ENetObjectAttachmentType::Normal, ObjectIndex);
+	FNetObjectAttachmentSendQueue* Queue = GetQueue(Type, ObjectIndex);
 	if (Queue == nullptr)
 	{
 		bOutHasUnsentAttachments = false;
diff --git a/Engine/Source/Runtime/Experimental/Iris/Core/Private/Iris/ReplicationSystem/AttachmentReplication.h b/Engine/Source/Runtime/Experimental/Iris/Core/Private/Iris/ReplicationSystem/AttachmentReplication.h
--- a/Engine/Source/Runtime/Experimental/Iris/Core/Private/Iris/ReplicationSystem/AttachmentReplication.h
+++ b/Engine/Source/Runtime/Experimental/Iris/Core/Private/Iris/ReplicationSystem/AttachmentReplication.h
@@ -71,6 +71,8 @@ public:
 
 	bool Enqueue(TArrayView<const TRefCountPtr<FNetBlob>> Attachments);
 
+	SIZE_T GetUnreliableCount() const;
+
 	void DropUnreliable(bool &bOutHasUnsentAttachments);
 
 	bool HasUnsent() const;
@@ -133,6 +135,10 @@ public:
 	bool IsSafeToDestroy(ENetObjectAttachmentType Type, uint32 ObjectIndex) const;
 
 	void DropAllAttachments(ENetObjectAttachmentType Type, uint32 ObjectIndex);
+
+	// Counts all unreliable attachments for the given type
+	SIZE_T GetUnreliableAttachmentCount(ENetObjectAttachmentType Type, uint32 ObjectIndex) const;
+
 	void DropUnreliableAttachments(ENetObjectAttachmentType Type, uint32 ObjectIndex, bool& bOutHasUnsentAttachments);
 
 	EAttachmentWriteStatus Serialize(FNetSerializationContext& Context, ENetObjectAttachmentType Type, uint32 ObjectIndex, FNetRefHandle RefHandle, FCommitRecord& OutRecord, bool& bOutHasUnsentAttachments);
diff --git a/Engine/Source/Runtime/Experimental/Iris/Core/Private/Iris/ReplicationSystem/NetBlob/ReliableNetBlobQueue.cpp b/Engine/Source/Runtime/Experimental/Iris/Core/Private/Iris/ReplicationSystem/NetBlob/ReliableNetBlobQueue.cpp
--- a/Engine/Source/Runtime/Experimental/Iris/Core/Private/Iris/ReplicationSystem/NetBlob/ReliableNetBlobQueue.cpp
+++ b/Engine/Source/Runtime/Experimental/Iris/Core/Private/Iris/ReplicationSystem/NetBlob/ReliableNetBlobQueue.cpp
@@ -218,6 +218,38 @@ void FReliableNetBlobQueue::CommitReplicationRecord(const FReliableNetBlobQueue:
 	}
 }
 
+SIZE_T FReliableNetBlobQueue::GetUnreliableCount() const
+{
+	SIZE_T UnreliableCount = 0U;
+
+	for (uint32 Seq = FirstSeq; Seq < LastSeq; ++Seq)
+	{
+		const uint32 Index = SequenceToIndex(Seq);
+		const TRefCountPtr<FNetBlob>& RefCntBlob = NetBlobs[Index];
+		if (const FNetBlob* Blob = RefCntBlob.GetReference())
+		{
+			UnreliableCount += Blob && !Blob->IsReliable() ? 1 : 0;
+		}
+	}
+
+	return UnreliableCount;
+}
+
+void FReliableNetBlobQueue::DropUnreliable()
+{
+	for (uint32 Seq = FirstSeq; Seq < LastSeq; ++Seq)
+	{
+		const uint32 Index = SequenceToIndex(Seq);
+
+		TRefCountPtr<FNetBlob>& RefCntBlob = NetBlobs[Index];
+		const FNetBlob* Blob = RefCntBlob.GetReference();
+		if (Blob && !Blob->IsReliable())
+		{
+			RefCntBlob.SafeRelease();
+		}
+	}
+}
+
 uint32 FReliableNetBlobQueue::Deserialize(FNetSerializationContext& Context)
 {
 	FNetRefHandle InvalidNetHandle;
diff --git a/Engine/Source/Runtime/Experimental/Iris/Core/Private/Iris/ReplicationSystem/ReplicationDataStream.cpp b/Engine/Source/Runtime/Experimental/Iris/Core/Private/Iris/ReplicationSystem/ReplicationDataStream.cpp
--- a/Engine/Source/Runtime/Experimental/Iris/Core/Private/Iris/ReplicationSystem/ReplicationDataStream.cpp
+++ b/Engine/Source/Runtime/Experimental/Iris/Core/Private/Iris/ReplicationSystem/ReplicationDataStream.cpp
@@ -53,6 +53,14 @@ void UReplicationDataStream::Deinit()
 	ReplicationReader = nullptr;
 }
 
+void UReplicationDataStream::Update(const FUpdateParameters& Params)
+{
+	if (ReplicationWriter)
+	{
+		ReplicationWriter->Update(Params);
+	}
+}
+
 UDataStream::EWriteResult UReplicationDataStream::BeginWrite(const FBeginWriteParameters& Params)
 {
 	if (ReplicationWriter != nullptr)
diff --git a/Engine/Source/Runtime/Experimental/Iris/Core/Private/Iris/ReplicationSystem/ReplicationDataStream.h b/Engine/Source/Runtime/Experimental/Iris/Core/Private/Iris/ReplicationSystem/ReplicationDataStream.h
--- a/Engine/Source/Runtime/Experimental/Iris/Core/Private/Iris/ReplicationSystem/ReplicationDataStream.h
+++ b/Engine/Source/Runtime/Experimental/Iris/Core/Private/Iris/ReplicationSystem/ReplicationDataStream.h
@@ -23,7 +23,7 @@ private:
 	// UDataStream interface
 	virtual void Init(const UDataStream::FInitParameters& Params) override;
 	virtual void Deinit() override;
-
+	virtual void Update(const FUpdateParameters& Params) override;
 	virtual EWriteResult BeginWrite(const FBeginWriteParameters& Params) override;
 	virtual EWriteResult WriteData(UE::Net::FNetSerializationContext& Context, FDataStreamRecord const*& OutRecord) override;
 	virtual void EndWrite() override;
diff --git a/Engine/Source/Runtime/Experimental/Iris/Core/Private/Iris/ReplicationSystem/ReplicationSystem.cpp b/Engine/Source/Runtime/Experimental/Iris/Core/Private/Iris/ReplicationSystem/ReplicationSystem.cpp
--- a/Engine/Source/Runtime/Experimental/Iris/Core/Private/Iris/ReplicationSystem/ReplicationSystem.cpp
+++ b/Engine/Source/Runtime/Experimental/Iris/Core/Private/Iris/ReplicationSystem/ReplicationSystem.cpp
@@ -1013,6 +1013,10 @@ void UReplicationSystem::PostSendUpdate()
 			Impl->ReplicationSystemInternal.GetDeltaCompressionBaselineManager().PostSendUpdate(UpdateParams);
 		}
 
+		// Allow data streams to do end of tick work.
+		UDataStream::FUpdateParameters DataStreamUpdateParams = { .UpdateType = UDataStream::EUpdateType::PostTickFlush };
+		Impl->UpdateDataStreams(DataStreamUpdateParams);
+
 #if UE_NET_IRIS_CSV_STATS
 		{
 			FNetSendStats& SendStats = Impl->ReplicationSystemInternal.GetSendStats();
diff --git a/Engine/Source/Runtime/Experimental/Iris/Core/Private/Iris/ReplicationSystem/ReplicationWriter.cpp b/Engine/Source/Runtime/Experimental/Iris/Core/Private/Iris/ReplicationSystem/ReplicationWriter.cpp
--- a/Engine/Source/Runtime/Experimental/Iris/Core/Private/Iris/ReplicationSystem/ReplicationWriter.cpp
+++ b/Engine/Source/Runtime/Experimental/Iris/Core/Private/Iris/ReplicationSystem/ReplicationWriter.cpp
@@ -97,6 +97,13 @@ static FAutoConsoleVariableRef CVarReplicationWriterCannotSendWarningInterval(TE
 static bool bAllowSendStateWithEndReplicationForSubObjects = false;
 static FAutoConsoleVariableRef CvarAllowSendStateWithEndReplicationForSubObjects(TEXT("net.Iris.AllowSendStateWithEndReplicationForSubObjects"), bAllowSendStateWithEndReplicationForSubObjects, TEXT("Allow state to be sent together with subobject destroy. Default is false. Cannot be true if net.Iris.ImmediateDispatchEndReplicationForSubObjects is true."));
 
+// Drop unsent ordered unreliable attachments at the end of the tick, 
+// -1 = Allow them to be queued for next tick, 
+// 0  = Always drop at the end of the tick
+// >0 = Drop and log the count if we have more queued ordered unreliable than GMaxUnsentOrderedUnreliableAttachmentAtEndOfTick
+static int32 GMaxUnsentOrderedUnreliableAttachmentAtEndOfTick = 1;
+static FAutoConsoleVariableRef CvarMaxUnsentOrderedUnreliableAttachmentAtEndOfTick(TEXT("net.Iris.ReplicationWriter.MaxUnsentOrderedUnreliableAttachmentAtEndOfTick"), GMaxUnsentOrderedUnreliableAttachmentAtEndOfTick, TEXT("Drop unsent ordered unreliable attachments at the end of the tick, -1 means that that we allow them to be queued for next tick, 0 or greater means that we will drop if we have more queued unreliable attachment than the set value at the end of the tick"));
+
 static const FName NetError_ObjectStateTooLarge("Object state is too large to be split.");
 
 const TCHAR* FReplicationWriter::LexToString(const EReplicatedObjectState State)
@@ -3769,6 +3776,27 @@ void FReplicationWriter::EndWrite()
 	}
 }
 
+void FReplicationWriter::Update(const UDataStream::FUpdateParameters& Params)
+{
+	if (Params.UpdateType == UDataStream::EUpdateType::PostTickFlush)
+	{
+		bool bHasUnsentReliable = false;
+		if (GMaxUnsentOrderedUnreliableAttachmentAtEndOfTick == 0)
+		{
+			Attachments.DropUnreliableAttachments(ENetObjectAttachmentType::OutOfBand, ObjectIndexForOOBAttachment, bHasUnsentReliable);
+		}
+		else if (GMaxUnsentOrderedUnreliableAttachmentAtEndOfTick > 0)
+		{
+			const SIZE_T UnreliableCount = Attachments.GetUnreliableAttachmentCount(ENetObjectAttachmentType::OutOfBand, ObjectIndexForOOBAttachment);
+			if (UnreliableCount >= GMaxUnsentOrderedUnreliableAttachmentAtEndOfTick)
+			{
+				Attachments.DropUnreliableAttachments(ENetObjectAttachmentType::OutOfBand, ObjectIndexForOOBAttachment, bHasUnsentReliable);
+				UE_LOG(LogIris, Warning, TEXT("FReplicationWriter::Discarded %d unsent ordered unreliable attachments"), UnreliableCount);
+			}
+		}
+	}
+}
+
 bool FReplicationWriter::HasDataToSend(const FWriteContext& Context) const
 {
 	return WriteContext.bIsValid & (Context.bHasDestroyedObjectsToSend | Context.bHasUpdatedObjectsToSend | Context.bHasHugeObjectToSend | Context.bHasOOBAttachmentsToSend);
diff --git a/Engine/Source/Runtime/Experimental/Iris/Core/Private/Iris/ReplicationSystem/ReplicationWriter.h b/Engine/Source/Runtime/Experimental/Iris/Core/Private/Iris/ReplicationSystem/ReplicationWriter.h
--- a/Engine/Source/Runtime/Experimental/Iris/Core/Private/Iris/ReplicationSystem/ReplicationWriter.h
+++ b/Engine/Source/Runtime/Experimental/Iris/Core/Private/Iris/ReplicationSystem/ReplicationWriter.h
@@ -203,6 +203,8 @@ public:
 
 	bool AreAllReliableAttachmentsSentAndAcked() const;
 
+	void Update(const UDataStream::FUpdateParameters& Params);
+
 	[[nodiscard]] FString PrintObjectInfo(FInternalNetRefIndex ObjectIndex) const;
 
 private:
diff --git a/Engine/Source/Runtime/Experimental/Iris/Core/Public/Iris/DataStream/DataStream.h b/Engine/Source/Runtime/Experimental/Iris/Core/Public/Iris/DataStream/DataStream.h
--- a/Engine/Source/Runtime/Experimental/Iris/Core/Public/Iris/DataStream/DataStream.h
+++ b/Engine/Source/Runtime/Experimental/Iris/Core/Public/Iris/DataStream/DataStream.h
@@ -75,6 +75,9 @@ public:
 	{
 		// Update originating from ReplicationSystem::PreSendUpdate
 		PreSendUpdate = 0,
+
+		// Update originating from the end of the main Network tick
+		PostTickFlush = 1,
 	};
 
 	enum class EDataStreamState : uint8
diff --git a/Engine/Source/Runtime/Experimental/Iris/Core/Public/Iris/ReplicationSystem/NetBlob/ReliableNetBlobQueue.h b/Engine/Source/Runtime/Experimental/Iris/Core/Public/Iris/ReplicationSystem/NetBlob/ReliableNetBlobQueue.h
--- a/Engine/Source/Runtime/Experimental/Iris/Core/Public/Iris/ReplicationSystem/NetBlob/ReliableNetBlobQueue.h
+++ b/Engine/Source/Runtime/Experimental/Iris/Core/Public/Iris/ReplicationSystem/NetBlob/ReliableNetBlobQueue.h
@@ -112,6 +112,12 @@ public:
 	/** On the receiving end this will move all received unreliable NetBlobs to the array and release them from the queue. This breaks the ordering guarantees provided by using Peek and Pop. Reliable NetBlobs are unaffected by this operation. */
 	IRISCORE_API void DequeueUnreliable(TArray<TRefCountPtr<FNetBlob>>& Unreliable);
 
+	/** Count all unreliable blobs */
+	IRISCORE_API SIZE_T GetUnreliableCount() const;
+
+	/** On sending side this will drop all unreliable in the queue to avoid creating a backlog of unsent unreliable blobs */
+	IRISCORE_API void DropUnreliable();
+
 	/**
 	 * Serializes as many blobs as possible using their respective SerializeWithObject() method. It is assumed the NetRefHandle will be 
 	 * reconstructed somehow on the receiving end and passed to DeserializeWithObject().
