Iris
- Fix issue where same ActorComponent being replicated using different internalindices when changing isreplicated on and off

diff --git a/Engine/Plugins/Runtime/ReplicationSystemTestPlugin/Source/Private/Tests/ReplicationSystem/TestReplicationSystem.cpp b/Engine/Plugins/Runtime/ReplicationSystemTestPlugin/Source/Private/Tests/ReplicationSystem/TestReplicationSystem.cpp
--- a/Engine/Plugins/Runtime/ReplicationSystemTestPlugin/Source/Private/Tests/ReplicationSystem/TestReplicationSystem.cpp
+++ b/Engine/Plugins/Runtime/ReplicationSystemTestPlugin/Source/Private/Tests/ReplicationSystem/TestReplicationSystem.cpp
@@ -1746,6 +1746,79 @@ UE_NET_TEST_FIXTURE(FReplicationSystemServerClientTestFixture, TestReplicatedSub
 	UE_NET_ASSERT_TRUE(Cast<UTestReplicatedIrisObjectWithNoReplicatedMembers>(Client->GetReplicationBridge()->GetReplicatedObject(ServerSubObject->NetRefHandle)) == nullptr);
 }
 
+// Test replicated SubObjectDestroy for filtered out subobject
+UE_NET_TEST_FIXTURE(FReplicationSystemServerClientTestFixture, TestReplicatedSubObjectDestroyForFilteredOutSubObject)
+{
+	UReplicationSystem* ReplicationSystem = Server->ReplicationSystem;
+
+	// Add a client
+	FReplicationSystemTestClient* Client = CreateClient();
+
+	// Spawn object on server
+	UTestReplicatedIrisObject* ServerObject = Server->CreateObject(0,0);
+
+	// Spawn second object on server as a subobject
+	UTestReplicatedIrisObjectWithNoReplicatedMembers* ServerSubObject = Server->CreateSubObject<UTestReplicatedIrisObjectWithNoReplicatedMembers>(ServerObject->NetRefHandle);
+
+	// Replicate object
+	Server->UpdateAndSend({Client});
+
+	// Verify that the subobject does  exist on client
+	UE_NET_ASSERT_TRUE(Cast<UTestReplicatedIrisObjectWithNoReplicatedMembers>(Client->GetReplicationBridge()->GetReplicatedObject(ServerSubObject->NetRefHandle)) != nullptr);
+
+	// Set condition
+	Server->ReplicationBridge->SetSubObjectNetCondition(ServerSubObject->NetRefHandle, ELifetimeCondition::COND_Never);
+
+	// Destroy subobject
+	Server->ReplicationBridge->EndReplication(ServerSubObject);
+
+	// Replicate object
+	Server->UpdateAndSend({Client});
+
+	// Verify that the object still exists on client
+	UE_NET_ASSERT_TRUE(Cast<UTestReplicatedIrisObject>(Client->GetReplicationBridge()->GetReplicatedObject(ServerObject->NetRefHandle)) != nullptr);
+
+	// Verify that the subobject does not exist on client
+	UE_NET_ASSERT_TRUE(Cast<UTestReplicatedIrisObjectWithNoReplicatedMembers>(Client->GetReplicationBridge()->GetReplicatedObject(ServerSubObject->NetRefHandle)) == nullptr);
+}
+
+// Test replicated SubObjectDestroy for filtered out subobject
+UE_NET_TEST_FIXTURE(FReplicationSystemServerClientTestFixture, TestReplicatedSubObjectDestroyForFilteredOutSubObjectBeforeSend)
+{
+	UReplicationSystem* ReplicationSystem = Server->ReplicationSystem;
+
+	// Add a client
+	FReplicationSystemTestClient* Client = CreateClient();
+
+	// Spawn object on server
+	UTestReplicatedIrisObject* ServerObject = Server->CreateObject(0,0);
+
+	// Spawn second object on server as a subobject
+	UTestReplicatedIrisObjectWithNoReplicatedMembers* ServerSubObject = Server->CreateSubObject<UTestReplicatedIrisObjectWithNoReplicatedMembers>(ServerObject->NetRefHandle);
+
+	// Set condition
+	Server->ReplicationBridge->SetSubObjectNetCondition(ServerSubObject->NetRefHandle, ELifetimeCondition::COND_Never);
+
+	// Replicate object
+	Server->UpdateAndSend({Client});
+
+	// Verify that the subobject does not exist on client
+	UE_NET_ASSERT_TRUE(Cast<UTestReplicatedIrisObjectWithNoReplicatedMembers>(Client->GetReplicationBridge()->GetReplicatedObject(ServerSubObject->NetRefHandle)) == nullptr);
+
+	// Destroy subobject
+	Server->ReplicationBridge->EndReplication(ServerSubObject);
+
+	// Replicate object
+	Server->UpdateAndSend({Client});
+
+	// Verify that the object still exists on client
+	UE_NET_ASSERT_TRUE(Cast<UTestReplicatedIrisObject>(Client->GetReplicationBridge()->GetReplicatedObject(ServerObject->NetRefHandle)) != nullptr);
+
+	// Verify that the subobject does not exist on client
+	UE_NET_ASSERT_TRUE(Cast<UTestReplicatedIrisObjectWithNoReplicatedMembers>(Client->GetReplicationBridge()->GetReplicatedObject(ServerSubObject->NetRefHandle)) == nullptr);
+}
+
+
 // Test tear-off object in PendingCreate state to ensure that tear-off logic works as expected
 UE_NET_TEST_FIXTURE(FReplicationSystemServerClientTestFixture, TestTearOffObjectWithNoFragmentsDoesNotTriggerCheckIfPendingCreateWhenDestroyed)
 {
diff --git a/Engine/Source/Runtime/Engine/Private/Actor.cpp b/Engine/Source/Runtime/Engine/Private/Actor.cpp
--- a/Engine/Source/Runtime/Engine/Private/Actor.cpp
+++ b/Engine/Source/Runtime/Engine/Private/Actor.cpp
@@ -101,6 +101,11 @@ DECLARE_CYCLE_STAT(TEXT("PostActorConstruction"), STAT_PostActorConstruction, ST
 
 FMakeNoiseDelegate AActor::MakeNoiseDelegate = FMakeNoiseDelegate::CreateStatic(&AActor::MakeNoiseImpl);
 
+#if UE_WITH_IRIS
+/** Controls if a replicated component should call EndReplication() immediately when removed from the registered subobjectlist. */
+extern bool GbEndReplicationForActorComponentWhenRemovedFromRegisteredSubObjectsList;
+#endif
+
 /** Selects if actor and actorcomponents will replicate subobjects using the registration list by default. */
 extern bool GDefaultUseSubObjectReplicationList;
 
@@ -3837,7 +3842,14 @@ void AActor::UpdateReplicatedComponent(UActorComponent* Component)
 		}
 		else if (bWasRemovedFromReplicatedComponents || NetCondition == COND_Never)
 		{
-			Component->EndReplication();
+			if (GbEndReplicationForActorComponentWhenRemovedFromRegisteredSubObjectsList)
+			{
+				Component->EndReplication();
+			}
+			else
+			{
+				UE::Net::FReplicationSystemUtil::SetNetConditionForActorComponent(Component, COND_Never);
+			}
 		}
 	}
 #endif // UE_WITH_IRIS
@@ -3859,7 +3871,14 @@ void AActor::UpdateAllReplicatedComponents()
 			// Need to end replication for all components that should no longer replicate
 			if (RepComponentInfo.Component && !RepComponentInfo.Component->GetIsReplicated())
 			{
-				RepComponentInfo.Component->EndReplication();
+				if (GbEndReplicationForActorComponentWhenRemovedFromRegisteredSubObjectsList)
+				{
+					RepComponentInfo.Component->EndReplication();
+				}
+				else
+				{
+					UE::Net::FReplicationSystemUtil::SetNetConditionForActorComponent(RepComponentInfo.Component, COND_Never);
+				}
 			}
 		}
 #endif
@@ -3893,7 +3912,14 @@ void AActor::UpdateAllReplicatedComponents()
 					}
 					else
 					{
-						Component->EndReplication();
+						if (GbEndReplicationForActorComponentWhenRemovedFromRegisteredSubObjectsList)
+						{
+							Component->EndReplication();
+						}
+						else
+						{
+							UE::Net::FReplicationSystemUtil::SetNetConditionForActorComponent(Component, COND_Never);
+						}
 					}
 				}
 #endif
diff --git a/Engine/Source/Runtime/Engine/Private/ActorReplication.cpp b/Engine/Source/Runtime/Engine/Private/ActorReplication.cpp
--- a/Engine/Source/Runtime/Engine/Private/ActorReplication.cpp
+++ b/Engine/Source/Runtime/Engine/Private/ActorReplication.cpp
@@ -31,6 +31,16 @@
 	AActor networking implementation.
 -----------------------------------------------------------------------------*/
 
+#if UE_WITH_IRIS
+
+bool GbEndReplicationForActorComponentWhenRemovedFromRegisteredSubObjectsList = false;
+static FAutoConsoleVariableRef CVarEndReplicationForActorComponentWhenRemovedFromRegisteredSubObjectsList(
+	TEXT("net.Iris.EndReplicationForActorComponentWhenRemovedFromRegisteredSubObjectsList"),
+	GbEndReplicationForActorComponentWhenRemovedFromRegisteredSubObjectsList,
+	TEXT("If true, replicated ActorComponentrs will call endreplication immediately when being removed from replicated subobjects list."));
+
+#endif
+
 //
 // Static variables for networking.
 //
@@ -618,7 +628,7 @@ void AActor::SetReplicatedComponentNetCondition(const UActorComponent* Replicate
 		if (NetCondition != ComponentInfo->NetCondition)
 		{
 #if UE_WITH_IRIS
-			UE::Net::FReplicationSystemUtil::SetActorComponentNetCondition(ReplicatedComponent, NetCondition);
+			UE::Net::FReplicationSystemUtil::SetNetConditionForActorComponent(ReplicatedComponent, NetCondition);
 #endif
 			ComponentInfo->NetCondition = NetCondition;
 		}
@@ -686,7 +696,14 @@ void AActor::RemoveReplicatedComponent(UActorComponent* Component)
 #if UE_WITH_IRIS
 		if (HasAuthority())
 		{
-			Component->EndReplication();
+			if (GbEndReplicationForActorComponentWhenRemovedFromRegisteredSubObjectsList)
+			{
+				Component->EndReplication();
+			}
+			else
+			{
+				UE::Net::FReplicationSystemUtil::SetNetConditionForActorComponent(Component, ELifetimeCondition::COND_Never);
+			}
 		}
 #endif
 	}
diff --git a/Engine/Source/Runtime/Engine/Private/Net/Iris/ReplicationSystem/EngineReplicationBridge.cpp b/Engine/Source/Runtime/Engine/Private/Net/Iris/ReplicationSystem/EngineReplicationBridge.cpp
--- a/Engine/Source/Runtime/Engine/Private/Net/Iris/ReplicationSystem/EngineReplicationBridge.cpp
+++ b/Engine/Source/Runtime/Engine/Private/Net/Iris/ReplicationSystem/EngineReplicationBridge.cpp
@@ -373,7 +373,7 @@ UE::Net::FNetRefHandle UEngineReplicationBridge::StartReplicatingActor(AActor* A
 			if (IsValid(SubObjectToReplicate) && SubObjectInfo.NetCondition != ELifetimeCondition::COND_Never)
 			{
 				FNetRefHandle SubObjectRefHandle = Super::StartReplicatingSubObject(SubObjectToReplicate, SubObjectParams, SubObjectFactoryId);
-				if (SubObjectRefHandle.IsValid() && SubObjectInfo.NetCondition != ELifetimeCondition::COND_None)
+				if (SubObjectRefHandle.IsValid())
 				{
 					UObjectReplicationBridge::SetSubObjectNetCondition(SubObjectRefHandle, SubObjectInfo.NetCondition);
 				}
@@ -442,10 +442,7 @@ UE::Net::FNetRefHandle UEngineReplicationBridge::StartReplicatingComponent(FNetR
 	}
 
 	// Update or set any conditionals
-	if (RepComponentInfo->NetCondition != ELifetimeCondition::COND_None)
-	{
-		SetSubObjectNetCondition(ReplicatedComponentHandle, RepComponentInfo->NetCondition);
-	}
+	SetSubObjectNetCondition(ReplicatedComponentHandle, RepComponentInfo->NetCondition);
 
 	// Begin replication for any SubObjects registered by the component
 	{
@@ -459,10 +456,10 @@ UE::Net::FNetRefHandle UEngineReplicationBridge::StartReplicatingComponent(FNetR
 		for (const FSubObjectRegistry::FEntry& SubObjectInfo : RepComponentInfo->SubObjects.GetRegistryList())
 		{
 			UObject* SubObjectToReplicate = SubObjectInfo.GetSubObject();
-			if (IsValid(SubObjectToReplicate) && SubObjectInfo.NetCondition != ELifetimeCondition::COND_Never)
+			if (IsValid(SubObjectToReplicate))
 			{
 				FNetRefHandle SubObjectHandle = Super::StartReplicatingSubObject(SubObjectToReplicate, Params, SubObjectFactoryId);
-				if (SubObjectHandle.IsValid() && SubObjectInfo.NetCondition != ELifetimeCondition::COND_None)
+				if (SubObjectHandle.IsValid())
 				{
 					SetSubObjectNetCondition(SubObjectHandle, SubObjectInfo.NetCondition);
 				}
diff --git a/Engine/Source/Runtime/Engine/Private/Net/Iris/ReplicationSystem/ReplicationSystemUtil.cpp b/Engine/Source/Runtime/Engine/Private/Net/Iris/ReplicationSystem/ReplicationSystemUtil.cpp
--- a/Engine/Source/Runtime/Engine/Private/Net/Iris/ReplicationSystem/ReplicationSystemUtil.cpp
+++ b/Engine/Source/Runtime/Engine/Private/Net/Iris/ReplicationSystem/ReplicationSystemUtil.cpp
@@ -319,6 +319,38 @@ void FReplicationSystemUtil::EndReplicationForActorSubObject(const AActor* Actor
 	});
 }
 
+void FReplicationSystemUtil::SetNetConditionForActorSubObject(const AActor* Actor, UObject* SubObject, ELifetimeCondition NetCondition)
+{
+	if (!IsValid(Actor) || !IsValid(SubObject))
+	{
+		return;
+	}
+
+	const FNetHandle ActorHandle = GetNetHandle(Actor);
+	if (!ActorHandle.IsValid())
+	{
+		return;
+	}
+
+	if (const UWorld* World = Actor->GetWorld())
+	{
+		ForEachReplicationSystem(GEngine, World, [SubObject, NetCondition](UReplicationSystem* ReplicationSystem)
+		{
+			if (ReplicationSystem->IsServer())
+			{
+				if (UObjectReplicationBridge* Bridge = ReplicationSystem->GetReplicationBridgeAs<UObjectReplicationBridge>())
+				{
+					const FNetRefHandle SubObjectRefHandle = Bridge->GetReplicatedRefHandle(SubObject);
+					if (SubObjectRefHandle.IsValid())
+					{
+						Bridge->SetSubObjectNetCondition(SubObjectRefHandle, NetCondition);
+					}
+				}
+			}
+		});
+	}
+}
+
 void FReplicationSystemUtil::EndReplicationForActorComponentSubObject(UActorComponent* ActorComponent, UObject* SubObject)
 {
 	EndReplicationForActorSubObject(ActorComponent->GetOwner(), SubObject);
@@ -404,7 +436,7 @@ void FReplicationSystemUtil::RemoveDependentActor(const AActor* Parent, AActor*
 	});
 }
 
-void FReplicationSystemUtil::SetActorComponentNetCondition(const UActorComponent* ActorComponent, ELifetimeCondition NetCondition)
+void FReplicationSystemUtil::SetNetConditionForActorComponent(const UActorComponent* ActorComponent, ELifetimeCondition NetCondition)
 {
 	if (!IsValid(ActorComponent))
 	{
diff --git a/Engine/Source/Runtime/Engine/Public/Net/Iris/ReplicationSystem/ReplicationSystemUtil.h b/Engine/Source/Runtime/Engine/Public/Net/Iris/ReplicationSystem/ReplicationSystemUtil.h
--- a/Engine/Source/Runtime/Engine/Public/Net/Iris/ReplicationSystem/ReplicationSystemUtil.h
+++ b/Engine/Source/Runtime/Engine/Public/Net/Iris/ReplicationSystem/ReplicationSystemUtil.h
@@ -90,13 +90,20 @@ struct FReplicationSystemUtil
 
 	/** Stop replicating an ActorComponent and its associated SubObjects. */
 	ENGINE_API static void EndReplicationForActorComponentSubObject(UActorComponent* ActorComponent, UObject* SubObject);
+
+	/** 
+	 * Set subobject NetCondition for a subobject, the conditions is used to determine if the SubObject should be replicated or not.
+	 * @note: As the filtering is done at the serialization level it is typically more efficient to use a separate object for connection 
+	 * specific data as filtering can be done at a higher level.
+	 */
+	ENGINE_API static void SetNetConditionForActorSubObject(const AActor* Actor, UObject* SubObject, ELifetimeCondition NetCondition);
 	
 	/** 
 	 * Set subobject NetCondition for a subobject, the conditions is used to determine if the SubObject should be replicated or not.
 	 * @note: As the filtering is done at the serialization level it is typically more efficient to use a separate object for connection 
 	 * specific data as filtering can be done at a higher level.
 	 */
-	ENGINE_API static void SetActorComponentNetCondition(const UActorComponent* SubObject, ELifetimeCondition Condition);
+	ENGINE_API static void SetNetConditionForActorComponent(const UActorComponent* SubObject, ELifetimeCondition Condition);
 
 	/** Update group memberships used by COND_Group sub object filtering for the specified SubObject */
 	ENGINE_API static void UpdateSubObjectGroupMemberships(const UObject* SubObject, const UWorld* World);
diff --git a/Engine/Source/Runtime/Experimental/Iris/Core/Private/Iris/ReplicationSystem/NetRefHandleManager.cpp b/Engine/Source/Runtime/Experimental/Iris/Core/Private/Iris/ReplicationSystem/NetRefHandleManager.cpp
--- a/Engine/Source/Runtime/Experimental/Iris/Core/Private/Iris/ReplicationSystem/NetRefHandleManager.cpp
+++ b/Engine/Source/Runtime/Experimental/Iris/Core/Private/Iris/ReplicationSystem/NetRefHandleManager.cpp
@@ -727,7 +727,7 @@ void FNetRefHandleManager::RemoveSubObject(FNetRefHandle Handle)
 	}
 }
 
-bool FNetRefHandleManager::SetSubObjectNetCondition(FInternalNetRefIndex SubObjectInternalIndex, FLifeTimeConditionStorage SubObjectCondition)
+bool FNetRefHandleManager::SetSubObjectNetCondition(FInternalNetRefIndex SubObjectInternalIndex, FLifeTimeConditionStorage SubObjectCondition, bool& bOutWasModified)
 {
 	if (ensure(SubObjectInternalIndex != InvalidInternalIndex))
 	{
@@ -742,14 +742,24 @@ bool FNetRefHandleManager::SetSubObjectNetCondition(FInternalNetRefIndex SubObje
 				const int32 SubObjectArrayIndex = SubObjectsArray->Find(SubObjectInternalIndex);
 				if (ensure(SubObjectArrayIndex != INDEX_NONE))
 				{
-					if (!SubObjectConditionals)
+					if (!SubObjectConditionals && SubObjectCondition != ELifetimeCondition::COND_None)
 					{
 						SubObjectConditionals = &SubObjects.GetOrCreateSubObjectConditionalsArray(SubObjectParentIndex);
 					}
+					else
+					{
+						// Valid but not modified.
+						bOutWasModified = false;
+						return true;
+					}					
 
 					check(SubObjectConditionals->Num() == SubObjectsArray->Num());
 
+					const FLifeTimeConditionStorage OldCondition = (*SubObjectConditionals)[SubObjectArrayIndex];
+					bOutWasModified = OldCondition != SubObjectCondition;
+
 					(*SubObjectConditionals)[SubObjectArrayIndex] = SubObjectCondition;
+					
 					return true;
 				}
 			}
@@ -757,6 +767,7 @@ bool FNetRefHandleManager::SetSubObjectNetCondition(FInternalNetRefIndex SubObje
 		}
 	}
 
+	bOutWasModified = false;
 	return false;
 }
 
diff --git a/Engine/Source/Runtime/Experimental/Iris/Core/Private/Iris/ReplicationSystem/NetRefHandleManager.h b/Engine/Source/Runtime/Experimental/Iris/Core/Private/Iris/ReplicationSystem/NetRefHandleManager.h
--- a/Engine/Source/Runtime/Experimental/Iris/Core/Private/Iris/ReplicationSystem/NetRefHandleManager.h
+++ b/Engine/Source/Runtime/Experimental/Iris/Core/Private/Iris/ReplicationSystem/NetRefHandleManager.h
@@ -243,7 +243,7 @@ public:
 	FNetRefHandle GetRootObjectOfSubObject(FNetRefHandle SubObjectHandle) const;
 	FInternalNetRefIndex GetRootObjectInternalIndexOfSubObject(FInternalNetRefIndex SubObjectIndex) const;
 
-	bool SetSubObjectNetCondition(FInternalNetRefIndex SubObjectInternalIndex, FLifeTimeConditionStorage SubObjectCondition);
+	bool SetSubObjectNetCondition(FInternalNetRefIndex SubObjectInternalIndex, FLifeTimeConditionStorage SubObjectCondition, bool& bOutWasModified);
 
 	bool IsSubObject(FInternalNetRefIndex InternalIndex) const
 	{
@@ -289,7 +289,7 @@ public:
 	// Get child subobjects for a object, used when we do hierarchical operations such as conditional serialization
 	inline TArrayView<const FInternalNetRefIndex> GetChildSubObjects(FInternalNetRefIndex ParentObjectIndex) const;
 
-	// Get child subobjects and condtionals array if one exists, if there are no child subobjects the method returns false
+	// Get child subobjects and conditionals array if one exists, if there are no child subobjects the method returns false
 	inline bool GetChildSubObjects(FInternalNetRefIndex OwnerIndex, FChildSubObjectsInfo& OutInfo) const;
 
 	/** Get the map that translates RefHandles into InternalIndexes */
diff --git a/Engine/Source/Runtime/Experimental/Iris/Core/Private/Iris/ReplicationSystem/ObjectReplicationBridge.cpp b/Engine/Source/Runtime/Experimental/Iris/Core/Private/Iris/ReplicationSystem/ObjectReplicationBridge.cpp
--- a/Engine/Source/Runtime/Experimental/Iris/Core/Private/Iris/ReplicationSystem/ObjectReplicationBridge.cpp
+++ b/Engine/Source/Runtime/Experimental/Iris/Core/Private/Iris/ReplicationSystem/ObjectReplicationBridge.cpp
@@ -658,10 +658,15 @@ void UObjectReplicationBridge::SetSubObjectNetCondition(FNetRefHandle SubObjectR
 	FNetRefHandleManager& LocalNetRefHandleManager = ReplicationSystemInternal->GetNetRefHandleManager();
 
 	const FInternalNetRefIndex SubObjectInternalIndex = LocalNetRefHandleManager.GetInternalIndex(SubObjectRefHandle);
-	if (LocalNetRefHandleManager.SetSubObjectNetCondition(SubObjectInternalIndex, (int8)Condition))
+	bool bWasModified = false;
+	
+	if (LocalNetRefHandleManager.SetSubObjectNetCondition(SubObjectInternalIndex, (int8)Condition, bWasModified))
 	{
-		UE_LOG_OBJECTREPLICATIONBRIDGE(Verbose, TEXT("SetSubObjectNetCondition for SubObject %s Condition %s"), *PrintObjectFromNetRefHandle(SubObjectRefHandle), *UEnum::GetValueAsString<ELifetimeCondition>(Condition));
-		MarkNetObjectStateDirty(ReplicationSystem->GetId(), SubObjectInternalIndex);
+		if (bWasModified)
+		{
+			UE_LOG_OBJECTREPLICATIONBRIDGE(Verbose, TEXT("SetSubObjectNetCondition for SubObject %s Condition %s"), *PrintObjectFromNetRefHandle(SubObjectRefHandle), *UEnum::GetValueAsString<ELifetimeCondition>(Condition));
+			MarkNetObjectStateDirty(ReplicationSystem->GetId(), SubObjectInternalIndex);
+		}
 	}
 	else
 	{
diff --git a/Engine/Source/Runtime/Experimental/Iris/Core/Private/Iris/ReplicationSystem/ReplicationWriter.cpp b/Engine/Source/Runtime/Experimental/Iris/Core/Private/Iris/ReplicationSystem/ReplicationWriter.cpp
--- a/Engine/Source/Runtime/Experimental/Iris/Core/Private/Iris/ReplicationSystem/ReplicationWriter.cpp
+++ b/Engine/Source/Runtime/Experimental/Iris/Core/Private/Iris/ReplicationSystem/ReplicationWriter.cpp
@@ -2724,8 +2724,18 @@ FReplicationWriter::EWriteObjectStatus FReplicationWriter::WriteObjectAndSubObje
 		bool bHasDirtySubObjects = false;
 		
 		FReplicationConditionals::FSubObjectsToReplicateArray SubObjectsToReplicate;
-		ReplicationConditionals->GetSubObjectsToReplicate(Parameters.ConnectionId, InternalIndex, SubObjectsToReplicate);		
+		ReplicationConditionals->GetSubObjectsToReplicate(Parameters.ConnectionId, InternalIndex, SubObjectsToReplicate);
 
+		// Patch up list with subobjects pending destroy.
+		for (uint32 SubObjectIndex : NetRefHandleManager->GetSubObjects(InternalIndex))
+		{
+			const FReplicationInfo& SubObjectInfo = GetReplicationInfo(SubObjectIndex);
+			if (SubObjectInfo.GetState() == EReplicatedObjectState::SubObjectPendingDestroy && !SubObjectsToReplicate.Contains(SubObjectIndex))
+			{
+				SubObjectsToReplicate.Add(SubObjectIndex);
+			}
+		}
+		
 		for (FInternalNetRefIndex SubObjectInternalIndex : SubObjectsToReplicate)
 		{
 			if (!ObjectsWithDirtyChanges.GetBit(SubObjectInternalIndex))
