Iris
* Removed ReplicationSystem::PreSendUpdate in favor of two different functions: NetUpdate() and TickPostReceive() to clarify their usage.
* Added Pre/Post ReceiveUpdate callbacks for the bridge

#rnx

diff --git a/Engine/Plugins/Runtime/ReplicationSystemTestPlugin/Source/Private/Tests/DataStream/TestChunkedDataStream.cpp b/Engine/Plugins/Runtime/ReplicationSystemTestPlugin/Source/Private/Tests/DataStream/TestChunkedDataStream.cpp
--- a/Engine/Plugins/Runtime/ReplicationSystemTestPlugin/Source/Private/Tests/DataStream/TestChunkedDataStream.cpp
+++ b/Engine/Plugins/Runtime/ReplicationSystemTestPlugin/Source/Private/Tests/DataStream/TestChunkedDataStream.cpp
@@ -77,7 +77,7 @@ protected:
 		for (uint32 It = 0; It <= InFlightCount; ++It)
 		{
 			Server->SetMaxSendPacketSize(FMath::RandRange(PacketMinSize, PacketMaxSize) & ~0x3);
-			Server->PreSendUpdate();
+			Server->NetUpdate();
 
 			SentCount += Server->SendTo(Client) ? 1U : 0U;
 
diff --git a/Engine/Plugins/Runtime/ReplicationSystemTestPlugin/Source/Private/Tests/DataStream/TestDataStream.cpp b/Engine/Plugins/Runtime/ReplicationSystemTestPlugin/Source/Private/Tests/DataStream/TestDataStream.cpp
--- a/Engine/Plugins/Runtime/ReplicationSystemTestPlugin/Source/Private/Tests/DataStream/TestDataStream.cpp
+++ b/Engine/Plugins/Runtime/ReplicationSystemTestPlugin/Source/Private/Tests/DataStream/TestDataStream.cpp
@@ -727,7 +727,7 @@ UE_NET_TEST_FIXTURE(FTestDynamicCreateDataStreamFixture, TestDynamicDataStream_C
 	UDataStream* ServerNetTokenStream = ServerReplicationSystem->OpenDataStream(Client->ConnectionIdOnServer, DataStreamName);
 
 	// Put data in flight containing create
-	Server->PreSendUpdate();
+	Server->NetUpdate();
 	Server->SendTo(Client);
 	Server->PostSendUpdate();
 
@@ -741,7 +741,7 @@ UE_NET_TEST_FIXTURE(FTestDynamicCreateDataStreamFixture, TestDynamicDataStream_C
 	UE_NET_ASSERT_TRUE(GetDataStreamStateOnServer() == UDataStream::EDataStreamState::PendingClose);
 
 	// Put data in flight containing close?
-	Server->PreSendUpdate();
+	Server->NetUpdate();
 	Server->SendTo(Client);
 	Server->PostSendUpdate();
 
@@ -793,23 +793,23 @@ UE_NET_TEST_FIXTURE(FTestDynamicCreateDataStreamFixture, TestDynamicDataStreamCl
 
 	// Put some data in flight in separate packets
 	FNetToken ServerHelloToken = CreateAndExportTokenToClient(TEXT("Hello"));
-	Server->PreSendUpdate();
+	Server->NetUpdate();
 	Server->SendTo(Client, TEXT("Hello"));
 	Server->PostSendUpdate();
 
 	FNetToken ServerHello2Token = CreateAndExportTokenToClient(TEXT("Hello2"));
-	Server->PreSendUpdate();
+	Server->NetUpdate();
 	Server->SendTo(Client, TEXT("Hello2"));
 	Server->PostSendUpdate();
 
 	FNetToken ServerHello3Token = CreateAndExportTokenToClient(TEXT("Hello3"));
-	Server->PreSendUpdate();
+	Server->NetUpdate();
 	Server->SendTo(Client, TEXT("Hello3"));
 	Server->PostSendUpdate();
 	
 	// Close from server
 	ServerReplicationSystem->CloseDataStream(Client->ConnectionIdOnServer, DataStreamName);
-	Server->PreSendUpdate();
+	Server->NetUpdate();
 	Server->SendTo(Client, TEXT("Close"));
 	Server->PostSendUpdate();
 
diff --git a/Engine/Plugins/Runtime/ReplicationSystemTestPlugin/Source/Private/Tests/ReplicationState/TestDerivedStruct.cpp b/Engine/Plugins/Runtime/ReplicationSystemTestPlugin/Source/Private/Tests/ReplicationState/TestDerivedStruct.cpp
--- a/Engine/Plugins/Runtime/ReplicationSystemTestPlugin/Source/Private/Tests/ReplicationState/TestDerivedStruct.cpp
+++ b/Engine/Plugins/Runtime/ReplicationSystemTestPlugin/Source/Private/Tests/ReplicationState/TestDerivedStruct.cpp
@@ -199,7 +199,7 @@ UE_NET_TEST_FIXTURE(FTestDerivedStructInObjectFixture, CanReplicateDerivedStruct
 	ServerObject->DerivedStruct.ByteMember2 ^= 47;
 
 	// Replicate
-	Server->PreSendUpdate();
+	Server->NetUpdate();
 	Server->SendAndDeliverTo(Client, DeliverPacket);
 	Server->PostSendUpdate();
 
@@ -226,7 +226,7 @@ UE_NET_TEST_FIXTURE(FTestDerivedStructInObjectFixture, CanReplicateDerivedStruct
 	}
 
 	// Replicate
-	Server->PreSendUpdate();
+	Server->NetUpdate();
 	Server->SendAndDeliverTo(Client, DeliverPacket);
 	Server->PostSendUpdate();
 
@@ -253,7 +253,7 @@ UE_NET_TEST_FIXTURE(FTestDerivedStructInObjectFixture, NotReplicatedMemberInDeri
 	UTestDerivedStruct_TestObject_Member* ServerObject = CreateObjectWithMember();
 
 	// Replicate
-	Server->PreSendUpdate();
+	Server->NetUpdate();
 	Server->SendAndDeliverTo(Client, DeliverPacket);
 	Server->PostSendUpdate();
 
@@ -270,7 +270,7 @@ UE_NET_TEST_FIXTURE(FTestDerivedStructInObjectFixture, NotReplicatedMemberInDeri
 	ServerObject->DerivedStruct.ByteMember2 = 2;
 	ServerObject->DerivedStruct.ByteMember3_NotReplicated = ClientModifiedNotReplicatedPropertyValue ^ 111;
 
-	Server->PreSendUpdate();
+	Server->NetUpdate();
 	Server->SendAndDeliverTo(Client, DeliverPacket);
 	Server->PostSendUpdate();
 
@@ -288,7 +288,7 @@ UE_NET_TEST_FIXTURE(FTestDerivedStructInObjectFixture, CanReplicateDerivedStruct
 	UTestDerivedStruct_TestObject_Member* ServerObject = CreateObjectWithMember();
 
 	// Replicate
-	Server->PreSendUpdate();
+	Server->NetUpdate();
 	Server->SendAndDeliverTo(Client, DeliverPacket);
 	Server->PostSendUpdate();
 
@@ -297,7 +297,7 @@ UE_NET_TEST_FIXTURE(FTestDerivedStructInObjectFixture, CanReplicateDerivedStruct
 
 	// Modify the property not replicated by the custom serializer.
 	ServerObject->DerivedStruct.ByteMember2 ^= 47;
-	Server->PreSendUpdate();
+	Server->NetUpdate();
 	Server->SendAndDeliverTo(Client, DeliverPacket);
 	Server->PostSendUpdate();
 
@@ -309,7 +309,7 @@ UE_NET_TEST_FIXTURE(FTestDerivedStructInObjectFixture, CanReplicateDerivedStruct
 	// Modify properties replicated by custom serializer
 	ServerObject->DerivedStruct.ByteMember0 ^= 13;
 	ServerObject->DerivedStruct.ByteMember1 ^= 37;
-	Server->PreSendUpdate();
+	Server->NetUpdate();
 	Server->SendAndDeliverTo(Client, DeliverPacket);
 	Server->PostSendUpdate();
 
@@ -326,7 +326,7 @@ UE_NET_TEST_FIXTURE(FTestDerivedStructInObjectFixture, NotAppliedMemberInDerived
 	UTestDerivedStructWithNetSerializerWithApply_TestObject_Member* ServerObject = Server->CreateObject<UTestDerivedStructWithNetSerializerWithApply_TestObject_Member>();
 
 	// Replicate
-	Server->PreSendUpdate();
+	Server->NetUpdate();
 	Server->SendAndDeliverTo(Client, DeliverPacket);
 	Server->PostSendUpdate();
 
@@ -342,7 +342,7 @@ UE_NET_TEST_FIXTURE(FTestDerivedStructInObjectFixture, NotAppliedMemberInDerived
 	ServerObject->DerivedStruct.ByteMember1 ^= 2U;
 	ServerObject->DerivedStruct.ByteMemberNotSetOnApply = ClientModifiedNotAppliedPropertyValue ^ 231U;
 
-	Server->PreSendUpdate();
+	Server->NetUpdate();
 	Server->SendAndDeliverTo(Client, DeliverPacket);
 	Server->PostSendUpdate();
 
@@ -368,7 +368,7 @@ UE_NET_TEST_FIXTURE(FTestDerivedStructInObjectFixture, NotAppliedMemberInDerived
 	}
 
 	// Replicate
-	Server->PreSendUpdate();
+	Server->NetUpdate();
 	Server->SendAndDeliverTo(Client, DeliverPacket);
 	Server->PostSendUpdate();
 
@@ -392,7 +392,7 @@ UE_NET_TEST_FIXTURE(FTestDerivedStructInObjectFixture, NotAppliedMemberInDerived
 		Element.ByteMemberNotSetOnApply ^= 231U;
 	}
 
-	Server->PreSendUpdate();
+	Server->NetUpdate();
 	Server->SendAndDeliverTo(Client, DeliverPacket);
 	Server->PostSendUpdate();
 
@@ -411,7 +411,7 @@ UE_NET_TEST_FIXTURE(FTestDerivedStructInObjectFixture, NotAppliedMemberInStructD
 	UTestDerivedStructWithNetSerializerWithApply_Inherited_TestObject_Member* ServerObject = Server->CreateObject<UTestDerivedStructWithNetSerializerWithApply_Inherited_TestObject_Member>();
 
 	// Replicate
-	Server->PreSendUpdate();
+	Server->NetUpdate();
 	Server->SendAndDeliverTo(Client, DeliverPacket);
 	Server->PostSendUpdate();
 
@@ -427,7 +427,7 @@ UE_NET_TEST_FIXTURE(FTestDerivedStructInObjectFixture, NotAppliedMemberInStructD
 	ServerObject->DerivedStruct.ByteMember2 ^= 2U;
 	ServerObject->DerivedStruct.ByteMemberNotSetOnApply = ClientModifiedNotAppliedPropertyValue ^ 231U;
 
-	Server->PreSendUpdate();
+	Server->NetUpdate();
 	Server->SendAndDeliverTo(Client, DeliverPacket);
 	Server->PostSendUpdate();
 
@@ -453,7 +453,7 @@ UE_NET_TEST_FIXTURE(FTestDerivedStructInObjectFixture, NotAppliedMemberInStructD
 	}
 
 	// Replicate
-	Server->PreSendUpdate();
+	Server->NetUpdate();
 	Server->SendAndDeliverTo(Client, DeliverPacket);
 	Server->PostSendUpdate();
 
@@ -481,7 +481,7 @@ UE_NET_TEST_FIXTURE(FTestDerivedStructInObjectFixture, NotAppliedMemberInStructD
 		Element.ByteMember3_NotReplicated = ClientModifiedNotReplicatedPropertyValue ^ 0x7FU;
 	}
 
-	Server->PreSendUpdate();
+	Server->NetUpdate();
 	Server->SendAndDeliverTo(Client, DeliverPacket);
 	Server->PostSendUpdate();
 
diff --git a/Engine/Plugins/Runtime/ReplicationSystemTestPlugin/Source/Private/Tests/ReplicationState/TestFastArrayReplicationState.cpp b/Engine/Plugins/Runtime/ReplicationSystemTestPlugin/Source/Private/Tests/ReplicationState/TestFastArrayReplicationState.cpp
--- a/Engine/Plugins/Runtime/ReplicationSystemTestPlugin/Source/Private/Tests/ReplicationState/TestFastArrayReplicationState.cpp
+++ b/Engine/Plugins/Runtime/ReplicationSystemTestPlugin/Source/Private/Tests/ReplicationState/TestFastArrayReplicationState.cpp
@@ -201,7 +201,7 @@ UE_NET_TEST_FIXTURE(FReplicationSystemServerClientTestFixture, TestAddItem)
 	FReplicationSystemTestClient* Client = CreateClient();
 
 	// Send and deliver packet
-	Server->PreSendUpdate();
+	Server->NetUpdate();
 	Server->SendAndDeliverTo(Client, true);
 	Server->PostSendUpdate();
 
@@ -230,7 +230,7 @@ UE_NET_TEST_FIXTURE(FReplicationSystemServerClientTestFixture, TestAddItem)
 	ServerObject->FastArray.Edit().Add(NewEntry);
 
 	// Send and deliver packet
-	Server->PreSendUpdate();
+	Server->NetUpdate();
 	Server->SendAndDeliverTo(Client, true);
 	Server->PostSendUpdate();
 
@@ -253,7 +253,7 @@ UE_NET_TEST_FIXTURE(FReplicationSystemServerClientTestFixture, TestRemoveItem)
 	FReplicationSystemTestClient* Client = CreateClient();
 
 	// Send and deliver packet
-	Server->PreSendUpdate();
+	Server->NetUpdate();
 	Server->SendAndDeliverTo(Client, true);
 	Server->PostSendUpdate();
 
@@ -282,7 +282,7 @@ UE_NET_TEST_FIXTURE(FReplicationSystemServerClientTestFixture, TestRemoveItem)
 	ServerObject->FastArray.Edit().Add(NewEntry);
 
 	// Send and deliver packet
-	Server->PreSendUpdate();
+	Server->NetUpdate();
 	Server->SendAndDeliverTo(Client, true);
 	Server->PostSendUpdate();
 
@@ -300,7 +300,7 @@ UE_NET_TEST_FIXTURE(FReplicationSystemServerClientTestFixture, TestRemoveItem)
 	ServerObject->FastArray.Edit().Remove(0);
 
 	// Send and deliver packet
-	Server->PreSendUpdate();
+	Server->NetUpdate();
 	Server->SendAndDeliverTo(Client, true);
 	Server->PostSendUpdate();
 
@@ -321,7 +321,7 @@ UE_NET_TEST_FIXTURE(FReplicationSystemServerClientTestFixture, TestLocalItemAreP
 	FReplicationSystemTestClient* Client = CreateClient();
 
 	// Send and deliver packet
-	Server->PreSendUpdate();
+	Server->NetUpdate();
 	Server->SendAndDeliverTo(Client, true);
 	Server->PostSendUpdate();
 
@@ -359,7 +359,7 @@ UE_NET_TEST_FIXTURE(FReplicationSystemServerClientTestFixture, TestLocalItemAreP
 	ServerObject->FastArray.Edit().Add(NewEntry);
 
 	// Send and deliver packet
-	Server->PreSendUpdate();
+	Server->NetUpdate();
 	Server->SendAndDeliverTo(Client, true);
 	Server->PostSendUpdate();
 
@@ -379,7 +379,7 @@ UE_NET_TEST_FIXTURE(FReplicationSystemServerClientTestFixture, TestLocalItemAreP
 	ServerObject->FastArray.Edit().Remove(0);
 
 	// Send and deliver packet
-	Server->PreSendUpdate();
+	Server->NetUpdate();
 	Server->SendAndDeliverTo(Client, true);
 	Server->PostSendUpdate();
 
@@ -404,7 +404,7 @@ UE_NET_TEST_FIXTURE(FReplicationSystemServerClientTestFixture, TestModifyItem)
 	FReplicationSystemTestClient* Client = CreateClient();
 
 	// Send and deliver packet
-	Server->PreSendUpdate();
+	Server->NetUpdate();
 	Server->SendAndDeliverTo(Client, true);
 	Server->PostSendUpdate();
 
@@ -433,7 +433,7 @@ UE_NET_TEST_FIXTURE(FReplicationSystemServerClientTestFixture, TestModifyItem)
 	ServerObject->FastArray.Edit().Add(NewEntry);
 
 	// Send and deliver packet
-	Server->PreSendUpdate();
+	Server->NetUpdate();
 	Server->SendAndDeliverTo(Client, true);
 	Server->PostSendUpdate();
 
@@ -455,7 +455,7 @@ UE_NET_TEST_FIXTURE(FReplicationSystemServerClientTestFixture, TestModifyItem)
 	ServerObject->FastArray.Edit()[0].RepInt32 = 3U;
 	
 	// Send and deliver packet
-	Server->PreSendUpdate();
+	Server->NetUpdate();
 	Server->SendAndDeliverTo(Client, true);
 	Server->PostSendUpdate();
 
@@ -477,7 +477,7 @@ UE_NET_TEST_FIXTURE(FReplicationSystemServerClientTestFixture, TestInsertItem)
 	FReplicationSystemTestClient* Client = CreateClient();
 
 	// Send and deliver packet
-	Server->PreSendUpdate();
+	Server->NetUpdate();
 	Server->SendAndDeliverTo(Client, true);
 	Server->PostSendUpdate();
 
@@ -514,7 +514,7 @@ UE_NET_TEST_FIXTURE(FReplicationSystemServerClientTestFixture, TestInsertItem)
 	ServerObject->FastArray.Edit()[1].RepInt32 = 13U;
 
 	// Send and deliver packet
-	Server->PreSendUpdate();
+	Server->NetUpdate();
 	Server->SendAndDeliverTo(Client, true);
 	Server->PostSendUpdate();
 
@@ -536,7 +536,7 @@ UE_NET_TEST_FIXTURE(FReplicationSystemServerClientTestFixture, TestInsertItem)
 	ServerObject->FastArray.Edit()[0].RepInt32 = 3U;
 	
 	// Send and deliver packet
-	Server->PreSendUpdate();
+	Server->NetUpdate();
 	Server->SendAndDeliverTo(Client, true);
 	Server->PostSendUpdate();
 
@@ -559,7 +559,7 @@ UE_NET_TEST_FIXTURE(FReplicationSystemServerClientTestFixture, TestModifyItemDoe
 	FReplicationSystemTestClient* Client = CreateClient();
 
 	// Send and deliver packet
-	Server->PreSendUpdate();
+	Server->NetUpdate();
 	Server->SendAndDeliverTo(Client, true);
 	Server->PostSendUpdate();
 
@@ -591,7 +591,7 @@ UE_NET_TEST_FIXTURE(FReplicationSystemServerClientTestFixture, TestModifyItemDoe
 	FastArrayEditor.Add(NewEntry);
 
 	// Send and deliver packet
-	Server->PreSendUpdate();
+	Server->NetUpdate();
 	Server->SendAndDeliverTo(Client, true);
 	Server->PostSendUpdate();
 
@@ -617,7 +617,7 @@ UE_NET_TEST_FIXTURE(FReplicationSystemServerClientTestFixture, TestModifyItemDoe
 	FastArrayEditor.Edit(0).RepInt32 = 3;
 	
 	// Send and deliver packet
-	Server->PreSendUpdate();
+	Server->NetUpdate();
 	Server->SendAndDeliverTo(Client, true);
 	Server->PostSendUpdate();
 
@@ -657,7 +657,7 @@ UE_NET_TEST_FIXTURE(FReplicationSystemServerClientTestFixture, TestAddItemWithou
 	FReplicationSystemTestClient* Client = CreateClient();
 
 	// Send and deliver packet
-	Server->PreSendUpdate();
+	Server->NetUpdate();
 
 	// Verify that poll has modified the index
 	UE_NET_ASSERT_NE(ServerObject->FastArray.GetItemArray()[0].ReplicationID, -1);
@@ -685,7 +685,7 @@ UE_NET_TEST_FIXTURE(FReplicationSystemServerClientTestFixture, TestChangeMaskWra
 	FReplicationSystemTestClient* Client = CreateClient();
 
 	// Send and deliver packet
-	Server->PreSendUpdate();
+	Server->NetUpdate();
 	Server->SendAndDeliverTo(Client, true);
 	Server->PostSendUpdate();
 
@@ -717,7 +717,7 @@ UE_NET_TEST_FIXTURE(FReplicationSystemServerClientTestFixture, TestChangeMaskWra
 	int32 CurrentItemCount = FastArrayEditor.Num();
 
 	// Send and deliver packet
-	Server->PreSendUpdate();
+	Server->NetUpdate();
 	Server->SendAndDeliverTo(Client, true);
 	Server->PostSendUpdate();
 
@@ -745,7 +745,7 @@ UE_NET_TEST_FIXTURE(FReplicationSystemServerClientTestFixture, TestChangeMaskWra
 	}
 	
 	// Send and deliver packet
-	Server->PreSendUpdate();
+	Server->NetUpdate();
 	Server->SendAndDeliverTo(Client, true);
 	Server->PostSendUpdate();
 
@@ -764,7 +764,7 @@ UE_NET_TEST_FIXTURE(FReplicationSystemServerClientTestFixture, TestChangeMaskWra
 	}
 
 	// Send and deliver packet
-	Server->PreSendUpdate();
+	Server->NetUpdate();
 	Server->SendAndDeliverTo(Client, true);
 	Server->PostSendUpdate();
 
@@ -808,7 +808,7 @@ UE_NET_TEST_FIXTURE(FReplicationSystemServerClientTestFixture, TestPostReplicate
 	Server->ReplicationSystem->AddToGroup(Server->GetReplicationSystem()->GetNotReplicatedNetObjectGroup(), ServerReferencedObjectB->NetRefHandle);
 	
 	// Replicate
-	Server->PreSendUpdate();
+	Server->NetUpdate();
 	Server->SendAndDeliverTo(Client, true);
 	Server->PostSendUpdate();
 
@@ -838,7 +838,7 @@ UE_NET_TEST_FIXTURE(FReplicationSystemServerClientTestFixture, TestPostReplicate
 	Server->ReplicationSystem->RemoveFromGroup(Server->GetReplicationSystem()->GetNotReplicatedNetObjectGroup(), ServerReferencedObjectA->NetRefHandle);
 
 	// Replicate
-	Server->PreSendUpdate();
+	Server->NetUpdate();
 	Server->SendAndDeliverTo(Client, true);
 	Server->PostSendUpdate();
 
@@ -857,7 +857,7 @@ UE_NET_TEST_FIXTURE(FReplicationSystemServerClientTestFixture, TestPostReplicate
 	Server->ReplicationSystem->RemoveFromGroup(Server->GetReplicationSystem()->GetNotReplicatedNetObjectGroup(), ServerReferencedObjectB->NetRefHandle);
 
 	// Replicate
-	Server->PreSendUpdate();
+	Server->NetUpdate();
 	Server->SendAndDeliverTo(Client, true);
 	Server->PostSendUpdate();
 
@@ -879,7 +879,7 @@ UE_NET_TEST_FIXTURE(FReplicationSystemServerClientTestFixture, TestFastArrayWith
 
 
 	// Send and deliver packet
-	Server->PreSendUpdate();
+	Server->NetUpdate();
 	Server->SendAndDeliverTo(Client, true);
 	Server->PostSendUpdate();
 
@@ -909,7 +909,7 @@ UE_NET_TEST_FIXTURE(FReplicationSystemServerClientTestFixture, TestFastArrayWith
 	ServerObject->FastArray.Edit().Add(NewEntry);
 
 	// Send and deliver packet
-	Server->PreSendUpdate();
+	Server->NetUpdate();
 	Server->SendAndDeliverTo(Client, true);
 	Server->PostSendUpdate();
 
@@ -928,7 +928,7 @@ UE_NET_TEST_FIXTURE(FReplicationSystemServerClientTestFixture, TestFastArrayWith
 	ServerObject->FastArray.ExtraInt = 3;
 
 	// Send and deliver packet
-	Server->PreSendUpdate();
+	Server->NetUpdate();
 	Server->SendAndDeliverTo(Client, true);
 	Server->PostSendUpdate();
 
@@ -944,7 +944,7 @@ UE_NET_TEST_FIXTURE(FReplicationSystemServerClientTestFixture, TestLostDataDirti
 	FReplicationSystemTestClient* Client = CreateClient();
 
 	// Send and deliver packet
-	Server->PreSendUpdate();
+	Server->NetUpdate();
 	Server->SendAndDeliverTo(Client, true);
 	Server->PostSendUpdate();
 
@@ -961,7 +961,7 @@ UE_NET_TEST_FIXTURE(FReplicationSystemServerClientTestFixture, TestLostDataDirti
 	ServerObject->FastArray.Edit().Add(NewEntry);
 
 	// Put data on wire, but hold off delivery
-	Server->PreSendUpdate();
+	Server->NetUpdate();
 	Server->SendTo(Client);
 	Server->PostSendUpdate();
 
@@ -969,7 +969,7 @@ UE_NET_TEST_FIXTURE(FReplicationSystemServerClientTestFixture, TestLostDataDirti
 	++ServerObject->FastArray.Edit()[0].RepInt32;
 
 	// Put latest state on the wire
-	Server->PreSendUpdate();
+	Server->NetUpdate();
 	Server->SendTo(Client);
 	Server->PostSendUpdate();
 
@@ -980,7 +980,7 @@ UE_NET_TEST_FIXTURE(FReplicationSystemServerClientTestFixture, TestLostDataDirti
 	Server->DeliverTo(Client, true);
 
 	// Send and deliver packet
-	Server->PreSendUpdate();
+	Server->NetUpdate();
 	Server->SendAndDeliverTo(Client, true);
 	Server->PostSendUpdate();
 
diff --git a/Engine/Plugins/Runtime/ReplicationSystemTestPlugin/Source/Private/Tests/ReplicationSystem/Conditionals/TestConditionals.cpp b/Engine/Plugins/Runtime/ReplicationSystemTestPlugin/Source/Private/Tests/ReplicationSystem/Conditionals/TestConditionals.cpp
--- a/Engine/Plugins/Runtime/ReplicationSystemTestPlugin/Source/Private/Tests/ReplicationSystem/Conditionals/TestConditionals.cpp
+++ b/Engine/Plugins/Runtime/ReplicationSystemTestPlugin/Source/Private/Tests/ReplicationSystem/Conditionals/TestConditionals.cpp
@@ -1,5 +1,7 @@
 // Copyright Epic Games, Inc. All Rights Reserved.
 
+#include "UObject/CoreNetTypes.h"
+
 #include "NetworkAutomationTest.h"
 #include "NetworkAutomationTestMacros.h"
 #include "Iris/ReplicationSystem/ReplicationSystem.h"
@@ -35,7 +37,7 @@ UE_NET_TEST_FIXTURE(FTestConditionalsFixture, ToOwnerStateIsReplicatedToOwner)
 	Server->ReplicationSystem->SetOwningNetConnection(ServerObject->NetRefHandle, Client->ConnectionIdOnServer);
 
 	// Send and deliver packet
-	Server->PreSendUpdate();
+	Server->NetUpdate();
 	Server->SendAndDeliverTo(Client, DeliverPacket);
 	Server->PostSendUpdate();
 	
@@ -68,7 +70,7 @@ UE_NET_TEST_FIXTURE(FTestConditionalsFixture, ToOwnerStateIsNotReplicatedToNonOw
 	// Do not set owner. Optionally one could set it to some connection other than the client's.
 
 	// Send and deliver packet
-	Server->PreSendUpdate();
+	Server->NetUpdate();
 	Server->SendAndDeliverTo(Client, DeliverPacket);
 	Server->PostSendUpdate();
 	
@@ -101,7 +103,7 @@ UE_NET_TEST_FIXTURE(FTestConditionalsFixture, SkipOwnerStateIsNotReplicatedToOwn
 	Server->ReplicationSystem->SetOwningNetConnection(ServerObject->NetRefHandle, Client->ConnectionIdOnServer);
 
 	// Send and deliver packet
-	Server->PreSendUpdate();
+	Server->NetUpdate();
 	Server->SendAndDeliverTo(Client, DeliverPacket);
 	Server->PostSendUpdate();
 	
@@ -132,7 +134,7 @@ UE_NET_TEST_FIXTURE(FTestConditionalsFixture, SkipOwnerStateIsReplicatedToNonOwn
 	// Do not set owner. Optionally one could set it to some connection other than the client's.
 
 	// Send and deliver packet
-	Server->PreSendUpdate();
+	Server->NetUpdate();
 	Server->SendAndDeliverTo(Client, DeliverPacket);
 	Server->PostSendUpdate();
 	
@@ -165,7 +167,7 @@ UE_NET_TEST_FIXTURE(FTestConditionalsFixture, SimulatedIsReplicatedToSimulated)
 	// Could set autonomous condition to be some connection other than the create client, or not set at all as is done now.
 
 	// Send and deliver packet
-	Server->PreSendUpdate();
+	Server->NetUpdate();
 	Server->SendAndDeliverTo(Client, DeliverPacket);
 	Server->PostSendUpdate();
 	
@@ -201,7 +203,7 @@ UE_NET_TEST_FIXTURE(FTestConditionalsFixture, SimulatedIsNotReplicatedToAutonomo
 	Server->ReplicationSystem->SetReplicationConditionConnectionFilter(ServerObject->NetRefHandle, EReplicationCondition::RoleAutonomous, Client->ConnectionIdOnServer, true);
 
 	// Send and deliver packet
-	Server->PreSendUpdate();
+	Server->NetUpdate();
 	Server->SendAndDeliverTo(Client, DeliverPacket);
 	Server->PostSendUpdate();
 	
@@ -234,7 +236,7 @@ UE_NET_TEST_FIXTURE(FTestConditionalsFixture, AutonomousIsReplicatedToAutonomous
 	Server->ReplicationSystem->SetReplicationConditionConnectionFilter(ServerObject->NetRefHandle, EReplicationCondition::RoleAutonomous, Client->ConnectionIdOnServer, true);
 
 	// Send and deliver packet
-	Server->PreSendUpdate();
+	Server->NetUpdate();
 	Server->SendAndDeliverTo(Client, DeliverPacket);
 	Server->PostSendUpdate();
 	
@@ -266,7 +268,7 @@ UE_NET_TEST_FIXTURE(FTestConditionalsFixture, AutonomousIsNotReplicatedToSimulat
 	Server->ReplicationSystem->SetReplicationConditionConnectionFilter(ServerObject->NetRefHandle, EReplicationCondition::RoleAutonomous, AutonomousClient->ConnectionIdOnServer, true);
 
 	// Send and deliver packet
-	Server->PreSendUpdate();
+	Server->NetUpdate();
 	Server->SendAndDeliverTo(Client, DeliverPacket);
 	Server->PostSendUpdate();
 	
@@ -299,7 +301,7 @@ UE_NET_TEST_FIXTURE(FTestConditionalsFixture, SimulatedOrPhysicsIsReplicatedWhen
 	Server->ReplicationSystem->SetReplicationCondition(ServerObject->NetRefHandle, EReplicationCondition::ReplicatePhysics, true);
 
 	// Send and deliver packet
-	Server->PreSendUpdate();
+	Server->NetUpdate();
 	Server->SendAndDeliverTo(Client, DeliverPacket);
 	Server->PostSendUpdate();
 	
@@ -333,7 +335,7 @@ UE_NET_TEST_FIXTURE(FTestConditionalsFixture, SimulatedOrPhysicsIsNotReplicatedW
 	Server->ReplicationSystem->SetReplicationCondition(ServerObject->NetRefHandle, EReplicationCondition::ReplicatePhysics, false);
 
 	// Send and deliver packet
-	Server->PreSendUpdate();
+	Server->NetUpdate();
 	Server->SendAndDeliverTo(Client, DeliverPacket);
 	Server->PostSendUpdate();
 	
@@ -364,7 +366,7 @@ UE_NET_TEST_FIXTURE(FTestConditionalsFixture, CanToggleAutonomous)
 	Server->ReplicationSystem->SetReplicationConditionConnectionFilter(ServerObject->NetRefHandle, EReplicationCondition::RoleAutonomous, Client->ConnectionIdOnServer, true);
 
 	// Send and deliver packet
-	Server->PreSendUpdate();
+	Server->NetUpdate();
 	Server->SendAndDeliverTo(Client, DeliverPacket);
 	Server->PostSendUpdate();
 	
@@ -385,7 +387,7 @@ UE_NET_TEST_FIXTURE(FTestConditionalsFixture, CanToggleAutonomous)
 	const int32 ExpectedAutonomousOnlyInt = ClientObject->ConnectionFilteredComponents[0]->AutonomousOnlyInt;
 
 	// Send and deliver packet
-	Server->PreSendUpdate();
+	Server->NetUpdate();
 	Server->SendAndDeliverTo(Client, DeliverPacket);
 	Server->PostSendUpdate();
 
@@ -413,7 +415,7 @@ UE_NET_TEST_FIXTURE(FTestConditionalsFixture, CanToggleSimulated)
 	ServerObject->ConnectionFilteredComponents[0]->SimulatedOrPhysicsNoReplayInt = 11;
 
 	// Send and deliver packet
-	Server->PreSendUpdate();
+	Server->NetUpdate();
 	Server->SendAndDeliverTo(Client, DeliverPacket);
 	Server->PostSendUpdate();
 	
@@ -443,7 +445,7 @@ UE_NET_TEST_FIXTURE(FTestConditionalsFixture, CanToggleSimulated)
 	const int32 ExpectedSimulatedOrPhysicsNoReplayInt = ClientObject->ConnectionFilteredComponents[0]->SimulatedOrPhysicsNoReplayInt;
 
 	// Send and deliver packet
-	Server->PreSendUpdate();
+	Server->NetUpdate();
 	Server->SendAndDeliverTo(Client, DeliverPacket);
 	Server->PostSendUpdate();
 
@@ -480,7 +482,7 @@ UE_NET_TEST_FIXTURE(FTestConditionalsFixture, CanMixConditions)
 	ServerObject->ConnectionFilteredComponents[0]->SimulatedOrPhysicsNoReplayInt = 11;
 	
 	// Send and deliver packets
-	Server->PreSendUpdate();
+	Server->NetUpdate();
 	for (FReplicationSystemTestClient* Client : ClientArray)
 	{
 		Server->SendAndDeliverTo(Client, DeliverPacket);
@@ -544,7 +546,7 @@ UE_NET_TEST_FIXTURE(FTestConditionalsFixture, CanSwitchAutonmousConnection)
 	ServerObject->ConnectionFilteredComponents[0]->SimulatedOrPhysicsNoReplayInt = 11;
 
 	// Send and deliver packets
-	Server->PreSendUpdate();
+	Server->NetUpdate();
 	for (FReplicationSystemTestClient* Client : ClientArray)
 	{
 		Server->SendAndDeliverTo(Client, DeliverPacket);
@@ -573,7 +575,7 @@ UE_NET_TEST_FIXTURE(FTestConditionalsFixture, CanSwitchAutonmousConnection)
 	ServerObject->ConnectionFilteredComponents[0]->SimulatedOrPhysicsNoReplayInt = 11+1;
 
 	// Send and deliver packets
-	Server->PreSendUpdate();
+	Server->NetUpdate();
 	for (FReplicationSystemTestClient* Client : ClientArray)
 	{
 		Server->SendAndDeliverTo(Client, DeliverPacket);
@@ -643,7 +645,7 @@ UE_NET_TEST_FIXTURE(FTestConditionalsFixture, SubObjectConditionsMatchesParentOb
 	}
 	
 	// Send and deliver packets
-	Server->PreSendUpdate();
+	Server->NetUpdate();
 	for (FReplicationSystemTestClient* Client : ClientArray)
 	{
 		Server->SendAndDeliverTo(Client, DeliverPacket);
@@ -710,7 +712,7 @@ UE_NET_TEST_FIXTURE(FTestConditionalsFixture, WithSkipOwnerSubObjectIsNotReplica
 	ServerSubObject->IntA = 13; 
 
 	// Send and deliver packet
-	Server->PreSendUpdate();
+	Server->NetUpdate();
 	Server->SendAndDeliverTo(Client, DeliverPacket);
 	Server->PostSendUpdate();
 	
@@ -742,12 +744,12 @@ UE_NET_TEST_FIXTURE(FTestConditionalsFixture, WithSkipOwnerSubObjectReplicatedTo
 	ServerSubObject->IntA = 13; 
 
 	// Send and deliver packet
-	Server->PreSendUpdate();
+	Server->NetUpdate();
 	Server->SendAndDeliverTo(Client, DeliverPacket);
 	Server->PostSendUpdate();
 
 	// Send and deliver packet
-	Server->PreSendUpdate();
+	Server->NetUpdate();
 	Server->SendAndDeliverTo(OtherClient, DeliverPacket);
 	Server->PostSendUpdate();
 	
@@ -786,7 +788,7 @@ UE_NET_TEST_FIXTURE(FTestConditionalsFixture, WithSkipOwnerSubObjectReplicatedTo
 	ServerSubObject->IntA = 13; 
 
 	// Send and deliver packet
-	Server->PreSendUpdate();
+	Server->NetUpdate();
 	Server->SendAndDeliverTo(Client, DeliverPacket);
 	Server->PostSendUpdate();
 	
@@ -804,7 +806,7 @@ UE_NET_TEST_FIXTURE(FTestConditionalsFixture, WithSkipOwnerSubObjectReplicatedTo
 	Server->ReplicationBridge->SetSubObjectNetCondition(ServerSubObject->NetRefHandle, ELifetimeCondition::COND_None);
 
 	// Send and deliver packet
-	Server->PreSendUpdate();
+	Server->NetUpdate();
 	Server->SendAndDeliverTo(Client, DeliverPacket);
 	Server->PostSendUpdate();
 
@@ -836,7 +838,7 @@ UE_NET_TEST_FIXTURE(FTestConditionalsFixture, WithSkipOwnerSubObjectOwnerCanBeCh
 	ServerSubObject->IntA = 13; 
 
 	// Send and deliver packet
-	Server->PreSendUpdate();
+	Server->NetUpdate();
 	Server->SendAndDeliverTo(Client, DeliverPacket);
 	Server->PostSendUpdate();
 	
@@ -854,7 +856,7 @@ UE_NET_TEST_FIXTURE(FTestConditionalsFixture, WithSkipOwnerSubObjectOwnerCanBeCh
 	Server->ReplicationSystem->SetOwningNetConnection(ServerObject->NetRefHandle, 0);
 
 	// Send and delay packet
-	Server->PreSendUpdate();
+	Server->NetUpdate();
 	Server->SendTo(Client);
 	Server->PostSendUpdate();
 
@@ -864,7 +866,7 @@ UE_NET_TEST_FIXTURE(FTestConditionalsFixture, WithSkipOwnerSubObjectOwnerCanBeCh
 	Server->ReplicationSystem->SetOwningNetConnection(ServerObject->NetRefHandle, 0);
 
 	// Send and delay packet
-	Server->PreSendUpdate();
+	Server->NetUpdate();
 	Server->SendTo(Client);
 	Server->PostSendUpdate();
 
@@ -901,7 +903,7 @@ UE_NET_TEST_FIXTURE(FTestConditionalsFixture, HierarchicalSubObjectIsNotReplicat
 	ServerSubObject->IntA = 13; 
 
 	// Send and deliver packet
-	Server->PreSendUpdate();
+	Server->NetUpdate();
 	Server->SendAndDeliverTo(Client, DeliverPacket);
 	Server->PostSendUpdate();
 	
@@ -924,7 +926,7 @@ UE_NET_TEST_FIXTURE(FTestConditionalsFixture, HierarchicalSubObjectIsNotReplicat
 	Server->ReplicationBridge->SetSubObjectNetCondition(ServerSubSubObject->NetRefHandle, ELifetimeCondition::COND_SkipOwner);
 
 	// Send and deliver packet
-	Server->PreSendUpdate();
+	Server->NetUpdate();
 	Server->SendAndDeliverTo(Client, DeliverPacket);
 	Server->PostSendUpdate();
 
@@ -941,7 +943,7 @@ UE_NET_TEST_FIXTURE(FTestConditionalsFixture, HierarchicalSubObjectIsNotReplicat
 	Server->ReplicationBridge->SetSubObjectNetCondition(ServerSubSubObject->NetRefHandle, ELifetimeCondition::COND_None);
 
 	// Send and deliver packet
-	Server->PreSendUpdate();
+	Server->NetUpdate();
 	Server->SendAndDeliverTo(Client, DeliverPacket);
 	Server->PostSendUpdate();
 
@@ -973,7 +975,7 @@ UE_NET_TEST_FIXTURE(FTestConditionalsFixture, NoneIsReplicated)
 	// Could set autonomous condition to be some connection other than the create client, or not set at all as is done now.
 
 	// Send and deliver packet
-	Server->PreSendUpdate();
+	Server->NetUpdate();
 	Server->SendAndDeliverTo(Client, DeliverPacket);
 	Server->PostSendUpdate();
 	
@@ -1004,7 +1006,7 @@ UE_NET_TEST_FIXTURE(FTestConditionalsFixture, NeverIsNotReplicated)
 	// Could set autonomous condition to be some connection other than the create client, or not set at all as is done now.
 
 	// Send and deliver packet
-	Server->PreSendUpdate();
+	Server->NetUpdate();
 	Server->SendAndDeliverTo(Client, DeliverPacket);
 	Server->PostSendUpdate();
 	
@@ -1031,7 +1033,7 @@ UE_NET_TEST_FIXTURE(FTestConditionalsFixture, SkipReplayIsReplicatedToRegularCon
 	ServerObject->ConnectionFilteredComponents[0]->SkipReplayInt = 13;
 
 	// Send and deliver packet
-	Server->PreSendUpdate();
+	Server->NetUpdate();
 	Server->SendAndDeliverTo(Client, DeliverPacket);
 	Server->PostSendUpdate();
 	
@@ -1062,7 +1064,7 @@ UE_NET_TEST_FIXTURE(FTestConditionalsFixture, ReplayOnlyIsNotReplicatedToRegular
 	// Could set autonomous condition to be some connection other than the create client, or not set at all as is done now.
 
 	// Send and deliver packet
-	Server->PreSendUpdate();
+	Server->NetUpdate();
 	Server->SendAndDeliverTo(Client, DeliverPacket);
 	Server->PostSendUpdate();
 	
diff --git a/Engine/Plugins/Runtime/ReplicationSystemTestPlugin/Source/Private/Tests/ReplicationSystem/Conditionals/TestDynamicCondition.cpp b/Engine/Plugins/Runtime/ReplicationSystemTestPlugin/Source/Private/Tests/ReplicationSystem/Conditionals/TestDynamicCondition.cpp
--- a/Engine/Plugins/Runtime/ReplicationSystemTestPlugin/Source/Private/Tests/ReplicationSystem/Conditionals/TestDynamicCondition.cpp
+++ b/Engine/Plugins/Runtime/ReplicationSystemTestPlugin/Source/Private/Tests/ReplicationSystem/Conditionals/TestDynamicCondition.cpp
@@ -28,7 +28,7 @@ UE_NET_TEST_FIXTURE(FTestDynamicConditionFixture, DynamicConditionIsReplicatedBy
 	ServerObject->DynamicConditionInt ^= 13;
 
 	// Send and deliver packet
-	Server->PreSendUpdate();
+	Server->NetUpdate();
 	Server->SendAndDeliverTo(Client, DeliverPacket);
 	Server->PostSendUpdate();
 	
@@ -56,7 +56,7 @@ UE_NET_TEST_FIXTURE(FTestDynamicConditionFixture, DynamicConditionSupports_COND_
 	ServerObject->SetDynamicCondition(COND_Never);
 
 	// Send and deliver packet
-	Server->PreSendUpdate();
+	Server->NetUpdate();
 	Server->SendAndDeliverTo(Client, DeliverPacket);
 	Server->PostSendUpdate();
 	
@@ -90,7 +90,7 @@ UE_NET_TEST_FIXTURE(FTestDynamicConditionFixture, DynamicConditionSupports_COND_
 	// Send and deliver packets
 	for (auto Client : SomeClients)
 	{
-		Server->PreSendUpdate();
+		Server->NetUpdate();
 		Server->SendAndDeliverTo(Client, DeliverPacket);
 		Server->PostSendUpdate();
 	}
@@ -136,7 +136,7 @@ UE_NET_TEST_FIXTURE(FTestDynamicConditionFixture, DynamicConditionSupports_COND_
 	// Send and deliver packets
 	for (auto Client : SomeClients)
 	{
-		Server->PreSendUpdate();
+		Server->NetUpdate();
 		Server->SendAndDeliverTo(Client, DeliverPacket);
 		Server->PostSendUpdate();
 	}
@@ -182,7 +182,7 @@ UE_NET_TEST_FIXTURE(FTestDynamicConditionFixture, DynamicConditionSupports_COND_
 	// Send and deliver packets
 	for (auto Client : SomeClients)
 	{
-		Server->PreSendUpdate();
+		Server->NetUpdate();
 		Server->SendAndDeliverTo(Client, DeliverPacket);
 		Server->PostSendUpdate();
 	}
@@ -228,7 +228,7 @@ UE_NET_TEST_FIXTURE(FTestDynamicConditionFixture, DynamicConditionSupports_COND_
 	// Send and deliver packets
 	for (auto Client : SomeClients)
 	{
-		Server->PreSendUpdate();
+		Server->NetUpdate();
 		Server->SendAndDeliverTo(Client, DeliverPacket);
 		Server->PostSendUpdate();
 	}
@@ -266,7 +266,7 @@ UE_NET_TEST_FIXTURE(FTestDynamicConditionFixture, DynamicConditionSupports_COND_
 	ServerObject->SetDynamicCondition(COND_InitialOnly);
 
 	// Send and deliver packet
-	Server->PreSendUpdate();
+	Server->NetUpdate();
 	Server->SendAndDeliverTo(Client, DeliverPacket);
 	Server->PostSendUpdate();
 	
@@ -281,7 +281,7 @@ UE_NET_TEST_FIXTURE(FTestDynamicConditionFixture, DynamicConditionSupports_COND_
 	ServerObject->DynamicConditionInt += 13;
 
 	// Send and deliver packet
-	Server->PreSendUpdate();
+	Server->NetUpdate();
 	Server->SendAndDeliverTo(Client, DeliverPacket);
 	Server->PostSendUpdate();
 
@@ -310,7 +310,7 @@ UE_NET_TEST_FIXTURE(FTestDynamicConditionFixture, DynamicConditionSupports_COND_
 	// Send and deliver packets
 	for (auto Client : SomeClients)
 	{
-		Server->PreSendUpdate();
+		Server->NetUpdate();
 		Server->SendAndDeliverTo(Client, DeliverPacket);
 		Server->PostSendUpdate();
 	}
@@ -340,7 +340,7 @@ UE_NET_TEST_FIXTURE(FTestDynamicConditionFixture, DynamicConditionSupports_COND_
 	// Send and deliver packets
 	for (auto Client : SomeClients)
 	{
-		Server->PreSendUpdate();
+		Server->NetUpdate();
 		Server->SendAndDeliverTo(Client, DeliverPacket);
 		Server->PostSendUpdate();
 	}
@@ -378,7 +378,7 @@ UE_NET_TEST_FIXTURE(FTestDynamicConditionFixture, DynamicConditionSupports_COND_
 	ServerObject->SetDynamicCondition(COND_Custom);
 
 	// Send and deliver packet
-	Server->PreSendUpdate();
+	Server->NetUpdate();
 	Server->SendAndDeliverTo(Client, DeliverPacket);
 	Server->PostSendUpdate();
 	
@@ -396,7 +396,7 @@ UE_NET_TEST_FIXTURE(FTestDynamicConditionFixture, DynamicConditionSupports_COND_
 	ServerObject->DynamicConditionInt += 13;
 
 	// Send and deliver packet
-	Server->PreSendUpdate();
+	Server->NetUpdate();
 	Server->SendAndDeliverTo(Client, DeliverPacket);
 	Server->PostSendUpdate();
 
@@ -424,7 +424,7 @@ UE_NET_TEST_FIXTURE(FTestDynamicConditionFixture, DynamicConditionSupports_COND_
 	ServerObject->SetDynamicCondition(COND_SimulatedOrPhysics);
 
 	// Send and deliver packet
-	Server->PreSendUpdate();
+	Server->NetUpdate();
 	Server->SendAndDeliverTo(Client, DeliverPacket);
 	Server->PostSendUpdate();
 	
@@ -442,7 +442,7 @@ UE_NET_TEST_FIXTURE(FTestDynamicConditionFixture, DynamicConditionSupports_COND_
 	ServerObject->DynamicConditionInt += 13;
 
 	// Send and deliver packet
-	Server->PreSendUpdate();
+	Server->NetUpdate();
 	Server->SendAndDeliverTo(Client, DeliverPacket);
 	Server->PostSendUpdate();
 
@@ -456,7 +456,7 @@ UE_NET_TEST_FIXTURE(FTestDynamicConditionFixture, DynamicConditionSupports_COND_
 	ServerObject->DynamicConditionInt += 17;
 
 	// Send and deliver packet
-	Server->PreSendUpdate();
+	Server->NetUpdate();
 	Server->SendAndDeliverTo(Client, DeliverPacket);
 	Server->PostSendUpdate();
 
@@ -480,7 +480,7 @@ UE_NET_TEST_FIXTURE(FTestDynamicConditionFixture, DynamicConditionSupportsGoingF
 	ServerObject->SetDynamicCondition(COND_Never);
 
 	// Send and deliver packet
-	Server->PreSendUpdate();
+	Server->NetUpdate();
 	Server->SendAndDeliverTo(Client, DeliverPacket);
 	Server->PostSendUpdate();
 	
@@ -495,7 +495,7 @@ UE_NET_TEST_FIXTURE(FTestDynamicConditionFixture, DynamicConditionSupportsGoingF
 	ServerObject->SetDynamicCondition(COND_None);
 
 	// Send and deliver packet
-	Server->PreSendUpdate();
+	Server->NetUpdate();
 	Server->SendAndDeliverTo(Client, DeliverPacket);
 	Server->PostSendUpdate();
 
diff --git a/Engine/Plugins/Runtime/ReplicationSystemTestPlugin/Source/Private/Tests/ReplicationSystem/Filtering/TestFiltering.cpp b/Engine/Plugins/Runtime/ReplicationSystemTestPlugin/Source/Private/Tests/ReplicationSystem/Filtering/TestFiltering.cpp
--- a/Engine/Plugins/Runtime/ReplicationSystemTestPlugin/Source/Private/Tests/ReplicationSystem/Filtering/TestFiltering.cpp
+++ b/Engine/Plugins/Runtime/ReplicationSystemTestPlugin/Source/Private/Tests/ReplicationSystem/Filtering/TestFiltering.cpp
@@ -253,7 +253,7 @@ UE_NET_TEST_FIXTURE(FTestFilteringFixture, OwnerFilterPreventsObjectFromReplicat
 	Server->ReplicationSystem->SetFilter(ServerObject->NetRefHandle, ToOwnerFilterHandle);
 
 	// Send and deliver packet
-	Server->PreSendUpdate();
+	Server->NetUpdate();
 	Server->SendAndDeliverTo(Client, DeliverPacket);
 	Server->PostSendUpdate();
 
@@ -274,7 +274,7 @@ UE_NET_TEST_FIXTURE(FTestFilteringFixture, OwnerFilterAllowsObjectToReplicateToO
 	Server->ReplicationSystem->SetOwningNetConnection(ServerObject->NetRefHandle, Client->ConnectionIdOnServer);
 
 	// Send and deliver packet
-	Server->PreSendUpdate();
+	Server->NetUpdate();
 	Server->SendAndDeliverTo(Client, DeliverPacket);
 	Server->PostSendUpdate();
 
@@ -284,7 +284,7 @@ UE_NET_TEST_FIXTURE(FTestFilteringFixture, OwnerFilterAllowsObjectToReplicateToO
 	// Destroy object on server and client
 	Server->DestroyObject(ServerObject);
 
-	Server->PreSendUpdate();
+	Server->NetUpdate();
 	Server->SendAndDeliverTo(Client, DeliverPacket);
 	Server->PostSendUpdate();
 }
@@ -301,7 +301,7 @@ UE_NET_TEST_FIXTURE(FTestFilteringFixture, OwnerFilterReplicatesOnlyToOwningConn
 	Server->ReplicationSystem->SetOwningNetConnection(ServerObject->NetRefHandle, ClientArray[LastClientIndex]->ConnectionIdOnServer);
 
 	// Send and deliver packets
-	Server->PreSendUpdate();
+	Server->NetUpdate();
 	for (FReplicationSystemTestClient* Client : ClientArray)
 	{
 		Server->SendAndDeliverTo(Client, DeliverPacket);
@@ -323,7 +323,7 @@ UE_NET_TEST_FIXTURE(FTestFilteringFixture, OwnerFilterReplicatesOnlyToOwningConn
 	}
 
 	Server->DestroyObject(ServerObject);
-	Server->PreSendUpdate();
+	Server->NetUpdate();
 	for (FReplicationSystemTestClient* Client : ClientArray)
 	{
 		Server->SendAndDeliverTo(Client, DeliverPacket);
@@ -340,7 +340,7 @@ UE_NET_TEST_FIXTURE(FTestFilteringFixture, CanChangeOwningConnection)
 	UReplicatedTestObject* ServerObject = Server->CreateObject(0, 0);
 
 	// Send and deliver packets
-	Server->PreSendUpdate();
+	Server->NetUpdate();
 	Server->SendAndDeliverTo(Client, DeliverPacket);
 	Server->PostSendUpdate();
 
@@ -351,7 +351,7 @@ UE_NET_TEST_FIXTURE(FTestFilteringFixture, CanChangeOwningConnection)
 	Server->ReplicationSystem->SetFilter(ServerObject->NetRefHandle, ToOwnerFilterHandle);
 
 	// Send and deliver packets
-	Server->PreSendUpdate();
+	Server->NetUpdate();
 	Server->SendAndDeliverTo(Client, DeliverPacket);
 	Server->PostSendUpdate();
 
@@ -362,7 +362,7 @@ UE_NET_TEST_FIXTURE(FTestFilteringFixture, CanChangeOwningConnection)
 	Server->ReplicationSystem->SetOwningNetConnection(ServerObject->NetRefHandle, Client->ConnectionIdOnServer);
 
 	// Send and deliver packets
-	Server->PreSendUpdate();
+	Server->NetUpdate();
 	Server->SendAndDeliverTo(Client, DeliverPacket);
 	Server->PostSendUpdate();
 
@@ -373,7 +373,7 @@ UE_NET_TEST_FIXTURE(FTestFilteringFixture, CanChangeOwningConnection)
 	Server->ReplicationSystem->SetOwningNetConnection(ServerObject->NetRefHandle, 0);
 
 	// Send and deliver packets
-	Server->PreSendUpdate();
+	Server->NetUpdate();
 	Server->SendAndDeliverTo(Client, DeliverPacket);
 	Server->PostSendUpdate();
 
@@ -392,7 +392,7 @@ UE_NET_TEST_FIXTURE(FTestFilteringFixture, CanChangeOwnerFilter)
 	UReplicatedTestObject* ServerObject = Server->CreateObject(0, 0);
 
 	// Send and deliver packets
-	Server->PreSendUpdate();
+	Server->NetUpdate();
 	Server->SendAndDeliverTo(Client, DeliverPacket);
 	Server->PostSendUpdate();
 
@@ -403,7 +403,7 @@ UE_NET_TEST_FIXTURE(FTestFilteringFixture, CanChangeOwnerFilter)
 	Server->ReplicationSystem->SetFilter(ServerObject->NetRefHandle, ToOwnerFilterHandle);
 
 	// Send and deliver packets
-	Server->PreSendUpdate();
+	Server->NetUpdate();
 	Server->SendAndDeliverTo(Client, DeliverPacket);
 	Server->PostSendUpdate();
 
@@ -414,7 +414,7 @@ UE_NET_TEST_FIXTURE(FTestFilteringFixture, CanChangeOwnerFilter)
 	Server->ReplicationSystem->SetFilter(ServerObject->NetRefHandle, InvalidNetObjectFilterHandle);
 
 	// Send and deliver packets
-	Server->PreSendUpdate();
+	Server->NetUpdate();
 	Server->SendAndDeliverTo(Client, DeliverPacket);
 	Server->PostSendUpdate();
 
@@ -425,7 +425,7 @@ UE_NET_TEST_FIXTURE(FTestFilteringFixture, CanChangeOwnerFilter)
 	Server->DestroyObject(ServerObject);
 
 	// Send and deliver packets
-	Server->PreSendUpdate();
+	Server->NetUpdate();
 	Server->SendAndDeliverTo(Client, DeliverPacket);
 	Server->PostSendUpdate();
 }
@@ -449,7 +449,7 @@ UE_NET_TEST_FIXTURE(FTestFilteringFixture, LateAddedSubObjectGetsOwnerPropagated
 	}
 
 	// Net update
-	Server->PreSendUpdate();
+	Server->NetUpdate();
 	Server->PostSendUpdate();
 
 	// Create subobject to arbitrary object
@@ -457,7 +457,7 @@ UE_NET_TEST_FIXTURE(FTestFilteringFixture, LateAddedSubObjectGetsOwnerPropagated
 	UReplicatedTestObject* ServerSubObject = Server->CreateSubObject(ArbitraryServerObject->NetRefHandle, 1, 1);
 
 	// Net update
-	Server->PreSendUpdate();
+	Server->NetUpdate();
 	Server->PostSendUpdate();
 
 	// Verify subobject owner is as expected
@@ -479,7 +479,7 @@ UE_NET_TEST_FIXTURE(FTestFilteringFixture, ConnectionFilterPreventsObjectFromRep
 	UE_NET_ASSERT_TRUE(bFilterWasApplied);
 
 	// Send and deliver packet
-	Server->PreSendUpdate();
+	Server->NetUpdate();
 	Server->SendAndDeliverTo(Client, DeliverPacket);
 	Server->PostSendUpdate();
 
@@ -502,7 +502,7 @@ UE_NET_TEST_FIXTURE(FTestFilteringFixture, ConnectionFilterAllowsObjectToReplica
 	Server->ReplicationSystem->SetConnectionFilter(ServerObject->NetRefHandle, AllowedConnections, ENetFilterStatus::Allow);
 
 	// Send and deliver packet
-	Server->PreSendUpdate();
+	Server->NetUpdate();
 	Server->SendAndDeliverTo(Client, DeliverPacket);
 	Server->PostSendUpdate();
 
@@ -521,14 +521,14 @@ UE_NET_TEST_FIXTURE(FTestFilteringFixture, ConnectionFilterAllowsObjectToReplica
 	Server->ReplicationSystem->SetConnectionFilter(ServerObject->NetRefHandle, NoConnections, ReplicationStatus);
 
 	// Send and deliver packets
-	Server->PreSendUpdate();
+	Server->NetUpdate();
 	Server->PostSendUpdate();
 
 	// Late join clients
 	FReplicationSystemTestClient* ClientArray[] = {CreateClient(), CreateClient(), CreateClient()};
 
 	// Send and deliver packets
-	Server->PreSendUpdate();
+	Server->NetUpdate();
 	for (FReplicationSystemTestClient* Client : ClientArray)
 	{
 		Server->SendAndDeliverTo(Client, DeliverPacket);
@@ -565,11 +565,11 @@ UE_NET_TEST_FIXTURE(FTestFilteringFixture, ConnectionFilterAllowsObjectToReplica
 	Server->ReplicationSystem->SetConnectionFilter(ServerObject->NetRefHandle, NoConnections, ReplicationStatus);
 
 	// Send and deliver packets
-	Server->PreSendUpdate();
+	Server->NetUpdate();
 	Server->PostSendUpdate();
 
 	// Send and deliver packets
-	Server->PreSendUpdate();
+	Server->NetUpdate();
 	Server->SendAndDeliverTo(Client, DeliverPacket);
 	Server->PostSendUpdate();
 
@@ -594,7 +594,7 @@ UE_NET_TEST_FIXTURE(FTestFilteringFixture, ConnectionFilterPreventsObjectFromRep
 	Server->ReplicationSystem->SetConnectionFilter(ServerObject->NetRefHandle, Connection1, ReplicationStatus);
 
 	// Send and deliver packets
-	Server->PreSendUpdate();
+	Server->NetUpdate();
 	Server->PostSendUpdate();
 
 	// Late join clients
@@ -665,7 +665,7 @@ UE_NET_TEST_FIXTURE(FTestFilteringFixture, CanChangeConnectionFilter)
 	UReplicatedTestObject* ServerObject = Server->CreateObject(0, 0);
 
 	// Send and deliver packets
-	Server->PreSendUpdate();
+	Server->NetUpdate();
 	Server->SendAndDeliverTo(Client, DeliverPacket);
 	Server->PostSendUpdate();
 
@@ -677,7 +677,7 @@ UE_NET_TEST_FIXTURE(FTestFilteringFixture, CanChangeConnectionFilter)
 	Server->ReplicationSystem->SetConnectionFilter(ServerObject->NetRefHandle, NoConnections, ENetFilterStatus::Allow);
 
 	// Send and deliver packets
-	Server->PreSendUpdate();
+	Server->NetUpdate();
 	Server->SendAndDeliverTo(Client, DeliverPacket);
 	Server->PostSendUpdate();
 
@@ -691,7 +691,7 @@ UE_NET_TEST_FIXTURE(FTestFilteringFixture, CanChangeConnectionFilter)
 	Server->ReplicationSystem->SetConnectionFilter(ServerObject->NetRefHandle, AllowedConnections, ENetFilterStatus::Allow);
 
 	// Send and deliver packets
-	Server->PreSendUpdate();
+	Server->NetUpdate();
 	Server->SendAndDeliverTo(Client, DeliverPacket);
 	Server->PostSendUpdate();
 
@@ -702,7 +702,7 @@ UE_NET_TEST_FIXTURE(FTestFilteringFixture, CanChangeConnectionFilter)
 	Server->ReplicationSystem->SetConnectionFilter(ServerObject->NetRefHandle, NoConnections, ENetFilterStatus::Allow);
 
 	// Send and deliver packets
-	Server->PreSendUpdate();
+	Server->NetUpdate();
 	Server->SendAndDeliverTo(Client, DeliverPacket);
 	Server->PostSendUpdate();
 
@@ -719,7 +719,7 @@ UE_NET_TEST_FIXTURE(FTestFilteringFixture, CanToggleConnectionFilter)
 	UReplicatedTestObject* ServerObject = Server->CreateObject(0, 0);
 
 	// Send and deliver packets
-	Server->PreSendUpdate();
+	Server->NetUpdate();
 	Server->SendAndDeliverTo(Client, DeliverPacket);
 	Server->PostSendUpdate();
 
@@ -731,7 +731,7 @@ UE_NET_TEST_FIXTURE(FTestFilteringFixture, CanToggleConnectionFilter)
 	Server->ReplicationSystem->SetConnectionFilter(ServerObject->NetRefHandle, NoConnections, ENetFilterStatus::Allow);
 
 	// Send and deliver packets
-	Server->PreSendUpdate();
+	Server->NetUpdate();
 	Server->SendAndDeliverTo(Client, DeliverPacket);
 	Server->PostSendUpdate();
 
@@ -742,7 +742,7 @@ UE_NET_TEST_FIXTURE(FTestFilteringFixture, CanToggleConnectionFilter)
 	Server->ReplicationSystem->SetFilter(ServerObject->NetRefHandle, InvalidNetObjectFilterHandle);
 
 	// Send and deliver packets
-	Server->PreSendUpdate();
+	Server->NetUpdate();
 	Server->SendAndDeliverTo(Client, DeliverPacket);
 	Server->PostSendUpdate();
 
@@ -753,7 +753,7 @@ UE_NET_TEST_FIXTURE(FTestFilteringFixture, CanToggleConnectionFilter)
 	Server->DestroyObject(ServerObject);
 
 	// Send and deliver packets
-	Server->PreSendUpdate();
+	Server->NetUpdate();
 	Server->SendAndDeliverTo(Client, DeliverPacket);
 	Server->PostSendUpdate();
 }
@@ -774,7 +774,7 @@ UE_NET_TEST_FIXTURE(FTestFilteringFixture, CanToggleBetweenAllFilters)
 		Server->ReplicationSystem->SetFilter(ServerObject->NetRefHandle, ToOwnerFilterHandle);
 
 		// Send and deliver packets
-		Server->PreSendUpdate();
+		Server->NetUpdate();
 		Server->SendAndDeliverTo(Client, DeliverPacket);
 		Server->PostSendUpdate();
 
@@ -786,7 +786,7 @@ UE_NET_TEST_FIXTURE(FTestFilteringFixture, CanToggleBetweenAllFilters)
 		Server->ReplicationSystem->SetConnectionFilter(ServerObject->NetRefHandle, NoConnections, ENetFilterStatus::Allow);
 
 		// Send and deliver packets
-		Server->PreSendUpdate();
+		Server->NetUpdate();
 		Server->SendAndDeliverTo(Client, DeliverPacket);
 		Server->PostSendUpdate();
 
@@ -797,7 +797,7 @@ UE_NET_TEST_FIXTURE(FTestFilteringFixture, CanToggleBetweenAllFilters)
 		Server->ReplicationSystem->SetFilter(ServerObject->NetRefHandle, ToOwnerFilterHandle);
 
 		// Send and deliver packets
-		Server->PreSendUpdate();
+		Server->NetUpdate();
 		Server->SendAndDeliverTo(Client, DeliverPacket);
 		Server->PostSendUpdate();
 
@@ -823,7 +823,7 @@ UE_NET_TEST_FIXTURE(FTestFilteringFixture, GroupFilterPreventsObjectFromReplicat
 	Server->ReplicationSystem->SetGroupFilterStatus(GroupHandle, Client->ConnectionIdOnServer, ENetFilterStatus::Disallow);
 
 	// Send and deliver packet
-	Server->PreSendUpdate();
+	Server->NetUpdate();
 	Server->SendAndDeliverTo(Client, DeliverPacket);
 	Server->PostSendUpdate();
 
@@ -847,7 +847,7 @@ UE_NET_TEST_FIXTURE(FTestFilteringFixture, GroupFilterAllowsObjectToReplicate)
 	Server->ReplicationSystem->SetGroupFilterStatus(GroupHandle, Client->ConnectionIdOnServer, ENetFilterStatus::Disallow);
 
 	// Send and deliver packet
-	Server->PreSendUpdate();
+	Server->NetUpdate();
 	Server->SendAndDeliverTo(Client, DeliverPacket);
 	Server->PostSendUpdate();
 
@@ -858,7 +858,7 @@ UE_NET_TEST_FIXTURE(FTestFilteringFixture, GroupFilterAllowsObjectToReplicate)
 	Server->ReplicationSystem->SetGroupFilterStatus(GroupHandle, Client->ConnectionIdOnServer, ENetFilterStatus::Allow);
 
 	// Send and deliver packet
-	Server->PreSendUpdate();
+	Server->NetUpdate();
 	Server->SendAndDeliverTo(Client, DeliverPacket);
 	Server->PostSendUpdate();
 
@@ -960,7 +960,7 @@ UE_NET_TEST_FIXTURE(FTestFilteringFixture, DynamicFilterAddObjectAndRemoveObject
 	UReplicatedTestObject* ServerObject = Server->CreateObject(0, 0);
 	Server->ReplicationSystem->SetFilter(ServerObject->NetRefHandle, MockFilterHandle);
 
-	Server->PreSendUpdate();
+	Server->NetUpdate();
 	// Filter needs to be set now.
 	{
 		const UMockNetObjectFilter::FFunctionCallStatus& CallStatus = MockNetObjectFilter->GetFunctionCallStatus();
@@ -975,7 +975,7 @@ UE_NET_TEST_FIXTURE(FTestFilteringFixture, DynamicFilterAddObjectAndRemoveObject
 
 	Server->DestroyObject(ServerObject);
 
-	Server->PreSendUpdate();
+	Server->NetUpdate();
 	// Filter needs to be cleared now.
 	{
 		const UMockNetObjectFilter::FFunctionCallStatus& CallStatus = MockNetObjectFilter->GetFunctionCallStatus();
@@ -1001,7 +1001,7 @@ UE_NET_TEST_FIXTURE(FTestFilteringFixture, DynamicFilterCanAllowObjectToReplicat
 	Server->ReplicationSystem->SetFilter(ServerObject->NetRefHandle, MockFilterHandle);
 
 	// Send and deliver packets
-	Server->PreSendUpdate();
+	Server->NetUpdate();
 	Server->SendAndDeliverTo(Client, DeliverPacket);
 	Server->PostSendUpdate();
 
@@ -1022,7 +1022,7 @@ UE_NET_TEST_FIXTURE(FTestFilteringFixture, DynamicFilterCanDisallowObjectToRepli
 	Server->ReplicationSystem->SetFilter(ServerObject->NetRefHandle, MockFilterHandle);
 
 	// Send and deliver packets
-	Server->PreSendUpdate();
+	Server->NetUpdate();
 	Server->SendAndDeliverTo(Client, DeliverPacket);
 	Server->PostSendUpdate();
 
@@ -1040,7 +1040,7 @@ UE_NET_TEST_FIXTURE(FTestFilteringFixture, SwitchingFiltersCallsRemoveObjectOnPr
 	Server->ReplicationSystem->SetFilter(ServerObject->NetRefHandle, MockFilterHandle);
 
 	// Make sure filter is set
-	Server->PreSendUpdate();
+	Server->NetUpdate();
 	Server->PostSendUpdate();
 
 	// Check RemoveObject is called when switching filters.
@@ -1050,7 +1050,7 @@ UE_NET_TEST_FIXTURE(FTestFilteringFixture, SwitchingFiltersCallsRemoveObjectOnPr
 		const TBitArray<> NoConnections;
 		Server->ReplicationSystem->SetConnectionFilter(ServerObject->NetRefHandle, NoConnections, ENetFilterStatus::Disallow);
 
-		Server->PreSendUpdate();
+		Server->NetUpdate();
 		Server->PostSendUpdate();
 
 		const UMockNetObjectFilter::FFunctionCallStatus& FunctionCallStatus = MockNetObjectFilter->GetFunctionCallStatus();
@@ -1074,7 +1074,7 @@ UE_NET_TEST_FIXTURE(FTestFilteringFixture, SubObjectsAreReplicatedWhenOwnerDynam
 	UReplicatedTestObject* ServerSubObject = Server->CreateSubObject(ServerObject->NetRefHandle, 0, 0);
 
 	// Send and deliver packets
-	Server->PreSendUpdate();
+	Server->NetUpdate();
 	Server->SendAndDeliverTo(Client, DeliverPacket);
 	Server->PostSendUpdate();
 
@@ -1098,7 +1098,7 @@ UE_NET_TEST_FIXTURE(FTestFilteringFixture, SubObjectsAreNotReplicatedWhenOwnerDy
 	UReplicatedTestObject* ServerSubObject = Server->CreateSubObject(ServerObject->NetRefHandle, 0, 0);
 
 	// Send and deliver packets
-	Server->PreSendUpdate();
+	Server->NetUpdate();
 	Server->SendAndDeliverTo(Client, DeliverPacket);
 	Server->PostSendUpdate();
 
@@ -1121,7 +1121,7 @@ UE_NET_TEST_FIXTURE(FTestFilteringFixture, DependentObjectIsUnaffectedByDynamicF
 	Server->ReplicationSystem->SetFilter(ServerFutureDependentObject->NetRefHandle, MockFilterHandle);
 
 	// Send and deliver packets
-	Server->PreSendUpdate();
+	Server->NetUpdate();
 	Server->SendAndDeliverTo(Client, DeliverPacket);
 	Server->PostSendUpdate();
 
@@ -1134,7 +1134,7 @@ UE_NET_TEST_FIXTURE(FTestFilteringFixture, DependentObjectIsUnaffectedByDynamicF
 	UReplicatedTestObject* ServerDependentObject = ServerFutureDependentObject;
 
 	// Send and deliver packets
-	Server->PreSendUpdate();
+	Server->NetUpdate();
 	Server->SendAndDeliverTo(Client, DeliverPacket);
 	Server->PostSendUpdate();
 
@@ -1146,7 +1146,7 @@ UE_NET_TEST_FIXTURE(FTestFilteringFixture, DependentObjectIsUnaffectedByDynamicF
 	UReplicatedTestObject* ServerFormerDependentObject = ServerDependentObject;
 
 	// Send and deliver packets
-	Server->PreSendUpdate();
+	Server->NetUpdate();
 	Server->SendAndDeliverTo(Client, DeliverPacket);
 	Server->PostSendUpdate();
 
@@ -1171,7 +1171,7 @@ UE_NET_TEST_FIXTURE(FTestFilteringFixture, NestedDependentObjectIsFilteredAsPare
 	Server->ReplicationSystem->SetFilter(ServerFutureNestedDependentObject->NetRefHandle, MockFilterHandle);
 
 	// Send and deliver packets
-	Server->PreSendUpdate();
+	Server->NetUpdate();
 	Server->SendAndDeliverTo(Client, DeliverPacket);
 	Server->PostSendUpdate();
 
@@ -1188,7 +1188,7 @@ UE_NET_TEST_FIXTURE(FTestFilteringFixture, NestedDependentObjectIsFilteredAsPare
 	UReplicatedTestObject* ServerNestedDependentObject = ServerFutureNestedDependentObject;
 
 	// Send and deliver packets
-	Server->PreSendUpdate();
+	Server->NetUpdate();
 	Server->SendAndDeliverTo(Client, DeliverPacket);
 	Server->PostSendUpdate();
 
@@ -1201,7 +1201,7 @@ UE_NET_TEST_FIXTURE(FTestFilteringFixture, NestedDependentObjectIsFilteredAsPare
 	UReplicatedTestObject* ServerFormerDependentObject = ServerDependentObject;
 
 	// Send and deliver packets
-	Server->PreSendUpdate();
+	Server->NetUpdate();
 	Server->SendAndDeliverTo(Client, DeliverPacket);
 	Server->PostSendUpdate();
 
@@ -1215,7 +1215,7 @@ UE_NET_TEST_FIXTURE(FTestFilteringFixture, NestedDependentObjectIsFilteredAsPare
 	Server->ReplicationSystem->SetFilter(ServerNestedDependentObject->NetRefHandle, InvalidNetObjectFilterHandle);
 
 	// Send and deliver packets
-	Server->PreSendUpdate();
+	Server->NetUpdate();
 	Server->SendAndDeliverTo(Client, DeliverPacket);
 	Server->PostSendUpdate();
 
@@ -1243,7 +1243,7 @@ UE_NET_TEST_FIXTURE(FTestFilteringFixture, DynamicFilteredOutSubObjectsAreResetW
 	}
 
 	// Send and deliver packets
-	Server->PreSendUpdate();
+	Server->NetUpdate();
 	Server->SendAndDeliverTo(Client, DeliverPacket);
 	Server->PostSendUpdate();
 
@@ -1254,7 +1254,7 @@ UE_NET_TEST_FIXTURE(FTestFilteringFixture, DynamicFilteredOutSubObjectsAreResetW
 	UReplicatedTestObject* ServerObject2 = Server->CreateObject(0, 0);
 
 	// Send and deliver packets
-	Server->PreSendUpdate();
+	Server->NetUpdate();
 	Server->SendAndDeliverTo(Client, DeliverPacket);
 	Server->PostSendUpdate();
 
@@ -1518,7 +1518,7 @@ UE_NET_TEST_FIXTURE(FTestFilteringFixture, InclusionGroupDoesNotOverrideExclusio
 	Server->ReplicationSystem->SetGroupFilterStatus(InclusionGroupHandle, ENetFilterStatus::Allow);
 
 	// Send and deliver packet
-	Server->PreSendUpdate();
+	Server->NetUpdate();
 	for (FReplicationSystemTestClient* Client : ClientArray)
 	{
 		Server->SendAndDeliverTo(Client, DeliverPacket);
diff --git a/Engine/Plugins/Runtime/ReplicationSystemTestPlugin/Source/Private/Tests/ReplicationSystem/Filtering/TestGridFilter.cpp b/Engine/Plugins/Runtime/ReplicationSystemTestPlugin/Source/Private/Tests/ReplicationSystem/Filtering/TestGridFilter.cpp
--- a/Engine/Plugins/Runtime/ReplicationSystemTestPlugin/Source/Private/Tests/ReplicationSystem/Filtering/TestGridFilter.cpp
+++ b/Engine/Plugins/Runtime/ReplicationSystemTestPlugin/Source/Private/Tests/ReplicationSystem/Filtering/TestGridFilter.cpp
@@ -192,7 +192,7 @@ UE_NET_TEST_FIXTURE(FTestGridFilterFixture, TestWorldLocGridFilter)
 	Server->ReplicationSystem->SetReplicationView(Client->ConnectionIdOnServer, ReplicationView);
 
 	// Send and deliver packet
-	Server->PreSendUpdate();
+	Server->NetUpdate();
 	Server->SendAndDeliverTo(Client, DeliverPacket);
 	Server->PostSendUpdate();
 
@@ -248,7 +248,7 @@ UE_NET_TEST_FIXTURE(FTestGridFilterFixture, TestWorldLocGridFilterWithHugeCullDi
 	Server->ReplicationSystem->SetReplicationView(Client->ConnectionIdOnServer, ReplicationView);
 
 	// Send and deliver packet
-	Server->PreSendUpdate();
+	Server->NetUpdate();
 	Server->SendAndDeliverTo(Client, DeliverPacket);
 	Server->PostSendUpdate();
 
diff --git a/Engine/Plugins/Runtime/ReplicationSystemTestPlugin/Source/Private/Tests/ReplicationSystem/Filtering/TestGroups.cpp b/Engine/Plugins/Runtime/ReplicationSystemTestPlugin/Source/Private/Tests/ReplicationSystem/Filtering/TestGroups.cpp
--- a/Engine/Plugins/Runtime/ReplicationSystemTestPlugin/Source/Private/Tests/ReplicationSystem/Filtering/TestGroups.cpp
+++ b/Engine/Plugins/Runtime/ReplicationSystemTestPlugin/Source/Private/Tests/ReplicationSystem/Filtering/TestGroups.cpp
@@ -36,7 +36,7 @@ protected:
 
 	void Filter()
 	{
-		Server->PreSendUpdate();
+		Server->NetUpdate();
 		Server->SendAndDeliverTo(Client, DeliverPacket);
 		Server->PostSendUpdate();
 	}
diff --git a/Engine/Plugins/Runtime/ReplicationSystemTestPlugin/Source/Private/Tests/ReplicationSystem/MultiReplicationSystemsTestFixture.cpp b/Engine/Plugins/Runtime/ReplicationSystemTestPlugin/Source/Private/Tests/ReplicationSystem/MultiReplicationSystemsTestFixture.cpp
--- a/Engine/Plugins/Runtime/ReplicationSystemTestPlugin/Source/Private/Tests/ReplicationSystem/MultiReplicationSystemsTestFixture.cpp
+++ b/Engine/Plugins/Runtime/ReplicationSystemTestPlugin/Source/Private/Tests/ReplicationSystem/MultiReplicationSystemsTestFixture.cpp
@@ -158,7 +158,7 @@ void FMultiReplicationSystemsTestFixture::FullSendAndDeliverUpdate()
 {
 	for (FReplicationSystemTestServer* Server : Servers)
 	{
-		Server->PreSendUpdate();
+		Server->NetUpdate();
 	}
 
 	for (TMap<FReplicationSystemTestServer*, TArray<FReplicationSystemTestClient*>>::ElementType& ServerWithClients : ServerClients)
@@ -184,7 +184,7 @@ void FMultiReplicationSystemsTestFixture::FullSendAndDeliverUpdateTwoPass()
 		TArray<FReplicationSystemTestClient*>* FoundClients = ServerClients.Find(Server);
 		if (FoundClients && FoundClients->Num() > 0)
 		{
-			Server->PreSendUpdate();
+			Server->NetUpdate();
 
 			for (FReplicationSystemTestClient* Client : *FoundClients)
 			{
diff --git a/Engine/Plugins/Runtime/ReplicationSystemTestPlugin/Source/Private/Tests/ReplicationSystem/NetBlob/TestPartialNetBlob.cpp b/Engine/Plugins/Runtime/ReplicationSystemTestPlugin/Source/Private/Tests/ReplicationSystem/NetBlob/TestPartialNetBlob.cpp
--- a/Engine/Plugins/Runtime/ReplicationSystemTestPlugin/Source/Private/Tests/ReplicationSystem/NetBlob/TestPartialNetBlob.cpp
+++ b/Engine/Plugins/Runtime/ReplicationSystemTestPlugin/Source/Private/Tests/ReplicationSystem/NetBlob/TestPartialNetBlob.cpp
@@ -120,7 +120,7 @@ UE_NET_TEST_FIXTURE(FPartialNetBlobTestFixture, TestSplitBlobIsReceivedDespitePa
 
 	UReplicatedTestObject* ServerObject = Server->CreateObject(0, 0);
 
-	Server->PreSendUpdate();
+	Server->NetUpdate();
 	Server->SendAndDeliverTo(Client, DeliverPacket);
 	Server->PostSendUpdate();
 
@@ -135,7 +135,7 @@ UE_NET_TEST_FIXTURE(FPartialNetBlobTestFixture, TestSplitBlobIsReceivedDespitePa
 	}
 
 	// Deliver a packet
-	Server->PreSendUpdate();
+	Server->NetUpdate();
 	Server->SendAndDeliverTo(Client, DeliverPacket);
 	Server->PostSendUpdate();
 
@@ -144,7 +144,7 @@ UE_NET_TEST_FIXTURE(FPartialNetBlobTestFixture, TestSplitBlobIsReceivedDespitePa
 	// Lose a few packets
 	for (auto It : {0, 1, 2, 3})
 	{
-		Server->PreSendUpdate();
+		Server->NetUpdate();
 		Server->SendAndDeliverTo(Client, DoNotDeliverPacket);
 		Server->PostSendUpdate();
 	}
@@ -152,7 +152,7 @@ UE_NET_TEST_FIXTURE(FPartialNetBlobTestFixture, TestSplitBlobIsReceivedDespitePa
 	UE_NET_ASSERT_EQ(ClientMockNetObjectAttachmentHandler->GetFunctionCallCounts().OnNetBlobReceived, 0U);
 
 	// Deliver a packet. At this point we expect the entire original blob to have been received properly.
-	Server->PreSendUpdate();
+	Server->NetUpdate();
 	Server->SendAndDeliverTo(Client, DeliverPacket);
 	Server->PostSendUpdate();
 
diff --git a/Engine/Plugins/Runtime/ReplicationSystemTestPlugin/Source/Private/Tests/ReplicationSystem/NetBlob/TestReliableNetBlobQueue.cpp b/Engine/Plugins/Runtime/ReplicationSystemTestPlugin/Source/Private/Tests/ReplicationSystem/NetBlob/TestReliableNetBlobQueue.cpp
--- a/Engine/Plugins/Runtime/ReplicationSystemTestPlugin/Source/Private/Tests/ReplicationSystem/NetBlob/TestReliableNetBlobQueue.cpp
+++ b/Engine/Plugins/Runtime/ReplicationSystemTestPlugin/Source/Private/Tests/ReplicationSystem/NetBlob/TestReliableNetBlobQueue.cpp
@@ -21,7 +21,7 @@ UE_NET_TEST_FIXTURE(FTestReliableNetBlobQueue, CanSendMoreReliableBlobsThanRelia
 
 	UReplicatedTestObject* ServerObject = Server->CreateObject(0, 0);
 
-	Server->PreSendUpdate();
+	Server->NetUpdate();
 	Server->SendAndDeliverTo(Client, DeliverPacket);
 	Server->PostSendUpdate();
 
@@ -38,7 +38,7 @@ UE_NET_TEST_FIXTURE(FTestReliableNetBlobQueue, CanSendMoreReliableBlobsThanRelia
 	// Deliver the entirety of the attachment.
 	for (SIZE_T TryIt = 0, TryEndIt = BlobPartCount; TryIt != TryEndIt; ++TryIt)
 	{
-		Server->PreSendUpdate();
+		Server->NetUpdate();
 		Server->SendAndDeliverTo(Client, DeliverPacket);
 		Server->PostSendUpdate();
 
@@ -58,7 +58,7 @@ UE_NET_TEST_FIXTURE(FTestReliableNetBlobQueue, CanSendMoreReliableBlobsThanRelia
 
 	UReplicatedTestObject* ServerObject = Server->CreateObject(0, 0);
 
-	Server->PreSendUpdate();
+	Server->NetUpdate();
 	Server->SendAndDeliverTo(Client, DeliverPacket);
 	Server->PostSendUpdate();
 
@@ -83,7 +83,7 @@ UE_NET_TEST_FIXTURE(FTestReliableNetBlobQueue, CanSendMoreReliableBlobsThanRelia
 		SIZE_T PacketCount = ConnectionInfo.WrittenPackets.Count();
 		for (SIZE_T TryIt = 0, TryEndIt = BlobPartCount; TryIt != TryEndIt; ++TryIt)
 		{
-			Server->PreSendUpdate();
+			Server->NetUpdate();
 			Server->SendTo(Client);
 			Server->PostSendUpdate();
 
@@ -105,7 +105,7 @@ UE_NET_TEST_FIXTURE(FTestReliableNetBlobQueue, CanSendMoreReliableBlobsThanRelia
 		// Send the rest of the blobs.
 		for (SIZE_T TryIt = 0, TryEndIt = BlobPartCount; TryIt != TryEndIt; ++TryIt)
 		{
-			Server->PreSendUpdate();
+			Server->NetUpdate();
 			Server->SendAndDeliverTo(Client, DeliverPacket);
 			Server->PostSendUpdate();
 
diff --git a/Engine/Plugins/Runtime/ReplicationSystemTestPlugin/Source/Private/Tests/ReplicationSystem/Prioritization/TestNetObjectPrioritizer.cpp b/Engine/Plugins/Runtime/ReplicationSystemTestPlugin/Source/Private/Tests/ReplicationSystem/Prioritization/TestNetObjectPrioritizer.cpp
--- a/Engine/Plugins/Runtime/ReplicationSystemTestPlugin/Source/Private/Tests/ReplicationSystem/Prioritization/TestNetObjectPrioritizer.cpp
+++ b/Engine/Plugins/Runtime/ReplicationSystemTestPlugin/Source/Private/Tests/ReplicationSystem/Prioritization/TestNetObjectPrioritizer.cpp
@@ -117,7 +117,7 @@ UE_NET_TEST_FIXTURE(FTestNetObjectPrioritizerAPIFixture, PrioritizerAddObjectSuc
 		ReplicationSystem->SetPrioritizer(NetRefHandle, MockPrioritizerHandle);
 
 		// We don't know if prioritizer changes are batched, but we assume everything is setup correctly in PreSendUpdate at least.
-		ReplicationSystem->PreSendUpdate(FakeDeltaTime);
+		ReplicationSystem->NetUpdate(FakeDeltaTime);
 
 		const UMockNetObjectPrioritizer::FFunctionCallStatus& FunctionCallStatus = MockNetObjectPrioritizer->GetFunctionCallStatus();
 		UE_NET_ASSERT_EQ(FunctionCallStatus.CallCounts.AddObject, 1U);
@@ -132,7 +132,7 @@ UE_NET_TEST_FIXTURE(FTestNetObjectPrioritizerAPIFixture, PrioritizerAddObjectSuc
 		MockNetObjectPrioritizer->ResetFunctionCallStatus();
 
 		ReplicationBridge->EndReplication(TestObject);
-		ReplicationSystem->PreSendUpdate(FakeDeltaTime);
+		ReplicationSystem->NetUpdate(FakeDeltaTime);
 
 		const UMockNetObjectPrioritizer::FFunctionCallStatus& FunctionCallStatus = MockNetObjectPrioritizer->GetFunctionCallStatus();
 		UE_NET_ASSERT_EQ(FunctionCallStatus.CallCounts.RemoveObject, 1U);
@@ -163,7 +163,7 @@ UE_NET_TEST_FIXTURE(FTestNetObjectPrioritizerAPIFixture, PrioritizerAddObjectFai
 		ReplicationSystem->SetPrioritizer(NetRefHandle, MockPrioritizerHandle);
 
 		// We don't know if prioritizer changes are batched, but we assume everything is setup correctly in PreSendUpdate at least.
-		ReplicationSystem->PreSendUpdate(FakeDeltaTime);
+		ReplicationSystem->NetUpdate(FakeDeltaTime);
 
 		const UMockNetObjectPrioritizer::FFunctionCallStatus& FunctionCallStatus = MockNetObjectPrioritizer->GetFunctionCallStatus();
 		UE_NET_ASSERT_EQ(FunctionCallStatus.CallCounts.AddObject, 1U);
@@ -178,7 +178,7 @@ UE_NET_TEST_FIXTURE(FTestNetObjectPrioritizerAPIFixture, PrioritizerAddObjectFai
 		MockNetObjectPrioritizer->ResetFunctionCallStatus();
 
 		ReplicationBridge->EndReplication(TestObject);
-		ReplicationSystem->PreSendUpdate(FakeDeltaTime);
+		ReplicationSystem->NetUpdate(FakeDeltaTime);
 
 		const UMockNetObjectPrioritizer::FFunctionCallStatus& FunctionCallStatus = MockNetObjectPrioritizer->GetFunctionCallStatus();
 		UE_NET_ASSERT_EQ(FunctionCallStatus.CallCounts.RemoveObject, 0U);
@@ -207,7 +207,7 @@ UE_NET_TEST_FIXTURE(FTestNetObjectPrioritizerAPIFixture, SwitchingPrioritizersCa
 		ReplicationSystem->SetPrioritizer(NetRefHandle, MockPrioritizerHandle);
 
 		// We don't know if prioritizer changes are batched, but we assume everything is setup correctly in PreSendUpdate at least.
-		ReplicationSystem->PreSendUpdate(FakeDeltaTime);
+		ReplicationSystem->NetUpdate(FakeDeltaTime);
 
 		const UMockNetObjectPrioritizer::FFunctionCallStatus& FunctionCallStatus = MockNetObjectPrioritizer->GetFunctionCallStatus();
 		UE_NET_ASSERT_EQ(FunctionCallStatus.CallCounts.AddObject, 1U);
@@ -222,7 +222,7 @@ UE_NET_TEST_FIXTURE(FTestNetObjectPrioritizerAPIFixture, SwitchingPrioritizersCa
 
 		ReplicationSystem->SetStaticPriority(TestObject->NetRefHandle, 1.0f);
 
-		ReplicationSystem->PreSendUpdate(FakeDeltaTime);
+		ReplicationSystem->NetUpdate(FakeDeltaTime);
 
 		const UMockNetObjectPrioritizer::FFunctionCallStatus& FunctionCallStatus = MockNetObjectPrioritizer->GetFunctionCallStatus();
 		UE_NET_ASSERT_EQ(FunctionCallStatus.CallCounts.RemoveObject, 1U);
@@ -243,7 +243,7 @@ UE_NET_TEST_FIXTURE(FTestNetObjectPrioritizerAPIFixture, PrioritizeIsNotCalledWh
 		const FNetRefHandle NetRefHandle = ReplicationBridge->BeginReplication(TestObject);
 		ReplicationSystem->SetStaticPriority(NetRefHandle, 1.0f);
 
-		ReplicationSystem->PreSendUpdate(FakeDeltaTime);
+		ReplicationSystem->NetUpdate(FakeDeltaTime);
 
 		const UMockNetObjectPrioritizer::FFunctionCallStatus& FunctionCallStatus = MockNetObjectPrioritizer->GetFunctionCallStatus();
 		UE_NET_ASSERT_EQ(FunctionCallStatus.CallCounts.AddObject, 0U);
@@ -271,7 +271,7 @@ UE_NET_TEST_FIXTURE(FTestNetObjectPrioritizerAPIFixture, PrioritizeIsNotCalledWh
 		const FNetRefHandle NetRefHandle = ReplicationBridge->BeginReplication(TestObject);
 		ReplicationSystem->SetPrioritizer(NetRefHandle, MockPrioritizerHandle);
 
-		ReplicationSystem->PreSendUpdate(FakeDeltaTime);
+		ReplicationSystem->NetUpdate(FakeDeltaTime);
 
 		const UMockNetObjectPrioritizer::FFunctionCallStatus& FunctionCallStatus = MockNetObjectPrioritizer->GetFunctionCallStatus();
 		UE_NET_ASSERT_EQ(FunctionCallStatus.SuccessfulCallCounts.AddObject, 1U);
@@ -315,7 +315,7 @@ UE_NET_TEST_FIXTURE(FTestNetObjectPrioritizerAPIFixture, PrioritizeIsCalledWhenT
 		const FNetRefHandle NetRefHandle2 = ReplicationBridge->BeginReplication(TestObject2);
 		ReplicationSystem->SetPrioritizer(NetRefHandle2, MockPrioritizerHandle);
 
-		ReplicationSystem->PreSendUpdate(FakeDeltaTime);
+		ReplicationSystem->NetUpdate(FakeDeltaTime);
 
 		const UMockNetObjectPrioritizer::FFunctionCallStatus& FunctionCallStatus = MockNetObjectPrioritizer->GetFunctionCallStatus();
 		UE_NET_ASSERT_EQ(FunctionCallStatus.CallCounts.Prioritize, 2U);
@@ -348,7 +348,7 @@ UE_NET_TEST_FIXTURE(FTestNetObjectPrioritizerAPIFixture, UpdateObjectsIsCalledFo
 		const FNetRefHandle NetRefHandle1 = ReplicationBridge->BeginReplication(TestObject1);
 		ReplicationSystem->SetPrioritizer(NetRefHandle1, MockPrioritizerHandle);
 
-		ReplicationSystem->PreSendUpdate(FakeDeltaTime);
+		ReplicationSystem->NetUpdate(FakeDeltaTime);
 		ReplicationSystem->PostSendUpdate();
 	}
 
@@ -358,7 +358,7 @@ UE_NET_TEST_FIXTURE(FTestNetObjectPrioritizerAPIFixture, UpdateObjectsIsCalledFo
 
 		TestObject1->IntA ^= 1;
 
-		ReplicationSystem->PreSendUpdate(FakeDeltaTime);
+		ReplicationSystem->NetUpdate(FakeDeltaTime);
 
 		const UMockNetObjectPrioritizer::FFunctionCallStatus& FunctionCallStatus = MockNetObjectPrioritizer->GetFunctionCallStatus();
 
@@ -386,7 +386,7 @@ UE_NET_TEST_FIXTURE(FTestNetObjectPrioritizerAPIFixture, UpdateObjectsIsNotCalle
 		const FNetRefHandle NetRefHandle1 = ReplicationBridge->BeginReplication(TestObject1);
 		ReplicationSystem->SetPrioritizer(NetRefHandle1, MockPrioritizerHandle);
 
-		ReplicationSystem->PreSendUpdate(FakeDeltaTime);
+		ReplicationSystem->NetUpdate(FakeDeltaTime);
 		ReplicationSystem->PostSendUpdate();
 	}
 
@@ -394,7 +394,7 @@ UE_NET_TEST_FIXTURE(FTestNetObjectPrioritizerAPIFixture, UpdateObjectsIsNotCalle
 	{
 		MockNetObjectPrioritizer->ResetFunctionCallStatus();
 
-		ReplicationSystem->PreSendUpdate(FakeDeltaTime);
+		ReplicationSystem->NetUpdate(FakeDeltaTime);
 
 		const UMockNetObjectPrioritizer::FFunctionCallStatus& FunctionCallStatus = MockNetObjectPrioritizer->GetFunctionCallStatus();
 
@@ -420,7 +420,7 @@ UE_NET_TEST_FIXTURE(FTestNetObjectPrioritizerAPIFixture, NativeIrisObjectGetsPri
 	const FNetRefHandle NetRefHandle = ReplicationBridge->BeginReplication(Object);
 	ReplicationSystem->SetPrioritizer(NetRefHandle, MockPrioritizerHandle);
 
-	ReplicationSystem->PreSendUpdate(FakeDeltaTime);
+	ReplicationSystem->NetUpdate(FakeDeltaTime);
 	ReplicationSystem->PostSendUpdate();
 
 	UE::Net::Private::FInternalNetRefIndex InternalIndex = NetRefHandleManager->GetInternalIndex(NetRefHandle);
@@ -447,7 +447,7 @@ UE_NET_TEST_FIXTURE(FTestNetObjectPrioritizerAPIFixture, NativeIrisObjectGetsUpd
 	constexpr float UpdatedPriority = 4711;
 	Object->SetPriority(UpdatedPriority);
 
-	ReplicationSystem->PreSendUpdate(FakeDeltaTime);
+	ReplicationSystem->NetUpdate(FakeDeltaTime);
 	ReplicationSystem->PostSendUpdate();
 
 	UE::Net::Private::FInternalNetRefIndex InternalIndex = NetRefHandleManager->GetInternalIndex(NetRefHandle);
@@ -471,7 +471,7 @@ UE_NET_TEST_FIXTURE(FTestNetObjectPrioritizerAPIFixture, ObjectGetsPriorityFromS
 	const FNetRefHandle NetRefHandle = ReplicationBridge->BeginReplication(Object);
 	ReplicationSystem->SetPrioritizer(NetRefHandle, MockPrioritizerHandle);
 
-	ReplicationSystem->PreSendUpdate(FakeDeltaTime);
+	ReplicationSystem->NetUpdate(FakeDeltaTime);
 	ReplicationSystem->PostSendUpdate();
 
 	UE::Net::Private::FInternalNetRefIndex InternalIndex = NetRefHandleManager->GetInternalIndex(NetRefHandle);
@@ -498,7 +498,7 @@ UE_NET_TEST_FIXTURE(FTestNetObjectPrioritizerAPIFixture, ObjectGetsUpdatedPriori
 	constexpr float UpdatedPriority = 4711;
 	Object->SetPriority(UpdatedPriority);
 
-	ReplicationSystem->PreSendUpdate(FakeDeltaTime);
+	ReplicationSystem->NetUpdate(FakeDeltaTime);
 	ReplicationSystem->PostSendUpdate();
 
 	UE::Net::Private::FInternalNetRefIndex InternalIndex = NetRefHandleManager->GetInternalIndex(NetRefHandle);
diff --git a/Engine/Plugins/Runtime/ReplicationSystemTestPlugin/Source/Private/Tests/ReplicationSystem/Prioritization/TestNetObjectPrioritizerFixture.cpp b/Engine/Plugins/Runtime/ReplicationSystemTestPlugin/Source/Private/Tests/ReplicationSystem/Prioritization/TestNetObjectPrioritizerFixture.cpp
--- a/Engine/Plugins/Runtime/ReplicationSystemTestPlugin/Source/Private/Tests/ReplicationSystem/Prioritization/TestNetObjectPrioritizerFixture.cpp
+++ b/Engine/Plugins/Runtime/ReplicationSystemTestPlugin/Source/Private/Tests/ReplicationSystem/Prioritization/TestNetObjectPrioritizerFixture.cpp
@@ -105,7 +105,7 @@ FTestNetObjectPrioritizerFixture::FPrioritizationResult FTestNetObjectPrioritize
 		Private::FReplicationPrioritization& ReplicationPrioritization = Server->GetReplicationSystem()->GetReplicationSystemInternal()->GetPrioritization();
 		Private::FNetRefHandleManager& NetRefHandleManager = Server->GetReplicationSystem()->GetReplicationSystemInternal()->GetNetRefHandleManager();
 
-		Server->PreSendUpdate();
+		Server->NetUpdate();
 
 		{
 			// Retrieve the priorities. These are floats for internal indices so we need to map the NetRefHandles to internal indices to construct the return value.
@@ -157,7 +157,7 @@ FTestNetObjectPrioritizerFixture::FPrioritizationResult FTestNetObjectPrioritize
 		Private::FReplicationPrioritization& ReplicationPrioritization = Server->GetReplicationSystem()->GetReplicationSystemInternal()->GetPrioritization();
 		Private::FNetRefHandleManager& NetRefHandleManager = Server->GetReplicationSystem()->GetReplicationSystemInternal()->GetNetRefHandleManager();
 
-		Server->PreSendUpdate();
+		Server->NetUpdate();
 
 		{
 			// Retrieve the priorities. These are floats for internal indices so we need to map the NetRefHandles to internal indices to construct the return value.
@@ -252,7 +252,7 @@ FNetObjectPrioritizerImage FTestNetObjectPrioritizerFixture::Visualize(FNetObjec
 				Private::FReplicationPrioritization& ReplicationPrioritization = Server->GetReplicationSystem()->GetReplicationSystemInternal()->GetPrioritization();
 				Private::FNetRefHandleManager& NetRefHandleManager = Server->GetReplicationSystem()->GetReplicationSystemInternal()->GetNetRefHandleManager();
 
-				Server->PreSendUpdate();
+				Server->NetUpdate();
 
 				{
 					// Retrieve the priorities. These are floats for internal indices so we need to map the NetRefHandles to internal indices to construct the return value.
diff --git a/Engine/Plugins/Runtime/ReplicationSystemTestPlugin/Source/Private/Tests/ReplicationSystem/RPC/TestObjectRPC.cpp b/Engine/Plugins/Runtime/ReplicationSystemTestPlugin/Source/Private/Tests/ReplicationSystem/RPC/TestObjectRPC.cpp
--- a/Engine/Plugins/Runtime/ReplicationSystemTestPlugin/Source/Private/Tests/ReplicationSystem/RPC/TestObjectRPC.cpp
+++ b/Engine/Plugins/Runtime/ReplicationSystemTestPlugin/Source/Private/Tests/ReplicationSystem/RPC/TestObjectRPC.cpp
@@ -28,7 +28,7 @@ namespace UE::Net::Private
 		Server->ReplicationSystem->SetOwningNetConnection(ServerObject->NetRefHandle, 0x01);
 
 		// Send and deliver packet
-		Server->PreSendUpdate();
+		Server->NetUpdate();
 		Server->SendAndDeliverTo(Client, true);
 		Server->PostSendUpdate();
 
@@ -43,7 +43,7 @@ namespace UE::Net::Private
 		ServerObject->ClientRPC();
 
 		// Send and deliver packet
-		Server->PreSendUpdate();
+		Server->NetUpdate();
 		Server->SendAndDeliverTo(Client, true);
 		Server->PostSendUpdate();
 
@@ -55,7 +55,7 @@ namespace UE::Net::Private
 		ServerObject->ClientRPCWithParam(0xBABA);
 
 		// Send and deliver packet
-		Server->PreSendUpdate();
+		Server->NetUpdate();
 		Server->SendAndDeliverTo(Client, true);
 		Server->PostSendUpdate();
 
@@ -102,7 +102,7 @@ namespace UE::Net::Private
 		ServerObject->ReplicationSystem->SetRPCSendPolicyFlags(UTestReplicatedObjectWithRPC::StaticClass()->FindFunctionByName(GET_FUNCTION_NAME_CHECKED(UTestReplicatedObjectWithRPC, NetMulticast_MultiCastRPCSendImmediate)), ENetObjectAttachmentSendPolicyFlags::SendImmediate);
 
 		// Send and deliver packet
-		Server->PreSendUpdate();
+		Server->NetUpdate();
 		Server->SendAndDeliverTo(Client, true);
 		Server->PostSendUpdate();
 
@@ -121,7 +121,7 @@ namespace UE::Net::Private
 		ServerObject->NetMulticast_MultiCastRPCSendImmediate();
 
 		// Send and deliver packet, simulating a send from PostTickDispatch
-		Server->PreSendUpdate(UReplicationSystem::FSendUpdateParams {.SendPass = UE::Net::EReplicationSystemSendPass::PostTickDispatch});
+		Server->TickPostReceive();
 		Server->SendAndDeliverTo(Client, true);
 		Server->PostSendUpdate();
 
@@ -130,7 +130,7 @@ namespace UE::Net::Private
 		UE_NET_ASSERT_EQ(ClientObject->NetMulticast_MultiCastRPCSendImmediateCallOrder, 1);
 
 		// Send and deliver packet
-		Server->PreSendUpdate();
+		Server->NetUpdate();
 		Server->SendAndDeliverTo(Client, true);
 		Server->PostSendUpdate();
 		
@@ -158,7 +158,7 @@ namespace UE::Net::Private
 		const FNetRefHandle ServerSubObjectHandle = ServerSubObject->NetRefHandle;
 
 		// Send and deliver packet
-		Server->PreSendUpdate();
+		Server->NetUpdate();
 		Server->SendAndDeliverTo(Client, true);
 		Server->PostSendUpdate();
 
diff --git a/Engine/Plugins/Runtime/ReplicationSystemTestPlugin/Source/Private/Tests/ReplicationSystem/ReplicationSystemServerClientTestFixture.cpp b/Engine/Plugins/Runtime/ReplicationSystemTestPlugin/Source/Private/Tests/ReplicationSystem/ReplicationSystemServerClientTestFixture.cpp
--- a/Engine/Plugins/Runtime/ReplicationSystemTestPlugin/Source/Private/Tests/ReplicationSystem/ReplicationSystemServerClientTestFixture.cpp
+++ b/Engine/Plugins/Runtime/ReplicationSystemTestPlugin/Source/Private/Tests/ReplicationSystem/ReplicationSystemServerClientTestFixture.cpp
@@ -353,15 +353,14 @@ void FReplicationSystemTestNode::RemoveConnection(uint32 ConnectionId)
 	}
 }
 
-void FReplicationSystemTestNode::PreSendUpdate(const UReplicationSystem::FSendUpdateParams& Params)
+void FReplicationSystemTestNode::NetUpdate()
 {
-	ReplicationSystem->PreSendUpdate(Params);
-	CurrentSendPass = Params.SendPass;
+	ReplicationSystem->NetUpdate(1.f);
 }
 
-void FReplicationSystemTestNode::PreSendUpdate()
+void FReplicationSystemTestNode::TickPostReceive()
 {
-	PreSendUpdate(UReplicationSystem::FSendUpdateParams {.SendPass = EReplicationSystemSendPass::TickFlush, .DeltaSeconds = 1.f});
+	ReplicationSystem->TickPostReceive();
 }
 
 bool FReplicationSystemTestNode::SendUpdate(uint32 ConnectionId, const TCHAR* Desc)
@@ -501,7 +500,7 @@ bool FReplicationSystemTestClient::UpdateAndSend(FReplicationSystemTestServer* S
 {
 	bool bSuccess = false;
 
-	PreSendUpdate();
+	NetUpdate();
 
 	if (SendUpdate(Desc))
 	{
@@ -553,7 +552,7 @@ bool FReplicationSystemTestServer::UpdateAndSend(const TArrayView<FReplicationSy
 {
 	bool bSuccess = true;
 
-	PreSendUpdate();
+	NetUpdate();
 
 	for (FReplicationSystemTestClient* Client : Clients)
 	{
diff --git a/Engine/Plugins/Runtime/ReplicationSystemTestPlugin/Source/Private/Tests/ReplicationSystem/ReplicationSystemServerClientTestFixture.h b/Engine/Plugins/Runtime/ReplicationSystemTestPlugin/Source/Private/Tests/ReplicationSystem/ReplicationSystemServerClientTestFixture.h
--- a/Engine/Plugins/Runtime/ReplicationSystemTestPlugin/Source/Private/Tests/ReplicationSystem/ReplicationSystemServerClientTestFixture.h
+++ b/Engine/Plugins/Runtime/ReplicationSystemTestPlugin/Source/Private/Tests/ReplicationSystem/ReplicationSystemServerClientTestFixture.h
@@ -168,14 +168,15 @@ public:
 	void RemoveConnection(uint32 ConnectionId);
 
 	// System Update
-	void PreSendUpdate(const UReplicationSystem::FSendUpdateParams& Params);
-	void PreSendUpdate();
-	void PostSendUpdate();
+	void NetUpdate();
+	void TickPostReceive();
 
 	// Update methods for server connections
 	bool SendUpdate(uint32 ConnectionId, const TCHAR* Desc = nullptr);
 	bool SendUpdate(const TCHAR* Desc = nullptr) { return SendUpdate(1, Desc); }
 
+	void PostSendUpdate();
+
 	void DeliverTo(FReplicationSystemTestNode& Dest, uint32 LocalConnectionId, uint32 RemoteConnectionId, bool bDeliver);
 	void RecvUpdate(uint32 ConnectionId, FNetSerializationContext& Context);
 	void RecvUpdate(FNetSerializationContext& Context) { RecvUpdate(1, Context); }
diff --git a/Engine/Plugins/Runtime/ReplicationSystemTestPlugin/Source/Private/Tests/ReplicationSystem/TestCancelPendingDestroy.cpp b/Engine/Plugins/Runtime/ReplicationSystemTestPlugin/Source/Private/Tests/ReplicationSystem/TestCancelPendingDestroy.cpp
--- a/Engine/Plugins/Runtime/ReplicationSystemTestPlugin/Source/Private/Tests/ReplicationSystem/TestCancelPendingDestroy.cpp
+++ b/Engine/Plugins/Runtime/ReplicationSystemTestPlugin/Source/Private/Tests/ReplicationSystem/TestCancelPendingDestroy.cpp
@@ -21,18 +21,18 @@ UE_NET_TEST_FIXTURE(FTestCancelPendingDestroyFixture, TestCancelPendingDestroyDu
 	UTestReplicatedIrisObject* ServerObject = Server->CreateObject(UTestReplicatedIrisObject::FComponents());
 
 	// Write packet
-	Server->PreSendUpdate();
+	Server->NetUpdate();
 	Server->SendUpdate(Client->ConnectionIdOnServer);
 	Server->PostSendUpdate();
 
 	// Filter out object to cause a PendingDestroy
 	Server->GetReplicationSystem()->AddToGroup(Server->GetReplicationSystem()->GetNotReplicatedNetObjectGroup(), ServerObject->NetRefHandle);
-	Server->PreSendUpdate();
+	Server->NetUpdate();
 	Server->PostSendUpdate();
 
 	// Remove object from filter to cause object to end up in CancelPendingDestroy
 	Server->GetReplicationSystem()->RemoveFromGroup(Server->GetReplicationSystem()->GetNotReplicatedNetObjectGroup(), ServerObject->NetRefHandle);
-	Server->PreSendUpdate();
+	Server->NetUpdate();
 	Server->PostSendUpdate();
 
 	// Deliver original packet
@@ -44,7 +44,7 @@ UE_NET_TEST_FIXTURE(FTestCancelPendingDestroyFixture, TestCancelPendingDestroyDu
 	// Modify a property on the object and make sure it's replicated as the object should now be confirmed created
 	ServerObject->IntA ^= 1;
 
-	Server->PreSendUpdate();
+	Server->NetUpdate();
 	Server->SendAndDeliverTo(Client, DeliverPacket);
 	Server->PostSendUpdate();
 
@@ -62,25 +62,25 @@ UE_NET_TEST_FIXTURE(FTestCancelPendingDestroyFixture, TestCancelPendingDestroyDu
 	UTestReplicatedIrisObject* ServerObject = Server->CreateObject(UTestReplicatedIrisObject::FComponents());
 
 	// Write packet
-	Server->PreSendUpdate();
+	Server->NetUpdate();
 	Server->SendUpdate(Client->ConnectionIdOnServer);
 	Server->PostSendUpdate();
 
 	// Filter out object to cause a PendingDestroy
 	Server->GetReplicationSystem()->AddToGroup(Server->GetReplicationSystem()->GetNotReplicatedNetObjectGroup(), ServerObject->NetRefHandle);
-	Server->PreSendUpdate();
+	Server->NetUpdate();
 	Server->PostSendUpdate();
 
 	// Remove object from filter to cause object to end up in CancelPendingDestroy
 	Server->GetReplicationSystem()->RemoveFromGroup(Server->GetReplicationSystem()->GetNotReplicatedNetObjectGroup(), ServerObject->NetRefHandle);
-	Server->PreSendUpdate();
+	Server->NetUpdate();
 	Server->PostSendUpdate();
 
 	// Cause packet loss on object creation
 	Server->DeliverTo(Client, DoNotDeliverPacket);
 
 	// Write and send packet and verify object is created
-	Server->PreSendUpdate();
+	Server->NetUpdate();
 	Server->SendAndDeliverTo(Client, DeliverPacket);
 	Server->PostSendUpdate();
 
@@ -96,13 +96,13 @@ UE_NET_TEST_FIXTURE(FTestCancelPendingDestroyFixture, TestCancelPendingDestroyDu
 	UTestReplicatedIrisObject* ServerObject = Server->CreateObject(UTestReplicatedIrisObject::FComponents());
 
 	// Write packet
-	Server->PreSendUpdate();
+	Server->NetUpdate();
 	Server->SendUpdate(Client->ConnectionIdOnServer);
 	Server->PostSendUpdate();
 
 	// Filter out object to cause a PendingDestroy
 	Server->GetReplicationSystem()->AddToGroup(Server->GetReplicationSystem()->GetNotReplicatedNetObjectGroup(), ServerObject->NetRefHandle);
-	Server->PreSendUpdate();
+	Server->NetUpdate();
 	Server->PostSendUpdate();
 
 	// Notify packet loss, which should set state to Invalid/not replicated
@@ -110,11 +110,11 @@ UE_NET_TEST_FIXTURE(FTestCancelPendingDestroyFixture, TestCancelPendingDestroyDu
 
 	// Remove object from filter to cause object to end up in PendingCreate
 	Server->GetReplicationSystem()->RemoveFromGroup(Server->GetReplicationSystem()->GetNotReplicatedNetObjectGroup(), ServerObject->NetRefHandle);
-	Server->PreSendUpdate();
+	Server->NetUpdate();
 	Server->PostSendUpdate();
 
 	// Write and send packet and verify object is created
-	Server->PreSendUpdate();
+	Server->NetUpdate();
 	Server->SendAndDeliverTo(Client, DeliverPacket);
 	Server->PostSendUpdate();
 
@@ -130,7 +130,7 @@ UE_NET_TEST_FIXTURE(FTestCancelPendingDestroyFixture, TestCancelPendingDestroyDu
 	UTestReplicatedIrisObject* ServerObject = Server->CreateObject(UTestReplicatedIrisObject::FComponents());
 
 	// Write and send packet
-	Server->PreSendUpdate();
+	Server->NetUpdate();
 	Server->SendAndDeliverTo(Client, DeliverPacket);
 	Server->PostSendUpdate();
 
@@ -139,18 +139,18 @@ UE_NET_TEST_FIXTURE(FTestCancelPendingDestroyFixture, TestCancelPendingDestroyDu
 
 	// Filter out object to cause a PendingDestroy
 	Server->GetReplicationSystem()->AddToGroup(Server->GetReplicationSystem()->GetNotReplicatedNetObjectGroup(), ServerObject->NetRefHandle);
-	Server->PreSendUpdate();
+	Server->NetUpdate();
 	Server->PostSendUpdate();
 
 	// Remove object from filter to cause object to end up in CancelPendingDestroy
 	Server->GetReplicationSystem()->RemoveFromGroup(Server->GetReplicationSystem()->GetNotReplicatedNetObjectGroup(), ServerObject->NetRefHandle);
-	Server->PreSendUpdate();
+	Server->NetUpdate();
 	Server->PostSendUpdate();
 
 	// Modify a property on the object and make sure it's replicated as the object should still be created
 	ServerObject->IntA ^= 1;
 
-	Server->PreSendUpdate();
+	Server->NetUpdate();
 	Server->SendAndDeliverTo(Client, DeliverPacket);
 	Server->PostSendUpdate();
 
@@ -168,19 +168,19 @@ UE_NET_TEST_FIXTURE(FTestCancelPendingDestroyFixture, TestCancelPendingDestroyDu
 	UTestReplicatedIrisObject* ServerObject = Server->CreateObject(UTestReplicatedIrisObject::FComponents());
 
 	// Write and send packet
-	Server->PreSendUpdate();
+	Server->NetUpdate();
 	Server->SendAndDeliverTo(Client, DeliverPacket);
 	Server->PostSendUpdate();
 
 	// Filter out object and write packet to cause object to end up in WaitOnDestroyConfirmation
 	Server->GetReplicationSystem()->AddToGroup(Server->GetReplicationSystem()->GetNotReplicatedNetObjectGroup(), ServerObject->NetRefHandle);
-	Server->PreSendUpdate();
+	Server->NetUpdate();
 	Server->SendUpdate(Client->ConnectionIdOnServer);
 	Server->PostSendUpdate();
 
 	// Remove object from filter to cause object to end up in CancelPendingDestroy
 	Server->GetReplicationSystem()->RemoveFromGroup(Server->GetReplicationSystem()->GetNotReplicatedNetObjectGroup(), ServerObject->NetRefHandle);
-	Server->PreSendUpdate();
+	Server->NetUpdate();
 	Server->PostSendUpdate();
 
 	// Cause packet loss on the object destroy packet
@@ -192,7 +192,7 @@ UE_NET_TEST_FIXTURE(FTestCancelPendingDestroyFixture, TestCancelPendingDestroyDu
 	// Modify a property on the object and make sure it's replicated as the object should still be created
 	ServerObject->IntA ^= 1;
 
-	Server->PreSendUpdate();
+	Server->NetUpdate();
 	Server->SendAndDeliverTo(Client, DeliverPacket);
 	Server->PostSendUpdate();
 
@@ -210,13 +210,13 @@ UE_NET_TEST_FIXTURE(FTestCancelPendingDestroyFixture, TestCancelPendingDestroyDu
 	UTestReplicatedIrisObject* ServerObject = Server->CreateObject(UTestReplicatedIrisObject::FComponents());
 
 	// Write and send packet
-	Server->PreSendUpdate();
+	Server->NetUpdate();
 	Server->SendAndDeliverTo(Client, DeliverPacket);
 	Server->PostSendUpdate();
 
 	// Filter out object and write packet to cause object to end up in WaitOnDestroyConfirmation
 	Server->GetReplicationSystem()->AddToGroup(Server->GetReplicationSystem()->GetNotReplicatedNetObjectGroup(), ServerObject->NetRefHandle);
-	Server->PreSendUpdate();
+	Server->NetUpdate();
 	Server->SendUpdate(Client->ConnectionIdOnServer);
 	Server->PostSendUpdate();
 
@@ -225,7 +225,7 @@ UE_NET_TEST_FIXTURE(FTestCancelPendingDestroyFixture, TestCancelPendingDestroyDu
 
 	// Remove object from filter to cause object to end up in CancelPendingDestroy
 	Server->GetReplicationSystem()->RemoveFromGroup(Server->GetReplicationSystem()->GetNotReplicatedNetObjectGroup(), ServerObject->NetRefHandle);
-	Server->PreSendUpdate();
+	Server->NetUpdate();
 	Server->PostSendUpdate();
 
 	// Verify that object still exists on client
@@ -234,7 +234,7 @@ UE_NET_TEST_FIXTURE(FTestCancelPendingDestroyFixture, TestCancelPendingDestroyDu
 	// Modify a property on the object and make sure it's replicated as the object should still be created
 	ServerObject->IntA ^= 1;
 
-	Server->PreSendUpdate();
+	Server->NetUpdate();
 	Server->SendAndDeliverTo(Client, DeliverPacket);
 	Server->PostSendUpdate();
 
@@ -253,14 +253,14 @@ UE_NET_TEST_FIXTURE(FTestCancelPendingDestroyFixture, TestCancelPendingDestroyDu
 	UTestReplicatedIrisObject* ServerObject = Server->CreateObject(UTestReplicatedIrisObject::FComponents());
 
 	// Introduce latency by not immediately delivering packets.
-	Server->PreSendUpdate();
+	Server->NetUpdate();
 	Server->SendTo(Client, TEXT("Create object"));
 	Server->PostSendUpdate();
 
 	// Filter out object to cause it to end up being PendingDestroy
 	Server->ReplicationSystem->AddToGroup(Server->GetReplicationSystem()->GetNotReplicatedNetObjectGroup(), ServerObject->NetRefHandle);
 
-	Server->PreSendUpdate();
+	Server->NetUpdate();
 	Server->SendTo(Client, TEXT("Destroy object"));
 	Server->PostSendUpdate();
 
@@ -270,7 +270,7 @@ UE_NET_TEST_FIXTURE(FTestCancelPendingDestroyFixture, TestCancelPendingDestroyDu
 	// Modify object.
 	ServerObject->IntA += 1;
 
-	Server->PreSendUpdate();
+	Server->NetUpdate();
 	Server->SendTo(Client, TEXT("Update object"));
 	Server->PostSendUpdate();
 
@@ -305,20 +305,20 @@ UE_NET_TEST_FIXTURE(FTestCancelPendingDestroyFixture, TestCancelPendingDestroyDu
 	UTestReplicatedIrisObject* ServerObject = Server->CreateObject(UTestReplicatedIrisObject::FComponents());
 
 	// Write and send packet
-	Server->PreSendUpdate();
+	Server->NetUpdate();
 	Server->SendAndDeliverTo(Client, DeliverPacket);
 	Server->PostSendUpdate();
 
 	// Filter out object and write packet to cause object to end up in WaitOnDestroyConfirmation
 	Server->GetReplicationSystem()->AddToGroup(Server->GetReplicationSystem()->GetNotReplicatedNetObjectGroup(), ServerObject->NetRefHandle);
-	Server->PreSendUpdate();
+	Server->NetUpdate();
 	Server->SendUpdate(Client->ConnectionIdOnServer);
 	Server->PostSendUpdate();
 
 	// Remove object from filter to cause object to end up in CancelPendingDestroy, also modify a property 
 	Server->GetReplicationSystem()->RemoveFromGroup(Server->GetReplicationSystem()->GetNotReplicatedNetObjectGroup(), ServerObject->NetRefHandle);
 	ServerObject->IntA ^= 1;
-	Server->PreSendUpdate();
+	Server->NetUpdate();
 	Server->PostSendUpdate();
 
 	// Deliver object destroy packet
@@ -328,7 +328,7 @@ UE_NET_TEST_FIXTURE(FTestCancelPendingDestroyFixture, TestCancelPendingDestroyDu
 	UE_NET_ASSERT_FALSE(Client->IsValidNetRefHandle(ServerObject->NetRefHandle));
 
 	// Write and send packet
-	Server->PreSendUpdate();
+	Server->NetUpdate();
 	Server->SendAndDeliverTo(Client, DeliverPacket);
 	Server->PostSendUpdate();
 
@@ -347,23 +347,23 @@ UE_NET_TEST_FIXTURE(FTestCancelPendingDestroyFixture, TestCancelingCancelPending
 	UTestReplicatedIrisObject* ServerObject = Server->CreateObject(UTestReplicatedIrisObject::FComponents());
 
 	// Write packet
-	Server->PreSendUpdate();
+	Server->NetUpdate();
 	Server->SendUpdate(Client->ConnectionIdOnServer);
 	Server->PostSendUpdate();
 
 	// Filter out object to cause a PendingDestroy
 	Server->GetReplicationSystem()->AddToGroup(Server->GetReplicationSystem()->GetNotReplicatedNetObjectGroup(), ServerObject->NetRefHandle);
-	Server->PreSendUpdate();
+	Server->NetUpdate();
 	Server->PostSendUpdate();
 
 	// Remove object from filter to cause object to end up in CancelPendingDestroy...
 	Server->GetReplicationSystem()->RemoveFromGroup(Server->GetReplicationSystem()->GetNotReplicatedNetObjectGroup(), ServerObject->NetRefHandle);
-	Server->PreSendUpdate();
+	Server->NetUpdate();
 	Server->PostSendUpdate();
 
 	// ... and cancel that thought! We want the object to be destroyed after all.
 	Server->GetReplicationSystem()->AddToGroup(Server->GetReplicationSystem()->GetNotReplicatedNetObjectGroup(), ServerObject->NetRefHandle);
-	Server->PreSendUpdate();
+	Server->NetUpdate();
 	Server->PostSendUpdate();
 
 	Server->DeliverTo(Client, DeliverPacket);
@@ -372,7 +372,7 @@ UE_NET_TEST_FIXTURE(FTestCancelPendingDestroyFixture, TestCancelingCancelPending
 	UE_NET_ASSERT_NE(Client->GetReplicationBridge()->GetReplicatedObject(ServerObject->NetRefHandle), nullptr);
 
 	// Write and send a packet
-	Server->PreSendUpdate();
+	Server->NetUpdate();
 	Server->SendAndDeliverTo(Client, DeliverPacket);
 	Server->PostSendUpdate();
 
@@ -389,23 +389,23 @@ UE_NET_TEST_FIXTURE(FTestCancelPendingDestroyFixture, TestCancelingCancelPending
 	UTestReplicatedIrisObject* ServerObject = Server->CreateObject(UTestReplicatedIrisObject::FComponents());
 
 	// Write packet
-	Server->PreSendUpdate();
+	Server->NetUpdate();
 	Server->SendUpdate(Client->ConnectionIdOnServer);
 	Server->PostSendUpdate();
 
 	// Filter out object to cause a PendingDestroy
 	Server->GetReplicationSystem()->AddToGroup(Server->GetReplicationSystem()->GetNotReplicatedNetObjectGroup(), ServerObject->NetRefHandle);
-	Server->PreSendUpdate();
+	Server->NetUpdate();
 	Server->PostSendUpdate();
 
 	// Remove object from filter to cause object to end up in CancelPendingDestroy...
 	Server->GetReplicationSystem()->RemoveFromGroup(Server->GetReplicationSystem()->GetNotReplicatedNetObjectGroup(), ServerObject->NetRefHandle);
-	Server->PreSendUpdate();
+	Server->NetUpdate();
 	Server->PostSendUpdate();
 
 	// ... and cancel that thought! We want the object to be destroyed after all.
 	Server->GetReplicationSystem()->AddToGroup(Server->GetReplicationSystem()->GetNotReplicatedNetObjectGroup(), ServerObject->NetRefHandle);
-	Server->PreSendUpdate();
+	Server->NetUpdate();
 	Server->PostSendUpdate();
 
 	// Cause packet loss on object creation
@@ -415,7 +415,7 @@ UE_NET_TEST_FIXTURE(FTestCancelPendingDestroyFixture, TestCancelingCancelPending
 	UE_NET_ASSERT_EQ(Client->GetReplicationBridge()->GetReplicatedObject(ServerObject->NetRefHandle), nullptr);
 
 	// Write and send a packet
-	Server->PreSendUpdate();
+	Server->NetUpdate();
 	Server->SendAndDeliverTo(Client, DeliverPacket);
 	Server->PostSendUpdate();
 
@@ -432,7 +432,7 @@ UE_NET_TEST_FIXTURE(FTestCancelPendingDestroyFixture, TestCancelingCancelPending
 	UTestReplicatedIrisObject* ServerObject = Server->CreateObject(UTestReplicatedIrisObject::FComponents());
 
 	// Write and send packet
-	Server->PreSendUpdate();
+	Server->NetUpdate();
 	Server->SendAndDeliverTo(Client, DeliverPacket);
 	Server->PostSendUpdate();
 
@@ -441,21 +441,21 @@ UE_NET_TEST_FIXTURE(FTestCancelPendingDestroyFixture, TestCancelingCancelPending
 
 	// Filter out object to cause a PendingDestroy
 	Server->GetReplicationSystem()->AddToGroup(Server->GetReplicationSystem()->GetNotReplicatedNetObjectGroup(), ServerObject->NetRefHandle);
-	Server->PreSendUpdate();
+	Server->NetUpdate();
 	Server->PostSendUpdate();
 
 	// Remove object from filter to cause object to end up in CancelPendingDestroy...
 	Server->GetReplicationSystem()->RemoveFromGroup(Server->GetReplicationSystem()->GetNotReplicatedNetObjectGroup(), ServerObject->NetRefHandle);
-	Server->PreSendUpdate();
+	Server->NetUpdate();
 	Server->PostSendUpdate();
 
 	// ... and cancel that thought! We want the object to be destroyed after all.
 	Server->GetReplicationSystem()->AddToGroup(Server->GetReplicationSystem()->GetNotReplicatedNetObjectGroup(), ServerObject->NetRefHandle);
-	Server->PreSendUpdate();
+	Server->NetUpdate();
 	Server->PostSendUpdate();
 
 	// Write and send a packet
-	Server->PreSendUpdate();
+	Server->NetUpdate();
 	Server->SendAndDeliverTo(Client, DeliverPacket);
 	Server->PostSendUpdate();
 
@@ -472,24 +472,24 @@ UE_NET_TEST_FIXTURE(FTestCancelPendingDestroyFixture, TestCancelingCancelPending
 	UTestReplicatedIrisObject* ServerObject = Server->CreateObject(UTestReplicatedIrisObject::FComponents());
 
 	// Write and send packet
-	Server->PreSendUpdate();
+	Server->NetUpdate();
 	Server->SendAndDeliverTo(Client, DeliverPacket);
 	Server->PostSendUpdate();
 
 	// Filter out object and write packet to cause object to end up in WaitOnDestroyConfirmation
 	Server->GetReplicationSystem()->AddToGroup(Server->GetReplicationSystem()->GetNotReplicatedNetObjectGroup(), ServerObject->NetRefHandle);
-	Server->PreSendUpdate();
+	Server->NetUpdate();
 	Server->SendUpdate(Client->ConnectionIdOnServer);
 	Server->PostSendUpdate();
 
 	// Remove object from filter to cause object to end up in CancelPendingDestroy...
 	Server->GetReplicationSystem()->RemoveFromGroup(Server->GetReplicationSystem()->GetNotReplicatedNetObjectGroup(), ServerObject->NetRefHandle);
-	Server->PreSendUpdate();
+	Server->NetUpdate();
 	Server->PostSendUpdate();
 
 	// ... and cancel that thought! We want the object to be destroyed after all.
 	Server->GetReplicationSystem()->AddToGroup(Server->GetReplicationSystem()->GetNotReplicatedNetObjectGroup(), ServerObject->NetRefHandle);
-	Server->PreSendUpdate();
+	Server->NetUpdate();
 	Server->PostSendUpdate();
 
 	// Cause packet loss on the object destroy packet
@@ -499,7 +499,7 @@ UE_NET_TEST_FIXTURE(FTestCancelPendingDestroyFixture, TestCancelingCancelPending
 	UE_NET_ASSERT_NE(Client->GetReplicationBridge()->GetReplicatedObject(ServerObject->NetRefHandle), nullptr);
 
 	// Write and send a packet
-	Server->PreSendUpdate();
+	Server->NetUpdate();
 	Server->SendAndDeliverTo(Client, DeliverPacket);
 	Server->PostSendUpdate();
 
@@ -516,24 +516,24 @@ UE_NET_TEST_FIXTURE(FTestCancelPendingDestroyFixture, TestCancelingCancelPending
 	UTestReplicatedIrisObject* ServerObject = Server->CreateObject(UTestReplicatedIrisObject::FComponents());
 
 	// Write and send packet
-	Server->PreSendUpdate();
+	Server->NetUpdate();
 	Server->SendAndDeliverTo(Client, DeliverPacket);
 	Server->PostSendUpdate();
 
 	// Filter out object and write packet to cause object to end up in WaitOnDestroyConfirmation
 	Server->GetReplicationSystem()->AddToGroup(Server->GetReplicationSystem()->GetNotReplicatedNetObjectGroup(), ServerObject->NetRefHandle);
-	Server->PreSendUpdate();
+	Server->NetUpdate();
 	Server->SendUpdate(Client->ConnectionIdOnServer);
 	Server->PostSendUpdate();
 
 	// Remove object from filter to cause object to end up in CancelPendingDestroy...
 	Server->GetReplicationSystem()->RemoveFromGroup(Server->GetReplicationSystem()->GetNotReplicatedNetObjectGroup(), ServerObject->NetRefHandle);
-	Server->PreSendUpdate();
+	Server->NetUpdate();
 	Server->PostSendUpdate();
 
 	// ... and cancel that thought! We want the object to be destroyed after all.
 	Server->GetReplicationSystem()->AddToGroup(Server->GetReplicationSystem()->GetNotReplicatedNetObjectGroup(), ServerObject->NetRefHandle);
-	Server->PreSendUpdate();
+	Server->NetUpdate();
 	Server->PostSendUpdate();
 
 	// Deliver destroy packet
@@ -543,7 +543,7 @@ UE_NET_TEST_FIXTURE(FTestCancelPendingDestroyFixture, TestCancelingCancelPending
 	UE_NET_ASSERT_EQ(Client->GetReplicationBridge()->GetReplicatedObject(ServerObject->NetRefHandle), nullptr);
 
 	// Write and send a packet
-	Server->PreSendUpdate();
+	Server->NetUpdate();
 	Server->SendAndDeliverTo(Client, DeliverPacket);
 	Server->PostSendUpdate();
 
@@ -563,7 +563,7 @@ UE_NET_TEST_FIXTURE(FTestCancelPendingDestroyFixture, TestCancelPendingSubObject
 	UTestReplicatedIrisObject* ServerSubObject = Server->CreateSubObject(ServerObject->NetRefHandle, UTestReplicatedIrisObject::FComponents());
 
 	// Write and send packet
-	Server->PreSendUpdate();
+	Server->NetUpdate();
 	Server->SendAndDeliverTo(Client, DeliverPacket);
 	Server->PostSendUpdate();
 
@@ -572,18 +572,18 @@ UE_NET_TEST_FIXTURE(FTestCancelPendingDestroyFixture, TestCancelPendingSubObject
 
 	// Filter out subobject to cause a SubObjectPendingDestroy
 	Server->GetReplicationSystem()->AddToGroup(Server->GetReplicationSystem()->GetNotReplicatedNetObjectGroup(), ServerSubObject->NetRefHandle);
-	Server->PreSendUpdate();
+	Server->NetUpdate();
 	Server->PostSendUpdate();
 
 	// Remove subobject from filter to cause object to end up in CancelPendingDestroy
 	Server->GetReplicationSystem()->RemoveFromGroup(Server->GetReplicationSystem()->GetNotReplicatedNetObjectGroup(), ServerSubObject->NetRefHandle);
-	Server->PreSendUpdate();
+	Server->NetUpdate();
 	Server->PostSendUpdate();
 
 	// Modify a property on the subobject and make sure it's replicated as the subobject should still exist
 	ServerSubObject->IntA ^= 1;
 
-	Server->PreSendUpdate();
+	Server->NetUpdate();
 	Server->SendAndDeliverTo(Client, DeliverPacket);
 	Server->PostSendUpdate();
 
@@ -611,12 +611,12 @@ UE_NET_TEST_FIXTURE(FTestCancelPendingDestroyFixture, TestDestroySubObjectWhenRo
 
 	// Destroy subobject to cause a SubObjectPendingDestroy
 	Server->DestroyObject(ServerSubObject);
-	Server->PreSendUpdate();
+	Server->NetUpdate();
 	Server->PostSendUpdate();
 
 	// Filter out root object and write packet. This causes subobject to go from SubObjectPendingDestroy to PendingDestroy.
 	Server->GetReplicationSystem()->AddToGroup(Server->GetReplicationSystem()->GetNotReplicatedNetObjectGroup(), ServerObject->NetRefHandle);
-	Server->PreSendUpdate();
+	Server->NetUpdate();
 	Server->SendUpdate(Client->ConnectionIdOnServer);
 	Server->PostSendUpdate();
 
@@ -645,19 +645,19 @@ UE_NET_TEST_FIXTURE(FTestCancelPendingDestroyFixture, TestCancelPendingSubObject
 	UTestReplicatedIrisObject* ServerSubObject = Server->CreateSubObject(ServerObject->NetRefHandle, UTestReplicatedIrisObject::FComponents());
 
 	// Write and send packet
-	Server->PreSendUpdate();
+	Server->NetUpdate();
 	Server->SendAndDeliverTo(Client, DeliverPacket);
 	Server->PostSendUpdate();
 
 	// Filter out object and write packet to cause subobject to end up in WaitOnDestroyConfirmation
 	Server->GetReplicationSystem()->AddToGroup(Server->GetReplicationSystem()->GetNotReplicatedNetObjectGroup(), ServerSubObject->NetRefHandle);
-	Server->PreSendUpdate();
+	Server->NetUpdate();
 	Server->SendUpdate(Client->ConnectionIdOnServer);
 	Server->PostSendUpdate();
 
 	// Remove subobject from filter to cause object to end up in CancelPendingDestroy
 	Server->GetReplicationSystem()->RemoveFromGroup(Server->GetReplicationSystem()->GetNotReplicatedNetObjectGroup(), ServerSubObject->NetRefHandle);
-	Server->PreSendUpdate();
+	Server->NetUpdate();
 	Server->PostSendUpdate();
 
 	// Cause packet loss on the subobject destroy packet
@@ -669,7 +669,7 @@ UE_NET_TEST_FIXTURE(FTestCancelPendingDestroyFixture, TestCancelPendingSubObject
 	// Modify a property on the subobject and make sure it's replicated as the subobject should still exist
 	ServerSubObject->IntA ^= 1;
 
-	Server->PreSendUpdate();
+	Server->NetUpdate();
 	Server->SendAndDeliverTo(Client, DeliverPacket);
 	Server->PostSendUpdate();
 
@@ -688,13 +688,13 @@ UE_NET_TEST_FIXTURE(FTestCancelPendingDestroyFixture, TestCancelPendingSubObject
 	UTestReplicatedIrisObject* ServerSubObject = Server->CreateSubObject(ServerObject->NetRefHandle, UTestReplicatedIrisObject::FComponents());
 
 	// Write and send packet
-	Server->PreSendUpdate();
+	Server->NetUpdate();
 	Server->SendAndDeliverTo(Client, DeliverPacket);
 	Server->PostSendUpdate();
 
 	// Filter out object and write packet to cause subobject to end up in WaitOnDestroyConfirmation
 	Server->GetReplicationSystem()->AddToGroup(Server->GetReplicationSystem()->GetNotReplicatedNetObjectGroup(), ServerSubObject->NetRefHandle);
-	Server->PreSendUpdate();
+	Server->NetUpdate();
 	Server->SendUpdate(Client->ConnectionIdOnServer);
 	Server->PostSendUpdate();
 
@@ -703,7 +703,7 @@ UE_NET_TEST_FIXTURE(FTestCancelPendingDestroyFixture, TestCancelPendingSubObject
 
 	// Remove subobject from filter to cause object to end up in CancelPendingDestroy
 	Server->GetReplicationSystem()->RemoveFromGroup(Server->GetReplicationSystem()->GetNotReplicatedNetObjectGroup(), ServerSubObject->NetRefHandle);
-	Server->PreSendUpdate();
+	Server->NetUpdate();
 	Server->PostSendUpdate();
 
 	// Verify that object still exists on client
@@ -712,7 +712,7 @@ UE_NET_TEST_FIXTURE(FTestCancelPendingDestroyFixture, TestCancelPendingSubObject
 	// Modify a property on the subobject and make sure it's replicated as the subobject should still exist
 	ServerSubObject->IntA ^= 1;
 
-	Server->PreSendUpdate();
+	Server->NetUpdate();
 	Server->SendAndDeliverTo(Client, DeliverPacket);
 	Server->PostSendUpdate();
 
@@ -732,20 +732,20 @@ UE_NET_TEST_FIXTURE(FTestCancelPendingDestroyFixture, TestCancelPendingSubObject
 	UTestReplicatedIrisObject* ServerSubObject = Server->CreateSubObject(ServerObject->NetRefHandle, UTestReplicatedIrisObject::FComponents());
 
 	// Write and send packet
-	Server->PreSendUpdate();
+	Server->NetUpdate();
 	Server->SendAndDeliverTo(Client, DeliverPacket);
 	Server->PostSendUpdate();
 
 	// Filter out object and write packet to cause subobject to end up in WaitOnDestroyConfirmation
 	Server->GetReplicationSystem()->AddToGroup(Server->GetReplicationSystem()->GetNotReplicatedNetObjectGroup(), ServerSubObject->NetRefHandle);
-	Server->PreSendUpdate();
+	Server->NetUpdate();
 	Server->SendUpdate(Client->ConnectionIdOnServer);
 	Server->PostSendUpdate();
 
 	// Remove object from filter to cause object to end up in CancelPendingDestroy, also modify a property 
 	Server->GetReplicationSystem()->RemoveFromGroup(Server->GetReplicationSystem()->GetNotReplicatedNetObjectGroup(), ServerSubObject->NetRefHandle);
 	ServerSubObject->IntA ^= 1;
-	Server->PreSendUpdate();
+	Server->NetUpdate();
 	Server->PostSendUpdate();
 
 	// Deliver subobject destroy packet
@@ -755,7 +755,7 @@ UE_NET_TEST_FIXTURE(FTestCancelPendingDestroyFixture, TestCancelPendingSubObject
 	UE_NET_ASSERT_EQ(Client->GetReplicationBridge()->GetReplicatedObject(ServerSubObject->NetRefHandle), nullptr);
 
 	// Write and send a packet
-	Server->PreSendUpdate();
+	Server->NetUpdate();
 	Server->SendAndDeliverTo(Client, DeliverPacket);
 	Server->PostSendUpdate();
 
@@ -775,7 +775,7 @@ UE_NET_TEST_FIXTURE(FTestCancelPendingDestroyFixture, TestCancelingCancelPending
 	UTestReplicatedIrisObject* ServerSubObject = Server->CreateSubObject(ServerObject->NetRefHandle, UTestReplicatedIrisObject::FComponents());
 
 	// Write and send packet
-	Server->PreSendUpdate();
+	Server->NetUpdate();
 	Server->SendAndDeliverTo(Client, DeliverPacket);
 	Server->PostSendUpdate();
 
@@ -784,21 +784,21 @@ UE_NET_TEST_FIXTURE(FTestCancelPendingDestroyFixture, TestCancelingCancelPending
 
 	// Filter out object to cause a SubObjectPendingDestroy
 	Server->GetReplicationSystem()->AddToGroup(Server->GetReplicationSystem()->GetNotReplicatedNetObjectGroup(), ServerSubObject->NetRefHandle);
-	Server->PreSendUpdate();
+	Server->NetUpdate();
 	Server->PostSendUpdate();
 
 	// Remove object from filter to cause object to end up in CancelPendingDestroy...
 	Server->GetReplicationSystem()->RemoveFromGroup(Server->GetReplicationSystem()->GetNotReplicatedNetObjectGroup(), ServerSubObject->NetRefHandle);
-	Server->PreSendUpdate();
+	Server->NetUpdate();
 	Server->PostSendUpdate();
 
 	// ... and cancel that thought! We want the subobject to be destroyed after all.
 	Server->GetReplicationSystem()->AddToGroup(Server->GetReplicationSystem()->GetNotReplicatedNetObjectGroup(), ServerSubObject->NetRefHandle);
-	Server->PreSendUpdate();
+	Server->NetUpdate();
 	Server->PostSendUpdate();
 
 	// Write and send a packet
-	Server->PreSendUpdate();
+	Server->NetUpdate();
 	Server->SendAndDeliverTo(Client, DeliverPacket);
 	Server->PostSendUpdate();
 
@@ -818,7 +818,7 @@ UE_NET_TEST_FIXTURE(FTestCancelPendingDestroyFixture, TestCancelPendingDestroyFr
 	UTestReplicatedIrisObject* ServerObject = Server->CreateObject(UTestReplicatedIrisObject::FComponents());
 
 	// Write and send packet
-	Server->PreSendUpdate();
+	Server->NetUpdate();
 	Server->SendAndDeliverTo(Client0, DeliverPacket);
 	Server->SendAndDeliverTo(Client1, DeliverPacket);
 	Server->PostSendUpdate();
@@ -840,14 +840,14 @@ UE_NET_TEST_FIXTURE(FTestCancelPendingDestroyFixture, TestCancelPendingDestroyFr
 	Server->ReplicationSystem->SetGroupFilterStatus(ExclusionGroupHandle, Client0->ConnectionIdOnServer, ENetFilterStatus::Disallow);
 	Server->ReplicationSystem->SetGroupFilterStatus(ExclusionGroupHandle, Client1->ConnectionIdOnServer, ENetFilterStatus::Allow);
 	
-	Server->PreSendUpdate();
+	Server->NetUpdate();
 	Server->SendTo(Client0);
 	Server->PostSendUpdate();
 
 	// Mark dirty
 	ServerObject->IntA = 3;
 
-	Server->PreSendUpdate();
+	Server->NetUpdate();
 	Server->SendAndDeliverTo(Client1, DeliverPacket);
 	Server->PostSendUpdate();
 
@@ -858,7 +858,7 @@ UE_NET_TEST_FIXTURE(FTestCancelPendingDestroyFixture, TestCancelPendingDestroyFr
 	Server->DeliverTo(Client0, false);
 
 	// Write and send a packet
-	Server->PreSendUpdate();
+	Server->NetUpdate();
 	Server->SendAndDeliverTo(Client0, DeliverPacket);
 	Server->SendAndDeliverTo(Client1, DeliverPacket);
 	Server->PostSendUpdate();
@@ -884,24 +884,24 @@ UE_NET_TEST_FIXTURE(FTestCancelPendingDestroyFixture, TestCancelingCancelPending
 	UTestReplicatedIrisObject* ServerSubObject = Server->CreateSubObject(ServerObject->NetRefHandle, UTestReplicatedIrisObject::FComponents());
 
 	// Write and send packet
-	Server->PreSendUpdate();
+	Server->NetUpdate();
 	Server->SendAndDeliverTo(Client, DeliverPacket);
 	Server->PostSendUpdate();
 
 	// Filter out object and write packet to cause subobject to end up in WaitOnDestroyConfirmation
 	Server->GetReplicationSystem()->AddToGroup(Server->GetReplicationSystem()->GetNotReplicatedNetObjectGroup(), ServerSubObject->NetRefHandle);
-	Server->PreSendUpdate();
+	Server->NetUpdate();
 	Server->SendUpdate(Client->ConnectionIdOnServer);
 	Server->PostSendUpdate();
 
 	// Remove object from filter to cause subobject to end up in CancelPendingDestroy...
 	Server->GetReplicationSystem()->RemoveFromGroup(Server->GetReplicationSystem()->GetNotReplicatedNetObjectGroup(), ServerSubObject->NetRefHandle);
-	Server->PreSendUpdate();
+	Server->NetUpdate();
 	Server->PostSendUpdate();
 
 	// ... and cancel that thought! We want the subobject to be destroyed after all.
 	Server->GetReplicationSystem()->AddToGroup(Server->GetReplicationSystem()->GetNotReplicatedNetObjectGroup(), ServerSubObject->NetRefHandle);
-	Server->PreSendUpdate();
+	Server->NetUpdate();
 	Server->PostSendUpdate();
 
 	// Cause packet loss on the subobject destroy packet
@@ -911,7 +911,7 @@ UE_NET_TEST_FIXTURE(FTestCancelPendingDestroyFixture, TestCancelingCancelPending
 	UE_NET_ASSERT_NE(Client->GetReplicationBridge()->GetReplicatedObject(ServerSubObject->NetRefHandle), nullptr);
 
 	// Write and send a packet
-	Server->PreSendUpdate();
+	Server->NetUpdate();
 	Server->SendAndDeliverTo(Client, DeliverPacket);
 	Server->PostSendUpdate();
 
@@ -929,24 +929,24 @@ UE_NET_TEST_FIXTURE(FTestCancelPendingDestroyFixture, TestCancelingCancelPending
 	UTestReplicatedIrisObject* ServerSubObject = Server->CreateSubObject(ServerObject->NetRefHandle, UTestReplicatedIrisObject::FComponents());
 
 	// Write and send packet
-	Server->PreSendUpdate();
+	Server->NetUpdate();
 	Server->SendAndDeliverTo(Client, DeliverPacket);
 	Server->PostSendUpdate();
 
 	// Filter out subobject and write packet to cause subobject to end up in WaitOnDestroyConfirmation
 	Server->GetReplicationSystem()->AddToGroup(Server->GetReplicationSystem()->GetNotReplicatedNetObjectGroup(), ServerSubObject->NetRefHandle);
-	Server->PreSendUpdate();
+	Server->NetUpdate();
 	Server->SendUpdate(Client->ConnectionIdOnServer);
 	Server->PostSendUpdate();
 
 	// Remove subobject from filter to cause subobject to end up in CancelPendingDestroy...
 	Server->GetReplicationSystem()->RemoveFromGroup(Server->GetReplicationSystem()->GetNotReplicatedNetObjectGroup(), ServerSubObject->NetRefHandle);
-	Server->PreSendUpdate();
+	Server->NetUpdate();
 	Server->PostSendUpdate();
 
 	// ... and cancel that thought! We want the subobject to be destroyed after all.
 	Server->GetReplicationSystem()->AddToGroup(Server->GetReplicationSystem()->GetNotReplicatedNetObjectGroup(), ServerSubObject->NetRefHandle);
-	Server->PreSendUpdate();
+	Server->NetUpdate();
 	Server->PostSendUpdate();
 
 	// Deliver destroy packet
@@ -956,7 +956,7 @@ UE_NET_TEST_FIXTURE(FTestCancelPendingDestroyFixture, TestCancelingCancelPending
 	UE_NET_ASSERT_EQ(Client->GetReplicationBridge()->GetReplicatedObject(ServerSubObject->NetRefHandle), nullptr);
 
 	// Write and send a packet
-	Server->PreSendUpdate();
+	Server->NetUpdate();
 	Server->SendAndDeliverTo(Client, DeliverPacket);
 	Server->PostSendUpdate();
 
@@ -974,7 +974,7 @@ UE_NET_TEST_FIXTURE(FTestCancelPendingDestroyFixture, TestCancelRootObjectDestro
 	UTestReplicatedIrisObject* ServerSubObject = Server->CreateSubObject(ServerObject->NetRefHandle, UTestReplicatedIrisObject::FComponents());
 
 	// Write and send packet
-	Server->PreSendUpdate();
+	Server->NetUpdate();
 	Server->SendAndDeliverTo(Client, DeliverPacket);
 	Server->PostSendUpdate();
 
@@ -983,21 +983,21 @@ UE_NET_TEST_FIXTURE(FTestCancelPendingDestroyFixture, TestCancelRootObjectDestro
 
 	// Filter out subobject to cause a SubObjectPendingDestroy
 	Server->GetReplicationSystem()->AddToGroup(Server->GetReplicationSystem()->GetNotReplicatedNetObjectGroup(), ServerSubObject->NetRefHandle);
-	Server->PreSendUpdate();
+	Server->NetUpdate();
 	Server->PostSendUpdate();
 
 	// Filter out root object to cause subobject to move to PendingDestroy instead
 	Server->GetReplicationSystem()->AddToGroup(Server->GetReplicationSystem()->GetNotReplicatedNetObjectGroup(), ServerObject->NetRefHandle);
-	Server->PreSendUpdate();
+	Server->NetUpdate();
 	Server->PostSendUpdate();
 
 	// ... and cancel that thought! We don't want the object to be destroyed after all.
 	Server->GetReplicationSystem()->RemoveFromGroup(Server->GetReplicationSystem()->GetNotReplicatedNetObjectGroup(), ServerObject->NetRefHandle);
-	Server->PreSendUpdate();
+	Server->NetUpdate();
 	Server->PostSendUpdate();
 
 	// Write and send a packet
-	Server->PreSendUpdate();
+	Server->NetUpdate();
 	Server->SendAndDeliverTo(Client, DeliverPacket);
 	Server->PostSendUpdate();
 
@@ -1032,7 +1032,7 @@ UE_NET_TEST_FIXTURE(FTestCancelPendingDestroyFixture, TestCancelPendingDestroyDu
 	Server->ReplicationSystem->AddExclusionFilterGroup(ExclusionGroupHandle);
 	Server->ReplicationSystem->SetGroupFilterStatus(ExclusionGroupHandle, Client->ConnectionIdOnServer, ENetFilterStatus::Disallow);
 
-	Server->PreSendUpdate();
+	Server->NetUpdate();
 	Server->SendUpdate(Client->ConnectionIdOnServer);
 	Server->PostSendUpdate();
 
@@ -1044,7 +1044,7 @@ UE_NET_TEST_FIXTURE(FTestCancelPendingDestroyFixture, TestCancelPendingDestroyDu
 		FEnsureScope EnsureScope;
 
 		Server->ReplicationSystem->SetGroupFilterStatus(ExclusionGroupHandle, Client->ConnectionIdOnServer, ENetFilterStatus::Allow);
-		Server->PreSendUpdate();
+		Server->NetUpdate();
 		Server->PostSendUpdate();
 
 		UE_NET_ASSERT_EQ(EnsureScope.GetCount(), 0);
diff --git a/Engine/Plugins/Runtime/ReplicationSystemTestPlugin/Source/Private/Tests/ReplicationSystem/TestDependentObjects.cpp b/Engine/Plugins/Runtime/ReplicationSystemTestPlugin/Source/Private/Tests/ReplicationSystem/TestDependentObjects.cpp
--- a/Engine/Plugins/Runtime/ReplicationSystemTestPlugin/Source/Private/Tests/ReplicationSystem/TestDependentObjects.cpp
+++ b/Engine/Plugins/Runtime/ReplicationSystemTestPlugin/Source/Private/Tests/ReplicationSystem/TestDependentObjects.cpp
@@ -111,7 +111,7 @@ UE_NET_TEST_FIXTURE(FReplicationSystemServerClientTestFixture, TestDependentObje
 	ServerObject->IntA = 1;
 
 	// Send and deliver packet
-	Server->PreSendUpdate();
+	Server->NetUpdate();
 	Server->SendAndDeliverTo(Client, true);
 	Server->PostSendUpdate();
 
@@ -119,7 +119,7 @@ UE_NET_TEST_FIXTURE(FReplicationSystemServerClientTestFixture, TestDependentObje
 	ServerDependentObject->IntA = 1;
 
 	// Send and deliver packet
-	Server->PreSendUpdate();
+	Server->NetUpdate();
 	Server->SendAndDeliverTo(Client, true);
 	Server->PostSendUpdate();
 
@@ -134,7 +134,7 @@ UE_NET_TEST_FIXTURE(FReplicationSystemServerClientTestFixture, TestDependentObje
 	ServerDependentObject->IntA = 2;
 
 	// Send and do not deliver packet
-	Server->PreSendUpdate();
+	Server->NetUpdate();
 	Server->SendAndDeliverTo(Client, false);
 	Server->PostSendUpdate();
 
@@ -142,7 +142,7 @@ UE_NET_TEST_FIXTURE(FReplicationSystemServerClientTestFixture, TestDependentObje
 	UE_NET_ASSERT_NE(ServerDependentObject->IntA, ClientDependentObject->IntA);
 
 	// Send and deliver packet
-	Server->PreSendUpdate();
+	Server->NetUpdate();
 	Server->SendAndDeliverTo(Client, true);
 	Server->PostSendUpdate();
 
@@ -173,7 +173,7 @@ UE_NET_TEST_FIXTURE(FReplicationSystemServerClientTestFixture, TestDependentObje
 	Bridge->AddDependentObject(ServerObject->NetRefHandle, ServerDependentObject->NetRefHandle);
 
 	// Send and deliver packet
-	Server->PreSendUpdate();
+	Server->NetUpdate();
 	Server->SendAndDeliverTo(Client, true);
 	Server->PostSendUpdate();
 
@@ -187,7 +187,7 @@ UE_NET_TEST_FIXTURE(FReplicationSystemServerClientTestFixture, TestDependentObje
 	ReplicationSystem->SetGroupFilterStatus(FilterGroup, ENetFilterStatus::Allow);
 
 	// Send and deliver packet
-	Server->PreSendUpdate();
+	Server->NetUpdate();
 	Server->SendAndDeliverTo(Client, true);
 	Server->PostSendUpdate();
 
@@ -227,7 +227,7 @@ UE_NET_TEST_FIXTURE(FReplicationSystemServerClientTestFixture, TestChainedDepend
 	Bridge->AddDependentObject(ServerDependentObject0->NetRefHandle, ServerDependentObject1->NetRefHandle);
 
 	// Send and deliver packet
-	Server->PreSendUpdate();
+	Server->NetUpdate();
 	Server->SendAndDeliverTo(Client, true);
 	Server->PostSendUpdate();
 
@@ -244,7 +244,7 @@ UE_NET_TEST_FIXTURE(FReplicationSystemServerClientTestFixture, TestChainedDepend
 	ReplicationSystem->SetGroupFilterStatus(FilterGroup, ENetFilterStatus::Allow);
 
 	// Send and deliver packet
-	Server->PreSendUpdate();
+	Server->NetUpdate();
 	Server->SendAndDeliverTo(Client, true);
 	Server->PostSendUpdate();
 
@@ -274,7 +274,7 @@ UE_NET_TEST_FIXTURE(FReplicationSystemServerClientTestFixture, TestDependentObje
 	UTestReplicatedIrisObject* ServerDependentObject = Server->CreateObject(Params);
 
 	// Send and deliver packet
-	Server->PreSendUpdate();
+	Server->NetUpdate();
 	Server->SendAndDeliverTo(Client, true);
 	Server->PostSendUpdate();
 
@@ -297,7 +297,7 @@ UE_NET_TEST_FIXTURE(FReplicationSystemServerClientTestFixture, TestDependentObje
 		ServerDependentObject->IntA += 1;
 
 		// Send and deliver packet
-		Server->PreSendUpdate();
+		Server->NetUpdate();
 		Server->SendAndDeliverTo(Client, true);
 		Server->PostSendUpdate();
 
@@ -319,7 +319,7 @@ UE_NET_TEST_FIXTURE(FReplicationSystemServerClientTestFixture, TestDependentObje
 	ServerObject->IntA += 1;
 
 	// Send and deliver packet
-	Server->PreSendUpdate();
+	Server->NetUpdate();
 	Server->SendAndDeliverTo(Client, true);
 	Server->PostSendUpdate();
 
@@ -356,7 +356,7 @@ UE_NET_TEST_FIXTURE(FReplicationSystemServerClientTestFixture, TestDependentObje
 	Server->ReplicationBridge->AddDependentObject(ServerObject1->NetRefHandle, ServerDependentObject->NetRefHandle);
 
 	// Send and deliver packet, All objects are polled and replicated
-	Server->PreSendUpdate();
+	Server->NetUpdate();
 	Server->SendAndDeliverTo(Client, true);
 	Server->PostSendUpdate();
 
@@ -383,7 +383,7 @@ UE_NET_TEST_FIXTURE(FReplicationSystemServerClientTestFixture, TestDependentObje
 		ServerDependentObject->IntA = CurrentFrame;
 
 		// Send and deliver packet
-		Server->PreSendUpdate();
+		Server->NetUpdate();
 		Server->SendAndDeliverTo(Client, true);
 		Server->PostSendUpdate();
 
@@ -425,7 +425,7 @@ UE_NET_TEST_FIXTURE(FReplicationSystemServerClientTestFixture, TestDependentObje
 	Server->ReplicationBridge->AddDependentObject(ServerObject->NetRefHandle, ServerDependentObject->NetRefHandle);
 
 	// Send and deliver packet
-	Server->PreSendUpdate();
+	Server->NetUpdate();
 	Server->SendAndDeliverTo(Client, true);
 	Server->PostSendUpdate();
 
@@ -442,7 +442,7 @@ UE_NET_TEST_FIXTURE(FReplicationSystemServerClientTestFixture, TestDependentObje
 		ServerDependentObject->IntA += 1;
 
 		// Send and deliver packet
-		Server->PreSendUpdate();
+		Server->NetUpdate();
 		Server->SendAndDeliverTo(Client, true);
 		Server->PostSendUpdate();
 
@@ -459,7 +459,7 @@ UE_NET_TEST_FIXTURE(FReplicationSystemServerClientTestFixture, TestDependentObje
 	ReplicationSystem->ForceNetUpdate(ServerDependentObject->NetRefHandle);
 
 	// Send and deliver packet, we expect dependent object to have replicated 
-	Server->PreSendUpdate();
+	Server->NetUpdate();
 	Server->SendAndDeliverTo(Client, true);
 	Server->PostSendUpdate();
 
@@ -475,7 +475,7 @@ UE_NET_TEST_FIXTURE(FReplicationSystemServerClientTestFixture, TestDependentObje
 	ReplicationSystem->ForceNetUpdate(ServerObject->NetRefHandle);
 
 	// Send and deliver packet
-	Server->PreSendUpdate();
+	Server->NetUpdate();
 	Server->SendAndDeliverTo(Client, true);
 	Server->PostSendUpdate();
 
@@ -506,7 +506,7 @@ UE_NET_TEST_FIXTURE(FReplicationSystemServerClientTestFixture, TestDependentObje
 	UReplicatedSubObjectOrderObject::RepOrderCounter = 0U;
 
 	// Send and deliver packet
-	Server->PreSendUpdate();
+	Server->NetUpdate();
 	Server->SendAndDeliverTo(Client, true);
 	Server->PostSendUpdate();
 
@@ -525,7 +525,7 @@ UE_NET_TEST_FIXTURE(FReplicationSystemServerClientTestFixture, TestDependentObje
 	ServerDependentObject->IntA = 1;
 
 	// Send and deliver packet
-	Server->PreSendUpdate();
+	Server->NetUpdate();
 	Server->SendAndDeliverTo(Client, true);
 	Server->PostSendUpdate();
 
@@ -555,7 +555,7 @@ UE_NET_TEST_FIXTURE(FReplicationSystemServerClientTestFixture, TestDependentObje
 	UReplicatedSubObjectOrderObject::RepOrderCounter = 0U;
 
 	// Send and deliver packet
-	Server->PreSendUpdate();
+	Server->NetUpdate();
 	Server->SendAndDeliverTo(Client, true);
 	Server->PostSendUpdate();
 
@@ -574,7 +574,7 @@ UE_NET_TEST_FIXTURE(FReplicationSystemServerClientTestFixture, TestDependentObje
 	ServerDependentObject->IntA = 1;
 
 	// Send and deliver packet
-	Server->PreSendUpdate();
+	Server->NetUpdate();
 	Server->SendAndDeliverTo(Client, true);
 	Server->PostSendUpdate();
 
@@ -604,7 +604,7 @@ UE_NET_TEST_FIXTURE(FReplicationSystemServerClientTestFixture, TestDependentObje
 	UReplicatedSubObjectOrderObject::RepOrderCounter = 0U;
 
 	// Send and deliver packet
-	Server->PreSendUpdate();
+	Server->NetUpdate();
 	Server->SendAndDeliverTo(Client, true);
 	Server->PostSendUpdate();
 
@@ -623,7 +623,7 @@ UE_NET_TEST_FIXTURE(FReplicationSystemServerClientTestFixture, TestDependentObje
 	ServerDependentObject->IntA = 1;
 
 	// Send and deliver packet
-	Server->PreSendUpdate();
+	Server->NetUpdate();
 	Server->SendAndDeliverTo(Client, true);
 	Server->PostSendUpdate();
 
@@ -656,7 +656,7 @@ UE_NET_TEST_FIXTURE(FReplicationSystemServerClientTestFixture, TestNestedDepende
 	UReplicatedSubObjectOrderObject::RepOrderCounter = 0U;
 
 	// Send and deliver packet
-	Server->PreSendUpdate();
+	Server->NetUpdate();
 	Server->SendAndDeliverTo(Client, true);
 	Server->PostSendUpdate();
 
@@ -679,7 +679,7 @@ UE_NET_TEST_FIXTURE(FReplicationSystemServerClientTestFixture, TestNestedDepende
 	ServerNestedDependentObject->IntA = 1;
 
 	// Send and deliver packet
-	Server->PreSendUpdate();
+	Server->NetUpdate();
 	Server->SendAndDeliverTo(Client, true);
 	Server->PostSendUpdate();
 
@@ -713,7 +713,7 @@ UE_NET_TEST_FIXTURE(FReplicationSystemServerClientTestFixture, TestNestedDepende
 	UReplicatedSubObjectOrderObject::RepOrderCounter = 0U;
 
 	// Send and deliver packet
-	Server->PreSendUpdate();
+	Server->NetUpdate();
 	Server->SendAndDeliverTo(Client, true);
 	Server->PostSendUpdate();
 
@@ -736,7 +736,7 @@ UE_NET_TEST_FIXTURE(FReplicationSystemServerClientTestFixture, TestNestedDepende
 	ServerNestedDependentObject->IntA = 1;
 
 	// Send and deliver packet
-	Server->PreSendUpdate();
+	Server->NetUpdate();
 	Server->SendAndDeliverTo(Client, true);
 	Server->PostSendUpdate();
 
@@ -767,7 +767,7 @@ UE_NET_TEST_FIXTURE(FReplicationSystemServerClientTestFixture, TestLateAddedNest
 	UReplicatedSubObjectOrderObject::RepOrderCounter = 0U;
 
 	// Send and deliver packet
-	Server->PreSendUpdate();
+	Server->NetUpdate();
 	Server->SendAndDeliverTo(Client, true);
 	Server->PostSendUpdate();
 
@@ -790,7 +790,7 @@ UE_NET_TEST_FIXTURE(FReplicationSystemServerClientTestFixture, TestLateAddedNest
 	ServerObject->IntA = 1;
 
 	// Send and deliver packet
-	Server->PreSendUpdate();
+	Server->NetUpdate();
 	Server->SendAndDeliverTo(Client, true);
 	Server->PostSendUpdate();
 
@@ -829,7 +829,7 @@ UE_NET_TEST_FIXTURE(FReplicationSystemServerClientTestFixture, TestLateAddedNest
 	Bridge->AddDependentObject(ServerDependentObject->NetRefHandle, ServerSharedDependentObject->NetRefHandle);
 
 	// Send data
-	Server->PreSendUpdate();
+	Server->NetUpdate();
 	Server->SendAndDeliverTo(Client, true);
 	Server->PostSendUpdate();
 
@@ -918,7 +918,7 @@ UE_NET_TEST_FIXTURE(FReplicationSystemServerClientTestFixture, TestDependentObje
 	UReplicatedSubObjectOrderObject::RepOrderCounter = 0U;
 
 	// Send and deliver packet
-	Server->PreSendUpdate();
+	Server->NetUpdate();
 	Server->SendAndDeliverTo(Client, true);
 	Server->PostSendUpdate();
 
@@ -937,7 +937,7 @@ UE_NET_TEST_FIXTURE(FReplicationSystemServerClientTestFixture, TestDependentObje
 	ServerDependentObject->IntA = 1;
 
 	// Send and deliver packet
-	Server->PreSendUpdate();
+	Server->NetUpdate();
 	Server->SendAndDeliverTo(Client, true);
 	Server->PostSendUpdate();
 
diff --git a/Engine/Plugins/Runtime/ReplicationSystemTestPlugin/Source/Private/Tests/ReplicationSystem/TestFlushBeforeDestroy.cpp b/Engine/Plugins/Runtime/ReplicationSystemTestPlugin/Source/Private/Tests/ReplicationSystem/TestFlushBeforeDestroy.cpp
--- a/Engine/Plugins/Runtime/ReplicationSystemTestPlugin/Source/Private/Tests/ReplicationSystem/TestFlushBeforeDestroy.cpp
+++ b/Engine/Plugins/Runtime/ReplicationSystemTestPlugin/Source/Private/Tests/ReplicationSystem/TestFlushBeforeDestroy.cpp
@@ -25,7 +25,7 @@ UE_NET_TEST_FIXTURE(FTestFlushBeforeDestroyFixture, TestReliableAttachmentFlushe
 	UReplicatedTestObject* ServerObject = Server->CreateObject(0, 0);
 	FNetRefHandle ObjectHandle = ServerObject->NetRefHandle;
 
-	Server->PreSendUpdate();
+	Server->NetUpdate();
 	Server->SendAndDeliverTo(Client, DeliverPacket);
 	Server->PostSendUpdate();
 
@@ -43,7 +43,7 @@ UE_NET_TEST_FIXTURE(FTestFlushBeforeDestroyFixture, TestReliableAttachmentFlushe
 	Server->DestroyObject(ServerObject);
 
 	// Deliver a packet, this should flush the object and deliver the attachment
-	Server->PreSendUpdate();
+	Server->NetUpdate();
 	Server->SendAndDeliverTo(Client, DeliverPacket);
 	Server->PostSendUpdate();
 
@@ -51,7 +51,7 @@ UE_NET_TEST_FIXTURE(FTestFlushBeforeDestroyFixture, TestReliableAttachmentFlushe
 	UE_NET_ASSERT_EQ(ClientMockNetObjectAttachmentHandler->GetFunctionCallCounts().OnNetBlobReceived, 1U);
 
 	// Deliver a packet. Should destroy the object on the client unless that was done
-	Server->PreSendUpdate();
+	Server->NetUpdate();
 	Server->SendAndDeliverTo(Client, DeliverPacket);
 	Server->PostSendUpdate();
 
@@ -255,7 +255,7 @@ UE_NET_TEST_FIXTURE(FTestFlushBeforeDestroyFixture, TestReliableAttachmentFlushe
 	UReplicatedTestObject* ServerObject = Server->CreateObject(0, 0);
 	FNetRefHandle ObjectHandle = ServerObject->NetRefHandle;
 
-	Server->PreSendUpdate();
+	Server->NetUpdate();
 	Server->SendAndDeliverTo(Client, DeliverPacket);
 	Server->PostSendUpdate();
 
@@ -270,7 +270,7 @@ UE_NET_TEST_FIXTURE(FTestFlushBeforeDestroyFixture, TestReliableAttachmentFlushe
 	}
 
 	// Setup a situation where we have reliable data in flight when the object is destroyed
-	Server->PreSendUpdate();
+	Server->NetUpdate();
 	Server->SendTo(Client);
 	Server->PostSendUpdate();
 
@@ -281,7 +281,7 @@ UE_NET_TEST_FIXTURE(FTestFlushBeforeDestroyFixture, TestReliableAttachmentFlushe
 	Server->DeliverTo(Client, false);
 
 	// Deliver a packet, this should flush the object and deliver the attachment
-	Server->PreSendUpdate();
+	Server->NetUpdate();
 	Server->SendAndDeliverTo(Client, DeliverPacket);
 	Server->PostSendUpdate();
 
@@ -289,7 +289,7 @@ UE_NET_TEST_FIXTURE(FTestFlushBeforeDestroyFixture, TestReliableAttachmentFlushe
 	UE_NET_ASSERT_EQ(ClientMockNetObjectAttachmentHandler->GetFunctionCallCounts().OnNetBlobReceived, 1U);
 
 	// Deliver a packet. Should destroy the object on the client unless that was done
-	Server->PreSendUpdate();
+	Server->NetUpdate();
 	Server->SendAndDeliverTo(Client, DeliverPacket);
 	Server->PostSendUpdate();
 
@@ -324,7 +324,7 @@ UE_NET_TEST_FIXTURE(FTestFlushBeforeDestroyFixture, TestReliableAttachmentFlushe
 	// Setup a situation where we have creation info in flight when the object is destroyed
 
 	// Send creation info
-	Server->PreSendUpdate();
+	Server->NetUpdate();
 	Server->SendTo(Client, TEXT("WaitOnCreateConfirmation"));
 	Server->PostSendUpdate();
 
@@ -340,7 +340,7 @@ UE_NET_TEST_FIXTURE(FTestFlushBeforeDestroyFixture, TestReliableAttachmentFlushe
 	Server->DestroyObject(ServerObject);
 
 	// Previously this would issue a flush and send the attachment data even though creation was not yet confirmed leading to a client disconnect..
-	Server->PreSendUpdate();
+	Server->NetUpdate();
 	const bool bDataWasSent = Server->SendTo(Client, TEXT("State should still be WaitOnCreateConfirmation"));
 	Server->PostSendUpdate();
 
@@ -378,7 +378,7 @@ UE_NET_TEST_FIXTURE(FTestFlushBeforeDestroyFixture, TestReliableAttachmentFlushe
 	// Setup a situation where we have creation info in flight when the object is destroyed
 
 	// Send creation info
-	Server->PreSendUpdate();
+	Server->NetUpdate();
 	Server->SendTo(Client, TEXT("WaitOnCreateConfirmation"));
 	Server->PostSendUpdate();
 
@@ -394,7 +394,7 @@ UE_NET_TEST_FIXTURE(FTestFlushBeforeDestroyFixture, TestReliableAttachmentFlushe
 	Server->DestroyObject(ServerObject);
 
 	// Previously this would issue a flush and send data before creation is confirmed.
-	Server->PreSendUpdate();
+	Server->NetUpdate();
 	const bool bDataWasSentInError = Server->SendTo(Client, TEXT("State should still be WaitOnCreateConfirmation"));
 	Server->PostSendUpdate();
 	
@@ -412,7 +412,7 @@ UE_NET_TEST_FIXTURE(FTestFlushBeforeDestroyFixture, TestReliableAttachmentFlushe
 	}
 
 	// Expected to write the attachment
-	Server->PreSendUpdate();
+	Server->NetUpdate();
 	Server->SendAndDeliverTo(Client, DeliverPacket, TEXT("WaitOnFlush"));
 	Server->PostSendUpdate();
 
@@ -420,7 +420,7 @@ UE_NET_TEST_FIXTURE(FTestFlushBeforeDestroyFixture, TestReliableAttachmentFlushe
 	UE_NET_ASSERT_EQ(ClientMockNetObjectAttachmentHandler->GetFunctionCallCounts().OnNetBlobReceived, 1U);
 
 	// Expected to destroy the object
-	Server->PreSendUpdate();
+	Server->NetUpdate();
 	Server->SendAndDeliverTo(Client, DeliverPacket, TEXT("Destroy"));
 	Server->PostSendUpdate();
 
@@ -439,7 +439,7 @@ UE_NET_TEST_FIXTURE(FTestFlushBeforeDestroyFixture, TestReliableAttachmentFlushe
 	// Setup a situation where we have creation info in flight when the object is destroyed
 
 	// Send creation info
-	Server->PreSendUpdate();
+	Server->NetUpdate();
 	Server->SendTo(Client, TEXT("WaitOnCreateConfirmation"));
 	Server->PostSendUpdate();
 
@@ -455,7 +455,7 @@ UE_NET_TEST_FIXTURE(FTestFlushBeforeDestroyFixture, TestReliableAttachmentFlushe
 	Server->DestroyObject(ServerObject, EEndReplicationFlags::Destroy | EEndReplicationFlags::Flush);
 
 	// Previously this would issue a flush and send data before creation is confirmed.
-	Server->PreSendUpdate();
+	Server->NetUpdate();
 	const bool bDataWasSentInError = Server->SendTo(Client, TEXT("State should still be WaitOnCreateConfirmation"));
 	Server->PostSendUpdate();
 	
@@ -473,7 +473,7 @@ UE_NET_TEST_FIXTURE(FTestFlushBeforeDestroyFixture, TestReliableAttachmentFlushe
 	}
 
 	// Expected to create object and send attachment
-	Server->PreSendUpdate();
+	Server->NetUpdate();
 	Server->SendAndDeliverTo(Client, DeliverPacket, TEXT("CreateResend"));
 	Server->PostSendUpdate();
 
@@ -484,12 +484,12 @@ UE_NET_TEST_FIXTURE(FTestFlushBeforeDestroyFixture, TestReliableAttachmentFlushe
 	UE_NET_ASSERT_EQ(ClientMockNetObjectAttachmentHandler->GetFunctionCallCounts().OnNetBlobReceived, 1U);
 
 	// Expected to write the attachment
-	Server->PreSendUpdate();
+	Server->NetUpdate();
 	Server->SendAndDeliverTo(Client, DeliverPacket, TEXT("WaitOnFlush"));
 	Server->PostSendUpdate();
 
 	// Expected to destroy the object
-	Server->PreSendUpdate();
+	Server->NetUpdate();
 	Server->SendAndDeliverTo(Client, DeliverPacket, TEXT("Destroy"));
 	Server->PostSendUpdate();
 
@@ -508,7 +508,7 @@ UE_NET_TEST_FIXTURE(FTestFlushBeforeDestroyFixture, TestReliableAttachmentSubObj
 	UReplicatedTestObject* ServerSubObject = Server->CreateSubObject(ObjectHandle, 0, 0);
 	FNetRefHandle SubObjectHandle = ServerSubObject->NetRefHandle;
 	
-	Server->PreSendUpdate();
+	Server->NetUpdate();
 	Server->SendAndDeliverTo(Client, DeliverPacket);
 	Server->PostSendUpdate();
 
@@ -526,7 +526,7 @@ UE_NET_TEST_FIXTURE(FTestFlushBeforeDestroyFixture, TestReliableAttachmentSubObj
 	Server->DestroyObject(ServerSubObject);
 
 	// Deliver a packet, this should flush the object and deliver the attachment
-	Server->PreSendUpdate();
+	Server->NetUpdate();
 	Server->SendAndDeliverTo(Client, DeliverPacket);
 	Server->PostSendUpdate();
 
@@ -534,7 +534,7 @@ UE_NET_TEST_FIXTURE(FTestFlushBeforeDestroyFixture, TestReliableAttachmentSubObj
 	UE_NET_ASSERT_EQ(ClientMockNetObjectAttachmentHandler->GetFunctionCallCounts().OnNetBlobReceived, 1U);
 
 	// Deliver a packet. Should destroy the object on the client unless that was done
-	Server->PreSendUpdate();
+	Server->NetUpdate();
 	Server->SendAndDeliverTo(Client, DeliverPacket);
 	Server->PostSendUpdate();
 
@@ -552,7 +552,7 @@ UE_NET_TEST_FIXTURE(FTestFlushBeforeDestroyFixture, TestReliableAttachmentSubObj
 	UReplicatedTestObject* ServerSubObject = Server->CreateSubObject(ObjectHandle, 0, 0);
 	FNetRefHandle SubObjectHandle = ServerSubObject->NetRefHandle;
 	
-	Server->PreSendUpdate();
+	Server->NetUpdate();
 	Server->SendAndDeliverTo(Client, DeliverPacket);
 	Server->PostSendUpdate();
 
@@ -570,7 +570,7 @@ UE_NET_TEST_FIXTURE(FTestFlushBeforeDestroyFixture, TestReliableAttachmentSubObj
 	Server->DestroyObject(ServerObject);
 
 	// Deliver a packet, this should flush the object and deliver the attachment
-	Server->PreSendUpdate();
+	Server->NetUpdate();
 	Server->SendAndDeliverTo(Client, DeliverPacket);
 	Server->PostSendUpdate();
 
@@ -578,7 +578,7 @@ UE_NET_TEST_FIXTURE(FTestFlushBeforeDestroyFixture, TestReliableAttachmentSubObj
 	UE_NET_ASSERT_EQ(ClientMockNetObjectAttachmentHandler->GetFunctionCallCounts().OnNetBlobReceived, 1U);
 
 	// Deliver a packet. Should destroy the object on the client unless that was done
-	Server->PreSendUpdate();
+	Server->NetUpdate();
 	Server->SendAndDeliverTo(Client, DeliverPacket);
 	Server->PostSendUpdate();
 
@@ -595,7 +595,7 @@ UE_NET_TEST_FIXTURE(FTestFlushBeforeDestroyFixture, TestStateFlushedBeforeDestro
 	UTestReplicatedIrisObject* ServerObject = Server->CreateObject(0, 0);
 	FNetRefHandle ObjectHandle = ServerObject->NetRefHandle;
 	
-	Server->PreSendUpdate();
+	Server->NetUpdate();
 	Server->SendAndDeliverTo(Client, DeliverPacket);
 	Server->PostSendUpdate();
 
@@ -609,7 +609,7 @@ UE_NET_TEST_FIXTURE(FTestFlushBeforeDestroyFixture, TestStateFlushedBeforeDestro
 	Server->DestroyObject(ServerObject, EEndReplicationFlags::Destroy | EEndReplicationFlags::Flush);
 
 	// Deliver a packet, this should flush the object and deliver the last state
-	Server->PreSendUpdate();
+	Server->NetUpdate();
 	Server->SendAndDeliverTo(Client, DeliverPacket);
 	Server->PostSendUpdate();
 
@@ -622,7 +622,7 @@ UE_NET_TEST_FIXTURE(FTestFlushBeforeDestroyFixture, TestStateFlushedBeforeDestro
 	UE_NET_ASSERT_EQ(ClientObject->IntA, 3);
 	
 	// Deliver a packet.
-	Server->PreSendUpdate();
+	Server->NetUpdate();
 	Server->SendAndDeliverTo(Client, DeliverPacket);
 	Server->PostSendUpdate();
 
@@ -638,7 +638,7 @@ UE_NET_TEST_FIXTURE(FTestFlushBeforeDestroyFixture, TestStateInFlightFlushedBefo
 	UTestReplicatedIrisObject* ServerObject = Server->CreateObject(0, 0);
 	FNetRefHandle ObjectHandle = ServerObject->NetRefHandle;
 	
-	Server->PreSendUpdate();
+	Server->NetUpdate();
 	Server->SendAndDeliverTo(Client, DeliverPacket);
 	Server->PostSendUpdate();
 
@@ -648,7 +648,7 @@ UE_NET_TEST_FIXTURE(FTestFlushBeforeDestroyFixture, TestStateInFlightFlushedBefo
 	// Modify state
 	ServerObject->IntA = 3;
 
-	Server->PreSendUpdate();
+	Server->NetUpdate();
 	Server->SendTo(Client);
 	Server->PostSendUpdate();
 
@@ -662,7 +662,7 @@ UE_NET_TEST_FIXTURE(FTestFlushBeforeDestroyFixture, TestStateInFlightFlushedBefo
 	Server->DeliverTo(Client, false);
 
 	// Deliver a packet, this should flush the object and deliver the complete last state
-	Server->PreSendUpdate();
+	Server->NetUpdate();
 	Server->SendAndDeliverTo(Client, DeliverPacket);
 	Server->PostSendUpdate();
 
@@ -676,7 +676,7 @@ UE_NET_TEST_FIXTURE(FTestFlushBeforeDestroyFixture, TestStateInFlightFlushedBefo
 	UE_NET_ASSERT_EQ(ClientObject->IntB, 4);
 
 	// Deliver a packet. Should destroy the object on the client unless that was done
-	Server->PreSendUpdate();
+	Server->NetUpdate();
 	Server->SendAndDeliverTo(Client, DeliverPacket);
 	Server->PostSendUpdate();
 
@@ -699,7 +699,7 @@ UE_NET_TEST_FIXTURE(FTestFlushBeforeDestroyFixture, TestDroppedPendingTearOffIsC
 	FReplicationSystemTestClient* Client = CreateClient();
 	RegisterNetBlobHandlers(Client);
 
-	Server->PreSendUpdate();
+	Server->NetUpdate();
 	Server->SendAndDeliverTo(Client, DeliverPacket);
 	Server->PostSendUpdate();
 
@@ -719,7 +719,7 @@ UE_NET_TEST_FIXTURE(FTestFlushBeforeDestroyFixture, TestDroppedPendingTearOffIsC
 	Server->ReplicationBridge->EndReplication(ServerObject, EEndReplicationFlags::TearOff);
 
 	// Send packet so that we have creationdata in flight
-	Server->PreSendUpdate();
+	Server->NetUpdate();
 	Server->SendTo(Client);
 	Server->PostSendUpdate();
 
@@ -731,7 +731,7 @@ UE_NET_TEST_FIXTURE(FTestFlushBeforeDestroyFixture, TestDroppedPendingTearOffIsC
 	Server->DeliverTo(Client, false);
 
 	// Deliver a packet, this should flush the object and deliver the attachment
-	Server->PreSendUpdate();
+	Server->NetUpdate();
 	Server->SendAndDeliverTo(Client, DeliverPacket);
 	Server->PostSendUpdate();
 
@@ -739,7 +739,7 @@ UE_NET_TEST_FIXTURE(FTestFlushBeforeDestroyFixture, TestDroppedPendingTearOffIsC
 	UE_NET_ASSERT_NE(ClientMockNetObjectAttachmentHandler->GetFunctionCallCounts().OnNetBlobReceived, 1U);
 
 	// Deliver a packet. Should destroy the object on the client unless that was already done
-	Server->PreSendUpdate();
+	Server->NetUpdate();
 	Server->SendAndDeliverTo(Client, DeliverPacket);
 	Server->PostSendUpdate();
 
@@ -778,14 +778,14 @@ UE_NET_TEST_FIXTURE(FTestFlushBeforeDestroyFixture, TestPendingCreateTearOffIsCa
 	Server->ReplicationBridge->EndReplication(ServerObject, EEndReplicationFlags::TearOff);
 
 	// PreUpdate to update scoping to get the object into the PendingCreate state.
-	Server->PreSendUpdate();
+	Server->NetUpdate();
 	Server->PostSendUpdate();
 
 	// Force destroy object already pending tearoff/flush. DestroyLocalNetHandle will invalidate cached creationinfo.
 	Server->ReplicationBridge->EndReplication(ServerObject, EEndReplicationFlags::Destroy);
 
 	// Send a packet.
-	Server->PreSendUpdate();
+	Server->NetUpdate();
 	Server->SendAndDeliverTo(Client, DeliverPacket);
 	Server->PostSendUpdate();
 
@@ -817,14 +817,14 @@ UE_NET_TEST_FIXTURE(FTestFlushBeforeDestroyFixture, TestPendingCreateTearOffIsNo
 	Server->ReplicationBridge->EndReplication(ServerObject, EEndReplicationFlags::TearOff);
 
 	// PreUpdate to update scoping to get the object into the PendingCreate state.
-	Server->PreSendUpdate();
+	Server->NetUpdate();
 	Server->PostSendUpdate();
 
 	// This should be ignored as we are already pending tear off.
 	Server->ReplicationBridge->EndReplication(ServerObject, EEndReplicationFlags::Destroy);
 
 	// Deliver a packet, this should flush the object and deliver the attachment
-	Server->PreSendUpdate();
+	Server->NetUpdate();
 	Server->SendAndDeliverTo(Client, DeliverPacket);
 	Server->PostSendUpdate();
 
@@ -840,7 +840,7 @@ UE_NET_TEST_FIXTURE(FTestFlushBeforeDestroyFixture, TestDroppedTearOffIsNotCance
 	FReplicationSystemTestClient* Client = CreateClient();
 	RegisterNetBlobHandlers(Client);
 
-	Server->PreSendUpdate();
+	Server->NetUpdate();
 	Server->SendAndDeliverTo(Client, DeliverPacket);
 	Server->PostSendUpdate();
 
@@ -860,7 +860,7 @@ UE_NET_TEST_FIXTURE(FTestFlushBeforeDestroyFixture, TestDroppedTearOffIsNotCance
 	Server->ReplicationBridge->EndReplication(ServerObject, EEndReplicationFlags::TearOff);
 
 	// Send packet so that we have creationdata in flight
-	Server->PreSendUpdate();
+	Server->NetUpdate();
 	Server->SendTo(Client);
 	Server->PostSendUpdate();
 
@@ -871,7 +871,7 @@ UE_NET_TEST_FIXTURE(FTestFlushBeforeDestroyFixture, TestDroppedTearOffIsNotCance
 	Server->DeliverTo(Client, false);
 
 	// Deliver a packet, this should flush the object and deliver the attachment
-	Server->PreSendUpdate();
+	Server->NetUpdate();
 	Server->SendAndDeliverTo(Client, DeliverPacket);
 	Server->PostSendUpdate();
 
@@ -879,7 +879,7 @@ UE_NET_TEST_FIXTURE(FTestFlushBeforeDestroyFixture, TestDroppedTearOffIsNotCance
 	UE_NET_ASSERT_EQ(ClientMockNetObjectAttachmentHandler->GetFunctionCallCounts().OnNetBlobReceived, 1U);
 
 	// Deliver a packet. Should tear off the object on the client
-	Server->PreSendUpdate();
+	Server->NetUpdate();
 	Server->SendAndDeliverTo(Client, DeliverPacket);
 	Server->PostSendUpdate();
 
@@ -897,7 +897,7 @@ UE_NET_TEST_FIXTURE(FTestFlushBeforeDestroyFixture, TestSubObjectStateFlushedBef
 	UTestReplicatedIrisObject* ServerSubObject = Server->CreateSubObject(ObjectHandle, 0, 0);
 	FNetRefHandle SubObjectHandle = ServerSubObject->NetRefHandle;
 
-	Server->PreSendUpdate();
+	Server->NetUpdate();
 	Server->SendAndDeliverTo(Client, DeliverPacket);
 	Server->PostSendUpdate();
 
@@ -912,7 +912,7 @@ UE_NET_TEST_FIXTURE(FTestFlushBeforeDestroyFixture, TestSubObjectStateFlushedBef
 	Server->DestroyObject(ServerObject, EEndReplicationFlags::Destroy | EEndReplicationFlags::Flush);
 
 	// Deliver a packet, this should flush the object and deliver the last state
-	Server->PreSendUpdate();
+	Server->NetUpdate();
 	Server->SendAndDeliverTo(Client, DeliverPacket);
 	Server->PostSendUpdate();
 
@@ -927,7 +927,7 @@ UE_NET_TEST_FIXTURE(FTestFlushBeforeDestroyFixture, TestSubObjectStateFlushedBef
 	UE_NET_ASSERT_EQ(ClientSubObject->IntA, 3);
 	
 	// Deliver a packet.
-	Server->PreSendUpdate();
+	Server->NetUpdate();
 	Server->SendAndDeliverTo(Client, DeliverPacket);
 	Server->PostSendUpdate();
 
@@ -946,7 +946,7 @@ UE_NET_TEST_FIXTURE(FTestFlushBeforeDestroyFixture, TestSubObjectStateFlushedBef
 	UTestReplicatedIrisObject* ServerSubObject = Server->CreateSubObject(ObjectHandle, 0, 0);
 	FNetRefHandle SubObjectHandle = ServerSubObject->NetRefHandle;
 
-	Server->PreSendUpdate();
+	Server->NetUpdate();
 	Server->SendAndDeliverTo(Client, DeliverPacket);
 	Server->PostSendUpdate();
 
@@ -961,7 +961,7 @@ UE_NET_TEST_FIXTURE(FTestFlushBeforeDestroyFixture, TestSubObjectStateFlushedBef
 	Server->DestroyObject(ServerSubObject, EEndReplicationFlags::Destroy | EEndReplicationFlags::Flush);
 
 	// Deliver a packet, this should flush the object and deliver the last state
-	Server->PreSendUpdate();
+	Server->NetUpdate();
 	Server->SendAndDeliverTo(Client, DeliverPacket);
 	Server->PostSendUpdate();
 
@@ -976,7 +976,7 @@ UE_NET_TEST_FIXTURE(FTestFlushBeforeDestroyFixture, TestSubObjectStateFlushedBef
 	UE_NET_ASSERT_EQ(ClientSubObject->IntA, 3);
 	
 	// Deliver a packet.
-	Server->PreSendUpdate();
+	Server->NetUpdate();
 	Server->SendAndDeliverTo(Client, DeliverPacket);
 	Server->PostSendUpdate();
 
@@ -1010,7 +1010,7 @@ UE_NET_TEST_FIXTURE(FTestFlushBeforeDestroyFixture, TestTearOffNewObjectWithReli
 	Server->ReplicationBridge->EndReplication(ServerObject, EEndReplicationFlags::TearOff);
 
 	// Send and deliver packet
-	Server->PreSendUpdate();
+	Server->NetUpdate();
 	Server->SendAndDeliverTo(Client, true);
 	Server->PostSendUpdate();
 
@@ -1024,7 +1024,7 @@ UE_NET_TEST_FIXTURE(FTestFlushBeforeDestroyFixture, TestTearOffNewObjectWithReli
 	UE_NET_ASSERT_EQ(ClientMockNetObjectAttachmentHandler->GetFunctionCallCounts().OnNetBlobReceived, 1U);
 
 	// Send and deliver packet
-	Server->PreSendUpdate();
+	Server->NetUpdate();
 	Server->SendAndDeliverTo(Client, true);
 	Server->PostSendUpdate();
 
@@ -1046,7 +1046,7 @@ UE_NET_TEST_FIXTURE(FTestFlushBeforeDestroyFixture, TestTearOffExistingObjectWit
 	ServerObject->IntA = 1;
 
 	// Send and deliver packet
-	Server->PreSendUpdate();
+	Server->NetUpdate();
 	Server->SendAndDeliverTo(Client, true);
 	Server->PostSendUpdate();
 
@@ -1071,7 +1071,7 @@ UE_NET_TEST_FIXTURE(FTestFlushBeforeDestroyFixture, TestTearOffExistingObjectWit
 	Server->ReplicationBridge->EndReplication(ServerObject, EEndReplicationFlags::TearOff);
 
 	// Send and deliver packet
-	Server->PreSendUpdate();
+	Server->NetUpdate();
 	Server->SendAndDeliverTo(Client, true);
 	Server->PostSendUpdate();
 
@@ -1083,7 +1083,7 @@ UE_NET_TEST_FIXTURE(FTestFlushBeforeDestroyFixture, TestTearOffExistingObjectWit
 	UE_NET_ASSERT_TRUE(Cast<UTestReplicatedIrisObject>(Client->GetReplicationBridge()->GetReplicatedObject(ServerObject->NetRefHandle)) != nullptr);
 
 	// Send and deliver packet
-	Server->PreSendUpdate();
+	Server->NetUpdate();
 	Server->SendAndDeliverTo(Client, true);
 	Server->PostSendUpdate();
 
@@ -1111,7 +1111,7 @@ UE_NET_TEST_FIXTURE(FTestFlushBeforeDestroyFixture, TestImmediateTearOffExisting
 	ServerSubObject->IntA = 1;
 
 	// Send and deliver packet
-	Server->PreSendUpdate();
+	Server->NetUpdate();
 	Server->SendAndDeliverTo(Client, true);
 	Server->PostSendUpdate();
 
@@ -1139,7 +1139,7 @@ UE_NET_TEST_FIXTURE(FTestFlushBeforeDestroyFixture, TestImmediateTearOffExisting
 	Server->ReplicationBridge->EndReplication(ServerObject, EEndReplicationFlags::TearOff);
 
 	// Send and deliver packet
-	Server->PreSendUpdate();
+	Server->NetUpdate();
 	Server->SendAndDeliverTo(Client, true);
 	Server->PostSendUpdate();
 
@@ -1148,7 +1148,7 @@ UE_NET_TEST_FIXTURE(FTestFlushBeforeDestroyFixture, TestImmediateTearOffExisting
 	UE_NET_ASSERT_EQ(ServerSubObject->IntA, ClientSubObjectThatWillBeTornOff->IntA);
 
 	// Send and deliver packet
-	Server->PreSendUpdate();
+	Server->NetUpdate();
 	Server->SendAndDeliverTo(Client, true);
 	Server->PostSendUpdate();
 
@@ -1176,7 +1176,7 @@ UE_NET_TEST_FIXTURE(FTestFlushBeforeDestroyFixture, TestImmediateTearOffWithSubO
 	ServerSubObject->IntA = 1;
 
 	// Send and deliver packet
-	Server->PreSendUpdate();
+	Server->NetUpdate();
 	Server->SendAndDeliverTo(Client, true, TEXT("Create Objects"));
 	Server->PostSendUpdate();
 
@@ -1200,21 +1200,21 @@ UE_NET_TEST_FIXTURE(FTestFlushBeforeDestroyFixture, TestImmediateTearOffWithSubO
 		Server->GetReplicationSystem()->QueueNetObjectAttachment(Client->ConnectionIdOnServer, AttachmentTarget, Attachment);
 	}
 
-	Server->PreSendUpdate();
+	Server->NetUpdate();
 	Server->SendTo(Client, TEXT("State data + Attachment"));
 	Server->PostSendUpdate();
 
 	// Modify the value of object only
 	++ServerObject->IntA ;
 
-	Server->PreSendUpdate();
+	Server->NetUpdate();
 	Server->SendTo(Client, TEXT("State data"));
 	Server->PostSendUpdate();
 
 	// TearOff the object using immediate tear-off
 	Server->ReplicationBridge->EndReplication(ServerObject, EEndReplicationFlags::TearOff);
 
-	Server->PreSendUpdate();
+	Server->NetUpdate();
 	Server->SendTo(Client, TEXT("Tear off"));
 	Server->PostSendUpdate();
 
@@ -1225,7 +1225,7 @@ UE_NET_TEST_FIXTURE(FTestFlushBeforeDestroyFixture, TestImmediateTearOffWithSubO
 	Server->DeliverTo(Client, false);
 
 	// This earlier caused an unwanted state transition
-	Server->PreSendUpdate();
+	Server->NetUpdate();
 	Server->SendTo(Client, TEXT("Packet after tearoff"));
 	Server->PostSendUpdate();
 
@@ -1236,7 +1236,7 @@ UE_NET_TEST_FIXTURE(FTestFlushBeforeDestroyFixture, TestImmediateTearOffWithSubO
 	Server->DeliverTo(Client, true);
 
 	// This should contain resend of lost state
-	Server->PreSendUpdate();
+	Server->NetUpdate();
 	Server->SendAndDeliverTo(Client, true, TEXT("Resending tearoff"));
 	Server->PostSendUpdate();
 
@@ -1263,7 +1263,7 @@ UE_NET_TEST_FIXTURE(FTestFlushBeforeDestroyFixture, TestCancelPendingDestroyWait
 	UTestReplicatedIrisObject* ServerObject = Server->CreateObject(0,0);
 
 	// Send and deliver packet
-	Server->PreSendUpdate();
+	Server->NetUpdate();
 	Server->SendAndDeliverTo(Client0, true, TEXT("Create Objects"));
 	Server->PostSendUpdate();
 
@@ -1280,7 +1280,7 @@ UE_NET_TEST_FIXTURE(FTestFlushBeforeDestroyFixture, TestCancelPendingDestroyWait
 		Server->GetReplicationSystem()->QueueNetObjectAttachment(Client0->ConnectionIdOnServer, AttachmentTarget, Attachment);
 	}
 
-	Server->PreSendUpdate();
+	Server->NetUpdate();
 	Server->SendTo(Client0, TEXT("Attachment"));
 	Server->PostSendUpdate();
 
@@ -1292,7 +1292,7 @@ UE_NET_TEST_FIXTURE(FTestFlushBeforeDestroyFixture, TestCancelPendingDestroyWait
 	Server->ReplicationSystem->SetGroupFilterStatus(ExclusionGroupHandle, Client0->ConnectionIdOnServer, ENetFilterStatus::Disallow);
 	Server->ReplicationSystem->SetGroupFilterStatus(ExclusionGroupHandle, Client1->ConnectionIdOnServer, ENetFilterStatus::Allow);
 
-	Server->PreSendUpdate();
+	Server->NetUpdate();
 	Server->SendTo(Client0, TEXT("Out of scope"));
 	Server->PostSendUpdate();
 
@@ -1300,7 +1300,7 @@ UE_NET_TEST_FIXTURE(FTestFlushBeforeDestroyFixture, TestCancelPendingDestroyWait
 	++ServerObject->IntA ;
 
 	// Trigger poll + propagate of state
-	Server->PreSendUpdate();
+	Server->NetUpdate();
 	Server->PostSendUpdate();
 	
 	// Trigger WaitOnFlush -> Created
@@ -1311,7 +1311,7 @@ UE_NET_TEST_FIXTURE(FTestFlushBeforeDestroyFixture, TestCancelPendingDestroyWait
 	Server->DeliverTo(Client0, false);
 
 	// Do a normal update, should send state changed that occurred while we where in pending flush state
-	Server->PreSendUpdate();
+	Server->NetUpdate();
 	Server->SendAndDeliverTo(Client0, DeliverPacket, TEXT("Expected state"));
 	Server->PostSendUpdate();
 
@@ -1835,7 +1835,7 @@ UE_NET_TEST_FIXTURE(FTestFlushBeforeDestroyFixture, TestTearOffAndFilterOutInSam
 	Server->DestroyObject(ServerSubObject1);
 
 	// Put data in flight
-	Server->PreSendUpdate();
+	Server->NetUpdate();
 	Server->SendTo(Client);
 	Server->PostSendUpdate();
 
@@ -1844,7 +1844,7 @@ UE_NET_TEST_FIXTURE(FTestFlushBeforeDestroyFixture, TestTearOffAndFilterOutInSam
 	Server->ReplicationSystem->AddToGroup(ExclusionGroupHandle, ServerObject->NetRefHandle);
 	Server->ReplicationSystem->AddExclusionFilterGroup(ExclusionGroupHandle);
 
-	Server->PreSendUpdate();
+	Server->NetUpdate();
 	Server->SendTo(Client);
 	Server->PostSendUpdate();
 	UE_NET_ASSERT_EQ(EnsureScope.GetCount(), 0);
@@ -1858,7 +1858,7 @@ UE_NET_TEST_FIXTURE(FTestFlushBeforeDestroyFixture, TestTearOffAndFilterOutInSam
 	// Tear off 
 	Server->ReplicationSystem->TearOffNextUpdate(ServerObject->NetRefHandle);
 
-	Server->PreSendUpdate();
+	Server->NetUpdate();
 	Server->SendTo(Client);
 	Server->PostSendUpdate();
 	UE_NET_ASSERT_EQ(EnsureScope.GetCount(), 0);
@@ -1921,7 +1921,7 @@ UE_NET_TEST_FIXTURE(FTestFlushBeforeDestroyFixture, TestTearOffAndFilterOutInSam
 	Server->DestroyObject(ServerSubObject1);
 
 	// Put data in flight
-	Server->PreSendUpdate();
+	Server->NetUpdate();
 	Server->SendTo(Client);
 	Server->PostSendUpdate();
 
@@ -1930,7 +1930,7 @@ UE_NET_TEST_FIXTURE(FTestFlushBeforeDestroyFixture, TestTearOffAndFilterOutInSam
 	Server->ReplicationSystem->AddToGroup(ExclusionGroupHandle, ServerObject->NetRefHandle);
 	Server->ReplicationSystem->AddExclusionFilterGroup(ExclusionGroupHandle);
 
-	Server->PreSendUpdate();
+	Server->NetUpdate();
 	Server->SendTo(Client);
 	Server->PostSendUpdate();
 	UE_NET_ASSERT_EQ(EnsureScope.GetCount(), 0);
@@ -1938,7 +1938,7 @@ UE_NET_TEST_FIXTURE(FTestFlushBeforeDestroyFixture, TestTearOffAndFilterOutInSam
 	ServerObject->IntA++;
 	ServerSubObject0->IntA++;
 
-	Server->PreSendUpdate();
+	Server->NetUpdate();
 	Server->SendTo(Client);
 	Server->PostSendUpdate();
 	UE_NET_ASSERT_EQ(EnsureScope.GetCount(), 0);
@@ -1951,7 +1951,7 @@ UE_NET_TEST_FIXTURE(FTestFlushBeforeDestroyFixture, TestTearOffAndFilterOutInSam
 	// Tear off 
 	Server->ReplicationSystem->TearOffNextUpdate(ServerObject->NetRefHandle);
 
-	Server->PreSendUpdate();
+	Server->NetUpdate();
 	Server->SendTo(Client);
 	Server->PostSendUpdate();
 	UE_NET_ASSERT_EQ(EnsureScope.GetCount(), 0);
@@ -2019,7 +2019,7 @@ UE_NET_TEST_FIXTURE(FTestFlushBeforeDestroyFixture, TestTearOffOfObjectNotInScop
 	ServerSubObject0->IntA++;
 
 	// Put data in flight
-	Server->PreSendUpdate();
+	Server->NetUpdate();
 	Server->SendTo(Client, TEXT("SubObject1 creation"));
 	Server->PostSendUpdate();
 
@@ -2028,7 +2028,7 @@ UE_NET_TEST_FIXTURE(FTestFlushBeforeDestroyFixture, TestTearOffOfObjectNotInScop
 	Server->ReplicationSystem->AddToGroup(ExclusionGroupHandle, ServerObject->NetRefHandle);
 	Server->ReplicationSystem->AddExclusionFilterGroup(ExclusionGroupHandle);
 
-	Server->PreSendUpdate();
+	Server->NetUpdate();
 	Server->SendTo(Client);
 	Server->PostSendUpdate();
 
@@ -2037,7 +2037,7 @@ UE_NET_TEST_FIXTURE(FTestFlushBeforeDestroyFixture, TestTearOffOfObjectNotInScop
 	// Do not deliver SubObject1 creation
 	Server->DeliverTo(Client, DoNotDeliverPacket);
 
-	Server->PreSendUpdate();
+	Server->NetUpdate();
 	Server->SendTo(Client);
 	Server->PostSendUpdate();
 
@@ -2046,7 +2046,7 @@ UE_NET_TEST_FIXTURE(FTestFlushBeforeDestroyFixture, TestTearOffOfObjectNotInScop
 	// Tear off 
 	Server->ReplicationSystem->TearOffNextUpdate(ServerObject->NetRefHandle);
 
-	Server->PreSendUpdate();
+	Server->NetUpdate();
 	Server->SendTo(Client);
 	Server->PostSendUpdate();
 	UE_NET_ASSERT_EQ(EnsureScope.GetCount(), 0);
@@ -2116,7 +2116,7 @@ UE_NET_TEST_FIXTURE(FTestFlushBeforeDestroyFixture, TestTearOffOfObjectNotInScop
 	ServerSubObject0->IntA++;
 
 	// Put data in flight
-	Server->PreSendUpdate();
+	Server->NetUpdate();
 	Server->SendTo(Client, TEXT("SubObject1 creation"));
 	Server->PostSendUpdate();
 
@@ -2125,7 +2125,7 @@ UE_NET_TEST_FIXTURE(FTestFlushBeforeDestroyFixture, TestTearOffOfObjectNotInScop
 	Server->ReplicationSystem->AddToGroup(ExclusionGroupHandle, ServerObject->NetRefHandle);
 	Server->ReplicationSystem->AddExclusionFilterGroup(ExclusionGroupHandle);
 
-	Server->PreSendUpdate();
+	Server->NetUpdate();
 	Server->SendTo(Client);
 	Server->PostSendUpdate();
 
@@ -2134,14 +2134,14 @@ UE_NET_TEST_FIXTURE(FTestFlushBeforeDestroyFixture, TestTearOffOfObjectNotInScop
 	// Do not deliver SubObject1 creation
 	Server->DeliverTo(Client, DoNotDeliverPacket);
 
-	Server->PreSendUpdate();
+	Server->NetUpdate();
 	Server->SendTo(Client);
 	Server->PostSendUpdate();
 
 	// Deliver SubObject1 creation
 	Server->DeliverTo(Client, DeliverPacket);
 
-	Server->PreSendUpdate();
+	Server->NetUpdate();
 	Server->SendTo(Client);
 	Server->PostSendUpdate();
 
@@ -2150,7 +2150,7 @@ UE_NET_TEST_FIXTURE(FTestFlushBeforeDestroyFixture, TestTearOffOfObjectNotInScop
 	// Tear off 
 	Server->ReplicationSystem->TearOffNextUpdate(ServerObject->NetRefHandle);
 
-	Server->PreSendUpdate();
+	Server->NetUpdate();
 	Server->SendTo(Client);
 	Server->PostSendUpdate();
 	UE_NET_ASSERT_EQ(EnsureScope.GetCount(), 0);
@@ -2214,7 +2214,7 @@ UE_NET_TEST_FIXTURE(FTestFlushBeforeDestroyFixture, TestSubObjectIsNotDestroyedB
 	Server->DestroyObject(ServerSubObject1);
 
 	// Put data in flight
-	Server->PreSendUpdate();
+	Server->NetUpdate();
 	Server->SendTo(Client, TEXT("SubObject1 creation"));
 	Server->PostSendUpdate();
 
@@ -2223,7 +2223,7 @@ UE_NET_TEST_FIXTURE(FTestFlushBeforeDestroyFixture, TestSubObjectIsNotDestroyedB
 	Server->ReplicationSystem->AddToGroup(ExclusionGroupHandle, ServerObject->NetRefHandle);
 	Server->ReplicationSystem->AddExclusionFilterGroup(ExclusionGroupHandle);
 
-	Server->PreSendUpdate();
+	Server->NetUpdate();
 	Server->SendTo(Client);
 	Server->PostSendUpdate();
 
@@ -2232,7 +2232,7 @@ UE_NET_TEST_FIXTURE(FTestFlushBeforeDestroyFixture, TestSubObjectIsNotDestroyedB
 	// Do not deliver SubObject1 creation
 	Server->DeliverTo(Client, DoNotDeliverPacket);
 
-	Server->PreSendUpdate();
+	Server->NetUpdate();
 	Server->SendTo(Client);
 	Server->PostSendUpdate();
 
@@ -2255,7 +2255,7 @@ UE_NET_TEST_FIXTURE(FTestFlushBeforeDestroyFixture, TestSubObjectStateIsNotSentW
 	UTestReplicatedIrisObject* ServerSubObject = Server->CreateSubObject(ObjectHandle);
 	FNetRefHandle SubObjectHandle = ServerSubObject->NetRefHandle;
 
-	Server->PreSendUpdate();
+	Server->NetUpdate();
 	Server->SendAndDeliverTo(Client, DeliverPacket);
 	Server->PostSendUpdate();
 
@@ -2267,7 +2267,7 @@ UE_NET_TEST_FIXTURE(FTestFlushBeforeDestroyFixture, TestSubObjectStateIsNotSentW
 	ServerSubObject->IntDWithOnRep += 3;
 
 	// Update but do not ack/nak the packet just yet
-	Server->PreSendUpdate();
+	Server->NetUpdate();
 	Server->SendTo(Client);
 	Server->PostSendUpdate();
 
diff --git a/Engine/Plugins/Runtime/ReplicationSystemTestPlugin/Source/Private/Tests/ReplicationSystem/TestObjectDeltaSerialization.cpp b/Engine/Plugins/Runtime/ReplicationSystemTestPlugin/Source/Private/Tests/ReplicationSystem/TestObjectDeltaSerialization.cpp
--- a/Engine/Plugins/Runtime/ReplicationSystemTestPlugin/Source/Private/Tests/ReplicationSystem/TestObjectDeltaSerialization.cpp
+++ b/Engine/Plugins/Runtime/ReplicationSystemTestPlugin/Source/Private/Tests/ReplicationSystem/TestObjectDeltaSerialization.cpp
@@ -42,7 +42,7 @@ UE_NET_TEST_FIXTURE(FTestObjectDeltaSerialization, ClientReceivesLatestValuesWit
 	Server->ReplicationSystem->SetDeltaCompressionStatus(ServerObject->NetRefHandle, ENetObjectDeltaCompressionStatus::Allow);
 
 	// Send and deliver packet
-	Server->PreSendUpdate();
+	Server->NetUpdate();
 	Server->SendAndDeliverTo(Client, DeliverPacket);
 	Server->PostSendUpdate();
 
@@ -56,7 +56,7 @@ UE_NET_TEST_FIXTURE(FTestObjectDeltaSerialization, ClientReceivesLatestValuesWit
 	AllowNewBaselineCreation();
 
 	// Send and deliver packet
-	Server->PreSendUpdate();
+	Server->NetUpdate();
 	Server->SendAndDeliverTo(Client, DeliverPacket);
 	Server->PostSendUpdate();
 
@@ -78,7 +78,7 @@ UE_NET_TEST_FIXTURE(FTestObjectDeltaSerialization, ClientReceivesLatestValuesWit
 	Server->ReplicationSystem->SetDeltaCompressionStatus(ServerObject->NetRefHandle, ENetObjectDeltaCompressionStatus::Allow);
 
 	// Send and deliver packet
-	Server->PreSendUpdate();
+	Server->NetUpdate();
 	Server->SendAndDeliverTo(Client, DeliverPacket);
 	Server->PostSendUpdate();
 
@@ -92,7 +92,7 @@ UE_NET_TEST_FIXTURE(FTestObjectDeltaSerialization, ClientReceivesLatestValuesWit
 	AllowNewBaselineCreation();
 
 	// Send and do not deliver packet
-	Server->PreSendUpdate();
+	Server->NetUpdate();
 	Server->SendAndDeliverTo(Client, DoNotDeliverPacket);
 	Server->PostSendUpdate();
 
@@ -102,7 +102,7 @@ UE_NET_TEST_FIXTURE(FTestObjectDeltaSerialization, ClientReceivesLatestValuesWit
 	AllowNewBaselineCreation();
 
 	// Send and deliver packet
-	Server->PreSendUpdate();
+	Server->NetUpdate();
 	Server->SendAndDeliverTo(Client, DeliverPacket);
 	Server->PostSendUpdate();
 
@@ -128,7 +128,7 @@ UE_NET_TEST_FIXTURE(FTestObjectDeltaSerialization, AllClientsReceivesLatestValue
 	Server->ReplicationSystem->SetDeltaCompressionStatus(ServerObject->NetRefHandle, ENetObjectDeltaCompressionStatus::Allow);
 
 	// Send and deliver packets
-	Server->PreSendUpdate();
+	Server->NetUpdate();
 	for (FReplicationSystemTestClient* Client : Clients)
 	{
 		Server->SendAndDeliverTo(Client, DeliverPacket);
@@ -142,7 +142,7 @@ UE_NET_TEST_FIXTURE(FTestObjectDeltaSerialization, AllClientsReceivesLatestValue
 	AllowNewBaselineCreation();
 
 	// Send and deliver packets
-	Server->PreSendUpdate();
+	Server->NetUpdate();
 	for (FReplicationSystemTestClient* Client : Clients)
 	{
 		Server->SendAndDeliverTo(Client, DeliverPacket);
@@ -180,7 +180,7 @@ UE_NET_TEST_FIXTURE(FTestObjectDeltaSerialization, SimulatedIsProperlyReplicated
 	Server->ReplicationSystem->SetReplicationConditionConnectionFilter(ServerObject->NetRefHandle, EReplicationCondition::RoleAutonomous, Clients[0]->ConnectionIdOnServer, true);
 
 	// Send and deliver packet
-	Server->PreSendUpdate();
+	Server->NetUpdate();
 	for (FReplicationSystemTestClient* Client : Clients)
 	{
 		Server->SendAndDeliverTo(Client, DeliverPacket);
@@ -197,7 +197,7 @@ UE_NET_TEST_FIXTURE(FTestObjectDeltaSerialization, SimulatedIsProperlyReplicated
 	ServerObject->ConnectionFilteredComponents[0]->SimulatedOrPhysicsNoReplayInt ^= 11;
 
 	// Send and deliver packet
-	Server->PreSendUpdate();
+	Server->NetUpdate();
 	for (FReplicationSystemTestClient* Client : Clients)
 	{
 		Server->SendAndDeliverTo(Client, DeliverPacket);
@@ -220,7 +220,7 @@ UE_NET_TEST_FIXTURE(FTestObjectDeltaSerialization, SimulatedIsProperlyReplicated
 	DisallowNewBaselineCreation();
 
 	// Send and deliver packet
-	Server->PreSendUpdate();
+	Server->NetUpdate();
 	for (FReplicationSystemTestClient* Client : Clients)
 	{
 		Server->SendAndDeliverTo(Client, DeliverPacket);
@@ -250,7 +250,7 @@ UE_NET_TEST_FIXTURE(FTestObjectDeltaSerialization, SimulatedIsProperlyReplicated
 	Server->ReplicationSystem->SetReplicationConditionConnectionFilter(ServerObject->NetRefHandle, EReplicationCondition::RoleAutonomous, Client->ConnectionIdOnServer, true);
 
 	// Send and deliver packet
-	Server->PreSendUpdate();
+	Server->NetUpdate();
 	Server->SendAndDeliverTo(Client, DeliverPacket);
 	Server->PostSendUpdate();
 
@@ -264,7 +264,7 @@ UE_NET_TEST_FIXTURE(FTestObjectDeltaSerialization, SimulatedIsProperlyReplicated
 	ServerObject->ConnectionFilteredComponents[0]->SimulatedOrPhysicsNoReplayInt ^= 11;
 
 	// Send and deliver packet
-	Server->PreSendUpdate();
+	Server->NetUpdate();
 	Server->SendAndDeliverTo(Client, DeliverPacket);
 	Server->PostSendUpdate();
 
@@ -284,7 +284,7 @@ UE_NET_TEST_FIXTURE(FTestObjectDeltaSerialization, SimulatedIsProperlyReplicated
 	AllowNewBaselineCreation();
 	
 	// Send and do not deliver packet
-	Server->PreSendUpdate();
+	Server->NetUpdate();
 	Server->SendAndDeliverTo(Client, DoNotDeliverPacket);
 	Server->PostSendUpdate();
 
@@ -295,7 +295,7 @@ UE_NET_TEST_FIXTURE(FTestObjectDeltaSerialization, SimulatedIsProperlyReplicated
 	UE_NET_ASSERT_EQ(ClientObject->ConnectionFilteredComponents[0]->SimulatedOrPhysicsNoReplayInt, 0);
 
 	// Send and deliver packet
-	Server->PreSendUpdate();
+	Server->NetUpdate();
 	Server->SendAndDeliverTo(Client, DeliverPacket);
 	Server->PostSendUpdate();
 
@@ -487,7 +487,7 @@ UE_NET_TEST_FIXTURE(FTestObjectDeltaSerialization, ToOwnerStateIsReplicatedToOwn
 	ServerObject->ConnectionFilteredComponents[0]->ToOwnerB = 37;
 
 	// Send and deliver packet
-	Server->PreSendUpdate();
+	Server->NetUpdate();
 	Server->SendAndDeliverTo(Client, DeliverPacket);
 	Server->PostSendUpdate();
 	
@@ -505,7 +505,7 @@ UE_NET_TEST_FIXTURE(FTestObjectDeltaSerialization, ToOwnerStateIsReplicatedToOwn
 	ServerObject->IntA ^= 1;
 
 	// Send and deliver packet
-	Server->PreSendUpdate();
+	Server->NetUpdate();
 	Server->SendAndDeliverTo(Client, DeliverPacket);
 	Server->PostSendUpdate();
 	
@@ -593,7 +593,7 @@ UE_NET_TEST_FIXTURE(FTestObjectDeltaSerialization, InFlightChangesForDisabledCon
 	const int LastReplicatedSimulatedOrPhysicsInt = ServerObject->ConnectionFilteredComponents[0]->SimulatedOrPhysicsInt;
 
 	// Send and deliver packet
-	Server->PreSendUpdate();
+	Server->NetUpdate();
 	Server->SendAndDeliverTo(Client, DeliverPacket);
 	Server->PostSendUpdate();
 	
@@ -611,7 +611,7 @@ UE_NET_TEST_FIXTURE(FTestObjectDeltaSerialization, InFlightChangesForDisabledCon
 	{
 		ServerObject->ConnectionFilteredComponents[0]->SimulatedOrPhysicsInt += 1;
 
-		Server->PreSendUpdate();
+		Server->NetUpdate();
 		Server->SendTo(Client);
 		Server->PostSendUpdate();
 	}
@@ -629,7 +629,7 @@ UE_NET_TEST_FIXTURE(FTestObjectDeltaSerialization, InFlightChangesForDisabledCon
 	ServerObject->IntA ^= 1;
 
 	// Send packet
-	Server->PreSendUpdate();
+	Server->NetUpdate();
 	Server->SendTo(Client);
 	Server->PostSendUpdate();
 	
@@ -678,13 +678,13 @@ UE_NET_TEST_FIXTURE(FTestObjectDeltaSerialization, LostChangesDuringPendingDestr
 	ServerObject->DynamicStateComponents[0]->IntArray.Empty();
 
 	// Fake latency by acking/naking later
-	Server->PreSendUpdate();
+	Server->NetUpdate();
 	Server->SendTo(Client);
 	Server->PostSendUpdate();
 
 	// Filter out object to cause a PendingDestroy
 	Server->GetReplicationSystem()->AddToGroup(Server->GetReplicationSystem()->GetNotReplicatedNetObjectGroup(), ServerObject->NetRefHandle);
-	Server->PreSendUpdate();
+	Server->NetUpdate();
 	Server->SendTo(Client);
 	Server->PostSendUpdate();
 
diff --git a/Engine/Plugins/Runtime/ReplicationSystemTestPlugin/Source/Private/Tests/ReplicationSystem/TestObjectReferences.cpp b/Engine/Plugins/Runtime/ReplicationSystemTestPlugin/Source/Private/Tests/ReplicationSystem/TestObjectReferences.cpp
--- a/Engine/Plugins/Runtime/ReplicationSystemTestPlugin/Source/Private/Tests/ReplicationSystem/TestObjectReferences.cpp
+++ b/Engine/Plugins/Runtime/ReplicationSystemTestPlugin/Source/Private/Tests/ReplicationSystem/TestObjectReferences.cpp
@@ -166,7 +166,7 @@ public:
 		CollectAllCollector.bCollectAll = true;
 
 		// To trigger copy of data so we have valid data in buffers
-		Server->PreSendUpdate();
+		Server->NetUpdate();
 		Server->PostSendUpdate();
 	}
 };
@@ -174,7 +174,7 @@ public:
 UE_NET_TEST_FIXTURE(FTestObjectReferencesFixture, TestReferenceVisitor_NoRef)
 {
 	// To trigger copy of data
-	Server->PreSendUpdate();
+	Server->NetUpdate();
 	Server->PostSendUpdate();
 
 	FNetSerializationContext Context;
@@ -198,7 +198,7 @@ UE_NET_TEST_FIXTURE(FTestObjectReferencesFixture, TestReferenceVisitor_SingleRef
 	TestObject->ObjectRef = TestReferences[0];
 
 	// To trigger copy of data
-	Server->PreSendUpdate();
+	Server->NetUpdate();
 	Server->PostSendUpdate();
 	
 	FReferenceCollector& Collector = CollectOnlyValidCollector;
@@ -216,7 +216,7 @@ UE_NET_TEST_FIXTURE(FTestObjectReferencesFixture, TestReferenceVisitor_ObjectRef
 	TestObject->ObjectRef = TestReferences[0];
 
 	// To trigger copy of data
-	Server->PreSendUpdate();
+	Server->NetUpdate();
 	Server->PostSendUpdate();
 
 	FReferenceCollector& Collector = CollectOnlyValidCollector;
@@ -234,7 +234,7 @@ UE_NET_TEST_FIXTURE(FTestObjectReferencesFixture, TestReferenceVisitor_StructWit
 	TestObject->StructWithRef.ObjectRef = TestReferences[0];
 
 	// To trigger copy of data
-	Server->PreSendUpdate();
+	Server->NetUpdate();
 	Server->PostSendUpdate();
 	
 	FReferenceCollector& Collector = CollectOnlyValidCollector;
@@ -252,7 +252,7 @@ UE_NET_TEST_FIXTURE(FTestObjectReferencesFixture, TestReferenceVisitor_StructWit
 	TestObject->StructWithRefCArray.Ref_CArray[0] = TestReferences[0];
 
 	// To trigger copy of data
-	Server->PreSendUpdate();
+	Server->NetUpdate();
 	Server->PostSendUpdate();
 	
 	FReferenceCollector& Collector = CollectOnlyValidCollector;
@@ -272,7 +272,7 @@ UE_NET_TEST_FIXTURE(FTestObjectReferencesFixture, TestReferenceVisitor_StructWit
 	TestObject->StructWithRefCArray.Ref_CArray[2] = TestReferences[2];
 
 	// To trigger copy of data
-	Server->PreSendUpdate();
+	Server->NetUpdate();
 	Server->PostSendUpdate();
 	
 	FReferenceCollector& Collector = CollectOnlyValidCollector;
@@ -292,7 +292,7 @@ UE_NET_TEST_FIXTURE(FTestObjectReferencesFixture, TestReferenceVisitor_StructWit
 	TestObject->StructWithRefTArray.Ref_TArray.Push(TestReferences[0]);
 
 	// To trigger copy of data
-	Server->PreSendUpdate();
+	Server->NetUpdate();
 	Server->PostSendUpdate();
 	
 	FReferenceCollector& Collector = CollectOnlyValidCollector;
@@ -312,7 +312,7 @@ UE_NET_TEST_FIXTURE(FTestObjectReferencesFixture, TestReferenceVisitor_StructWit
 	TestObject->StructWithRefTArray.Ref_TArray.Push(TestReferences[2]);
 
 	// To trigger copy of data
-	Server->PreSendUpdate();
+	Server->NetUpdate();
 	Server->PostSendUpdate();
 	
 	FReferenceCollector& Collector = CollectOnlyValidCollector;
@@ -332,7 +332,7 @@ UE_NET_TEST_FIXTURE(FTestObjectReferencesFixture, TestReferenceVisitor_RefCArray
 	TestObject->StructWithRefCArray.Ref_CArray[0] = TestReferences[0];
 
 	// To trigger copy of data
-	Server->PreSendUpdate();
+	Server->NetUpdate();
 	Server->PostSendUpdate();
 	
 	FReferenceCollector& Collector = CollectOnlyValidCollector;
@@ -352,7 +352,7 @@ UE_NET_TEST_FIXTURE(FTestObjectReferencesFixture, TestReferenceVisitor_RefCArray
 	TestObject->Ref_CArray[2] = TestReferences[2];
 
 	// To trigger copy of data
-	Server->PreSendUpdate();
+	Server->NetUpdate();
 	Server->PostSendUpdate();
 	
 	FReferenceCollector& Collector = CollectOnlyValidCollector;
@@ -372,7 +372,7 @@ UE_NET_TEST_FIXTURE(FTestObjectReferencesFixture, TestReferenceVisitor_RefTArray
 	TestObject->Ref_TArray.Push(TestReferences[0]);
 
 	// To trigger copy of data
-	Server->PreSendUpdate();
+	Server->NetUpdate();
 	Server->PostSendUpdate();
 	
 	FReferenceCollector& Collector = CollectOnlyValidCollector;
@@ -392,7 +392,7 @@ UE_NET_TEST_FIXTURE(FTestObjectReferencesFixture, TestReferenceVisitor_RefTArray
 	TestObject->Ref_TArray.Push(TestReferences[2]);
 
 	// To trigger copy of data
-	Server->PreSendUpdate();
+	Server->NetUpdate();
 	Server->PostSendUpdate();
 	
 	FReferenceCollector& Collector = CollectOnlyValidCollector;
@@ -412,7 +412,7 @@ UE_NET_TEST_FIXTURE(FTestObjectReferencesFixture, TestReferenceVisitor_StructWit
 	TestObject->StructWithRef_CArray[0].ObjectRef = TestReferences[0];
 
 	// To trigger copy of data
-	Server->PreSendUpdate();
+	Server->NetUpdate();
 	Server->PostSendUpdate();
 	
 	FReferenceCollector& Collector = CollectOnlyValidCollector;
@@ -432,7 +432,7 @@ UE_NET_TEST_FIXTURE(FTestObjectReferencesFixture, TestReferenceVisitor_StructWit
 	TestObject->StructWithRef_CArray[2].ObjectRef = TestReferences[2];
 
 	// To trigger copy of data
-	Server->PreSendUpdate();
+	Server->NetUpdate();
 	Server->PostSendUpdate();
 
 	FReferenceCollector& Collector = CollectOnlyValidCollector;
@@ -453,7 +453,7 @@ UE_NET_TEST_FIXTURE(FTestObjectReferencesFixture, TestReferenceVisitor_StructWit
 	TestObject->StructWithRef_TArray[0].ObjectRef = TestReferences[0];
 
 	// To trigger copy of data
-	Server->PreSendUpdate();
+	Server->NetUpdate();
 	Server->PostSendUpdate();
 	
 	FReferenceCollector& Collector = CollectOnlyValidCollector;
@@ -474,7 +474,7 @@ UE_NET_TEST_FIXTURE(FTestObjectReferencesFixture, TestReferenceVisitor_StructWit
 	TestObject->StructWithRef_TArray[2].ObjectRef = TestReferences[2];
 
 	// To trigger copy of data
-	Server->PreSendUpdate();
+	Server->NetUpdate();
 	Server->PostSendUpdate();
 
 	FReferenceCollector& Collector = CollectOnlyValidCollector;
@@ -494,7 +494,7 @@ UE_NET_TEST_FIXTURE(FTestObjectReferencesFixture, TestStructWithNestedRefCArray_
 	TestObject->TestStructWithNestedRefCArray_CArray[0].StructWithRef_CArray[0].Ref_CArray[0] = TestReferences[0];
 
 	// To trigger copy of data
-	Server->PreSendUpdate();
+	Server->NetUpdate();
 	Server->PostSendUpdate();
 	
 	FReferenceCollector& Collector = CollectOnlyValidCollector;
@@ -514,7 +514,7 @@ UE_NET_TEST_FIXTURE(FTestObjectReferencesFixture, TestStructWithNestedRefCArray_
 	TestObject->TestStructWithNestedRefCArray_CArray[1].StructWithRef_CArray[1].Ref_CArray[0] = TestReferences[2];
 
 	// To trigger copy of data
-	Server->PreSendUpdate();
+	Server->NetUpdate();
 	Server->PostSendUpdate();
 	
 	FReferenceCollector& Collector = CollectOnlyValidCollector;
@@ -537,7 +537,7 @@ UE_NET_TEST_FIXTURE(FTestObjectReferencesFixture, TestStructWithNestedRefTArray_
 	TestObject->TestStructWithNestedRefTArray_TArray[0].StructWithRef_TArray[0].Ref_TArray[0] = TestReferences[0];
 
 	// To trigger copy of data
-	Server->PreSendUpdate();
+	Server->NetUpdate();
 	Server->PostSendUpdate();
 	
 	FReferenceCollector& Collector = CollectOnlyValidCollector;
@@ -564,7 +564,7 @@ UE_NET_TEST_FIXTURE(FTestObjectReferencesFixture, TestStructWithNestedRefTArray_
 	TestObject->TestStructWithNestedRefTArray_TArray[1].StructWithRef_TArray[1].Ref_TArray[0] = TestReferences[2];
 
 	// To trigger copy of data
-	Server->PreSendUpdate();
+	Server->NetUpdate();
 	Server->PostSendUpdate();
 	
 	FReferenceCollector& Collector = CollectOnlyValidCollector;
@@ -628,7 +628,7 @@ UE_NET_TEST_FIXTURE(FTestObjectReferencesFixture, TestStructWithNestedRefTArray_
 	TestObject->TestStructWithNestedRefCArray_CArray[1].StructWithRef_CArray[1].Ref_CArray[0] = TestReferences[CurrentReferenceIndex++];
 
 	// To trigger copy of data
-	Server->PreSendUpdate();
+	Server->NetUpdate();
 	Server->PostSendUpdate();
 	
 	FReferenceCollector& Collector = CollectOnlyValidCollector;
@@ -696,7 +696,7 @@ UE_NET_TEST_FIXTURE(FTestObjectReferencesFixture, TestVisitAllDirty)
 	FNetBitArrayView ChangeMask = MakeNetBitArrayView(ChangeMaskStorage);
 
 	// To trigger copy of data
-	Server->PreSendUpdate();
+	Server->NetUpdate();
 	Server->PostSendUpdate();
 		
 	// Collect with empty changemask
@@ -774,7 +774,7 @@ UE_NET_TEST_FIXTURE(FReplicationSystemServerClientTestFixture, TestInlinedSubObj
 	UE_NET_ASSERT_TRUE((Object->CreatedSubObjectRef->HasAnyFlags(RF_DefaultSubObject) || Object->IsDefaultSubobject()));
 
 	// Replicate object
-	Server->PreSendUpdate();
+	Server->NetUpdate();
 	Server->SendAndDeliverTo(Client, true);
 	Server->PostSendUpdate();
 
@@ -789,7 +789,7 @@ UE_NET_TEST_FIXTURE(FReplicationSystemServerClientTestFixture, TestInlinedSubObj
 	Object->ObjectRef = Object->CreatedSubObjectRef;
 
 	// Replicate
-	Server->PreSendUpdate();
+	Server->NetUpdate();
 	Server->SendAndDeliverTo(Client, true);
 	Server->PostSendUpdate();
 
@@ -815,7 +815,7 @@ UE_NET_TEST_FIXTURE(FReplicationSystemServerClientTestFixture, TestExternalSubOb
 	UTestObjectReferences_TestClassWithDefaultSubObject* ExternalSubObject = Server->CreateSubObject<UTestObjectReferences_TestClassWithDefaultSubObject>(Object->NetRefHandle);
 
 	// Replicate object
-	Server->PreSendUpdate();
+	Server->NetUpdate();
 	Server->SendAndDeliverTo(Client, true);
 	Server->PostSendUpdate();
 
@@ -832,7 +832,7 @@ UE_NET_TEST_FIXTURE(FReplicationSystemServerClientTestFixture, TestExternalSubOb
 	Object->ObjectRef = ExternalSubObject;
 
 	// Replicate
-	Server->PreSendUpdate();
+	Server->NetUpdate();
 	Server->SendAndDeliverTo(Client, true);
 	Server->PostSendUpdate();
 
@@ -866,7 +866,7 @@ UE_NET_TEST_FIXTURE(FReplicationSystemServerClientTestFixture, TestPartialResolv
 	Server->ReplicationSystem->AddToGroup(Server->GetReplicationSystem()->GetNotReplicatedNetObjectGroup(), ServerReferencedObjectC->NetRefHandle);
 	
 	// Replicate
-	Server->PreSendUpdate();
+	Server->NetUpdate();
 	Server->SendAndDeliverTo(Client, true);
 	Server->PostSendUpdate();
 
@@ -886,7 +886,7 @@ UE_NET_TEST_FIXTURE(FReplicationSystemServerClientTestFixture, TestPartialResolv
 	Server->ReplicationSystem->RemoveFromGroup(Server->GetReplicationSystem()->GetNotReplicatedNetObjectGroup(), ServerReferencedObjectA->NetRefHandle);
 
 	// Replicate
-	Server->PreSendUpdate();
+	Server->NetUpdate();
 	Server->SendAndDeliverTo(Client, true);
 	Server->PostSendUpdate();
 
@@ -899,7 +899,7 @@ UE_NET_TEST_FIXTURE(FReplicationSystemServerClientTestFixture, TestPartialResolv
 	Server->ReplicationSystem->RemoveFromGroup(Server->GetReplicationSystem()->GetNotReplicatedNetObjectGroup(), ServerReferencedObjectB->NetRefHandle);
 
 	// Replicate
-	Server->PreSendUpdate();
+	Server->NetUpdate();
 	Server->SendAndDeliverTo(Client, true);
 	Server->PostSendUpdate();
 
@@ -912,7 +912,7 @@ UE_NET_TEST_FIXTURE(FReplicationSystemServerClientTestFixture, TestPartialResolv
 	Server->ReplicationSystem->RemoveFromGroup(Server->GetReplicationSystem()->GetNotReplicatedNetObjectGroup(), ServerReferencedObjectC->NetRefHandle);
 
 	// Replicate
-	Server->PreSendUpdate();
+	Server->NetUpdate();
 	Server->SendAndDeliverTo(Client, true);
 	Server->PostSendUpdate();	
 
@@ -961,7 +961,7 @@ UE_NET_TEST_FIXTURE(FReplicationSystemServerClientTestFixture, TestOutOfOrderRes
 	Server->ReplicationSystem->AddToGroup(Server->GetReplicationSystem()->GetNotReplicatedNetObjectGroup(), ServerReferencedObjectC->NetRefHandle);
 	
 	// Replicate
-	Server->PreSendUpdate();
+	Server->NetUpdate();
 	Server->SendAndDeliverTo(Client, true);
 	Server->PostSendUpdate();
 
@@ -983,7 +983,7 @@ UE_NET_TEST_FIXTURE(FReplicationSystemServerClientTestFixture, TestOutOfOrderRes
 	Server->ReplicationSystem->RemoveFromGroup(Server->GetReplicationSystem()->GetNotReplicatedNetObjectGroup(), ServerReferencedObjectA->NetRefHandle);
 
 	// Replicate
-	Server->PreSendUpdate();
+	Server->NetUpdate();
 	Server->SendAndDeliverTo(Client, true);
 	Server->PostSendUpdate();
 
@@ -1000,7 +1000,7 @@ UE_NET_TEST_FIXTURE(FReplicationSystemServerClientTestFixture, TestOutOfOrderRes
 	Server->ReplicationSystem->RemoveFromGroup(Server->GetReplicationSystem()->GetNotReplicatedNetObjectGroup(), ServerReferencedObjectB->NetRefHandle);
 
 	// Replicate
-	Server->PreSendUpdate();
+	Server->NetUpdate();
 	Server->SendAndDeliverTo(Client, true);
 	Server->PostSendUpdate();
 
@@ -1013,7 +1013,7 @@ UE_NET_TEST_FIXTURE(FReplicationSystemServerClientTestFixture, TestOutOfOrderRes
 	Server->ReplicationSystem->RemoveFromGroup(Server->GetReplicationSystem()->GetNotReplicatedNetObjectGroup(), ServerReferencedObjectC->NetRefHandle);
 
 	// Replicate
-	Server->PreSendUpdate();
+	Server->NetUpdate();
 	Server->SendAndDeliverTo(Client, true);
 	Server->PostSendUpdate();	
 
@@ -1056,7 +1056,7 @@ UE_NET_TEST_FIXTURE(FReplicationSystemServerClientTestFixture, TestOutOfOrderRes
 	Server->ReplicationSystem->AddToGroup(Server->GetReplicationSystem()->GetNotReplicatedNetObjectGroup(), ServerReferencedObjectC->NetRefHandle);
 	
 	// Replicate
-	Server->PreSendUpdate();
+	Server->NetUpdate();
 	Server->SendAndDeliverTo(Client, true);
 	Server->PostSendUpdate();
 
@@ -1078,7 +1078,7 @@ UE_NET_TEST_FIXTURE(FReplicationSystemServerClientTestFixture, TestOutOfOrderRes
 	Server->ReplicationSystem->RemoveFromGroup(Server->GetReplicationSystem()->GetNotReplicatedNetObjectGroup(), ServerReferencedObjectA->NetRefHandle);
 
 	// Replicate
-	Server->PreSendUpdate();
+	Server->NetUpdate();
 	Server->SendAndDeliverTo(Client, true);
 	Server->PostSendUpdate();
 
@@ -1091,7 +1091,7 @@ UE_NET_TEST_FIXTURE(FReplicationSystemServerClientTestFixture, TestOutOfOrderRes
 	Server->ReplicationSystem->RemoveFromGroup(Server->GetReplicationSystem()->GetNotReplicatedNetObjectGroup(), ServerReferencedObjectB->NetRefHandle);
 
 	// Replicate
-	Server->PreSendUpdate();
+	Server->NetUpdate();
 	Server->SendAndDeliverTo(Client, true);
 	Server->PostSendUpdate();
 
@@ -1104,7 +1104,7 @@ UE_NET_TEST_FIXTURE(FReplicationSystemServerClientTestFixture, TestOutOfOrderRes
 	Server->ReplicationSystem->RemoveFromGroup(Server->GetReplicationSystem()->GetNotReplicatedNetObjectGroup(), ServerReferencedObjectC->NetRefHandle);
 
 	// Replicate
-	Server->PreSendUpdate();
+	Server->NetUpdate();
 	Server->SendAndDeliverTo(Client, true);
 	Server->PostSendUpdate();	
 
@@ -1151,7 +1151,7 @@ UE_NET_TEST_FIXTURE(FReplicationSystemServerClientTestFixture, TestOutOfOrderRes
 	Server->ReplicationSystem->AddToGroup(Server->GetReplicationSystem()->GetNotReplicatedNetObjectGroup(), ServerReferencedObjectC->NetRefHandle);
 	
 	// Replicate
-	Server->PreSendUpdate();
+	Server->NetUpdate();
 	Server->SendAndDeliverTo(Client, true);
 	Server->PostSendUpdate();
 
@@ -1180,7 +1180,7 @@ UE_NET_TEST_FIXTURE(FReplicationSystemServerClientTestFixture, TestOutOfOrderRes
 	Server->ReplicationSystem->RemoveFromGroup(Server->GetReplicationSystem()->GetNotReplicatedNetObjectGroup(), ServerReferencedObjectA->NetRefHandle);
 
 	// Replicate
-	Server->PreSendUpdate();
+	Server->NetUpdate();
 	Server->SendAndDeliverTo(Client, true);
 	Server->PostSendUpdate();
 
@@ -1193,7 +1193,7 @@ UE_NET_TEST_FIXTURE(FReplicationSystemServerClientTestFixture, TestOutOfOrderRes
 	Server->ReplicationSystem->RemoveFromGroup(Server->GetReplicationSystem()->GetNotReplicatedNetObjectGroup(), ServerReferencedObjectB->NetRefHandle);
 
 	// Replicate
-	Server->PreSendUpdate();
+	Server->NetUpdate();
 	Server->SendAndDeliverTo(Client, true);
 	Server->PostSendUpdate();
 
@@ -1206,7 +1206,7 @@ UE_NET_TEST_FIXTURE(FReplicationSystemServerClientTestFixture, TestOutOfOrderRes
 	Server->ReplicationSystem->RemoveFromGroup(Server->GetReplicationSystem()->GetNotReplicatedNetObjectGroup(), ServerReferencedObjectC->NetRefHandle);
 
 	// Replicate
-	Server->PreSendUpdate();
+	Server->NetUpdate();
 	Server->SendAndDeliverTo(Client, true);
 	Server->PostSendUpdate();	
 
@@ -1232,7 +1232,7 @@ UE_NET_TEST_FIXTURE(FReplicationSystemServerClientTestFixture, TestRemovedRefere
 	ServerObject->StructWithRef.ObjectRef = ServerReferencedObjectA;
 
 	// Replicate
-	Server->PreSendUpdate();
+	Server->NetUpdate();
 	Server->SendAndDeliverTo(Client, true);
 	Server->PostSendUpdate();
 
@@ -1247,7 +1247,7 @@ UE_NET_TEST_FIXTURE(FReplicationSystemServerClientTestFixture, TestRemovedRefere
 	Server->ReplicationSystem->AddToGroup(Server->GetReplicationSystem()->GetNotReplicatedNetObjectGroup(), ServerReferencedObjectA->NetRefHandle);
 
 	// Replicate to make sure object is destroyed on client
-	Server->PreSendUpdate();
+	Server->NetUpdate();
 	Server->SendAndDeliverTo(Client, true);
 	Server->PostSendUpdate();
 
@@ -1255,7 +1255,7 @@ UE_NET_TEST_FIXTURE(FReplicationSystemServerClientTestFixture, TestRemovedRefere
 	ServerObject->StructWithRef.IntValue ^= 1;
 
 	// Replicate
-	Server->PreSendUpdate();
+	Server->NetUpdate();
 	Server->SendAndDeliverTo(Client, true);
 	Server->PostSendUpdate();
 
@@ -1266,12 +1266,12 @@ UE_NET_TEST_FIXTURE(FReplicationSystemServerClientTestFixture, TestRemovedRefere
 	Server->ReplicationSystem->RemoveFromGroup(Server->GetReplicationSystem()->GetNotReplicatedNetObjectGroup(), ServerReferencedObjectA->NetRefHandle);
 
 	// Replicate
-	Server->PreSendUpdate();
+	Server->NetUpdate();
 	Server->SendAndDeliverTo(Client, true);
 	Server->PostSendUpdate();
 
 	// Client will not send anything, but need object reference tracking update.
-	Client->PreSendUpdate();
+	Client->NetUpdate();
 	Client->PostSendUpdate();
 
 	// Verify that we can now resolve the referenced object again
@@ -1295,12 +1295,12 @@ UE_NET_TEST_FIXTURE(FReplicationSystemServerClientTestFixture, TestRemovedRefere
 	ServerObject->StructWithRef_CArray[0].ObjectRef = ServerReferencedObjectA;
 
 	// Replicate
-	Server->PreSendUpdate();
+	Server->NetUpdate();
 	Server->SendAndDeliverTo(Client, true);
 	Server->PostSendUpdate();
 
 	// Client will not send anything, but need object reference tracking update.
-	Client->PreSendUpdate();
+	Client->NetUpdate();
 	Client->PostSendUpdate();
 
 	// Verify that object has been spawned on client
@@ -1315,7 +1315,7 @@ UE_NET_TEST_FIXTURE(FReplicationSystemServerClientTestFixture, TestRemovedRefere
 	Server->ReplicationSystem->AddToGroup(Server->GetReplicationSystem()->GetNotReplicatedNetObjectGroup(), ServerReferencedObjectA->NetRefHandle);
 
 	// Replicate to make sure object is destroyed on client
-	Server->PreSendUpdate();
+	Server->NetUpdate();
 	Server->SendAndDeliverTo(Client, true);
 	Server->PostSendUpdate();
 
@@ -1324,7 +1324,7 @@ UE_NET_TEST_FIXTURE(FReplicationSystemServerClientTestFixture, TestRemovedRefere
 	ServerObject->StructWithRef_CArray[0].IntValue ^= 1;
 
 	// Replicate
-	Server->PreSendUpdate();
+	Server->NetUpdate();
 	Server->SendAndDeliverTo(Client, true);
 	Server->PostSendUpdate();
 
@@ -1336,12 +1336,12 @@ UE_NET_TEST_FIXTURE(FReplicationSystemServerClientTestFixture, TestRemovedRefere
 	Server->ReplicationSystem->RemoveFromGroup(Server->GetReplicationSystem()->GetNotReplicatedNetObjectGroup(), ServerReferencedObjectA->NetRefHandle);
 
 	// Replicate
-	Server->PreSendUpdate();
+	Server->NetUpdate();
 	Server->SendAndDeliverTo(Client, true);
 	Server->PostSendUpdate();
 
 	// Client will not send anything, but need object reference tracking update.
-	Client->PreSendUpdate();
+	Client->NetUpdate();
 	Client->PostSendUpdate();
 
 	// Verify that we can now resolve the referenced object again
@@ -1382,7 +1382,7 @@ UE_NET_TEST_FIXTURE(FReplicationSystemServerClientTestFixture, TestRemovingRefer
 	ServerTArray.Add(Item);
 
 	// Replicate
-	Server->PreSendUpdate();
+	Server->NetUpdate();
 	Server->SendAndDeliverTo(Client, true);
 	Server->PostSendUpdate();
 
@@ -1403,12 +1403,12 @@ UE_NET_TEST_FIXTURE(FReplicationSystemServerClientTestFixture, TestRemovingRefer
 	Server->ReplicationSystem->AddToGroup(Server->GetReplicationSystem()->GetNotReplicatedNetObjectGroup(), ServerReferencedObjectD->NetRefHandle);
 
 	// Replicate to make sure objects are destroyed on client
-	Server->PreSendUpdate();
+	Server->NetUpdate();
 	Server->SendAndDeliverTo(Client, true);
 	Server->PostSendUpdate();
 
 	// Client will not send anything, but need object reference tracking update.
-	Client->PreSendUpdate();
+	Client->NetUpdate();
 	Client->PostSendUpdate();
 
 	// Change the int values of the structs to force the object references to be re-resolved (normally GC would null them out).
@@ -1418,12 +1418,12 @@ UE_NET_TEST_FIXTURE(FReplicationSystemServerClientTestFixture, TestRemovingRefer
 	ServerTArray[3].IntValue ^= 1;
 
 	// Replicate
-	Server->PreSendUpdate();
+	Server->NetUpdate();
 	Server->SendAndDeliverTo(Client, true);
 	Server->PostSendUpdate();
 
 	// Client will not send anything, but need object reference tracking update.
-	Client->PreSendUpdate();
+	Client->NetUpdate();
 	Client->PostSendUpdate();
 
 	// Verify that objects A, B and D can no longer be resolved, while C still can be resolved.
@@ -1438,12 +1438,12 @@ UE_NET_TEST_FIXTURE(FReplicationSystemServerClientTestFixture, TestRemovingRefer
 	Server->ReplicationSystem->RemoveFromGroup(Server->GetReplicationSystem()->GetNotReplicatedNetObjectGroup(), ServerReferencedObjectD->NetRefHandle);
 
 	// Replicate
-	Server->PreSendUpdate();
+	Server->NetUpdate();
 	Server->SendAndDeliverTo(Client, true);
 	Server->PostSendUpdate();
 
 	// Client will not send anything, but need object reference tracking update.
-	Client->PreSendUpdate();
+	Client->NetUpdate();
 	Client->PostSendUpdate();
 
 	// Verify that we can now resolve all the referenced objects again
@@ -1469,12 +1469,12 @@ UE_NET_TEST_FIXTURE(FReplicationSystemServerClientTestFixture, TestTearOffDoesIn
 	ServerObject->StructWithRef.ObjectRef = ServerReferencedObjectA;
 
 	// Replicate
-	Server->PreSendUpdate();
+	Server->NetUpdate();
 	Server->SendAndDeliverTo(Client, true);
 	Server->PostSendUpdate();
 
 	// Client will not send anything, but need object reference tracking update.
-	Client->PreSendUpdate();
+	Client->NetUpdate();
 	Client->PostSendUpdate();
 
 	// Verify that object has been spawned on client
@@ -1488,24 +1488,24 @@ UE_NET_TEST_FIXTURE(FReplicationSystemServerClientTestFixture, TestTearOffDoesIn
 	Server->ReplicationBridge->EndReplication(ServerReferencedObjectA, EEndReplicationFlags::TearOff);
 
 	// Replicate
-	Server->PreSendUpdate();
+	Server->NetUpdate();
 	Server->SendAndDeliverTo(Client, true);
 	Server->PostSendUpdate();
 
 	// Client will not send anything, but need object reference tracking update.
-	Client->PreSendUpdate();
+	Client->NetUpdate();
 	Client->PostSendUpdate();
 
 	// Change the int value to force the reference in the struct to be re-resolved
 	ServerObject->StructWithRef.IntValue ^= 1;
 
 	// Replicate
-	Server->PreSendUpdate();
+	Server->NetUpdate();
 	Server->SendAndDeliverTo(Client, true);
 	Server->PostSendUpdate();
 
 	// Client will not send anything, but need object reference tracking update.
-	Client->PreSendUpdate();
+	Client->NetUpdate();
 	Client->PostSendUpdate();
 
 	// Verify that object cannot be resolved
diff --git a/Engine/Plugins/Runtime/ReplicationSystemTestPlugin/Source/Private/Tests/ReplicationSystem/TestObjectSplitting.cpp b/Engine/Plugins/Runtime/ReplicationSystemTestPlugin/Source/Private/Tests/ReplicationSystem/TestObjectSplitting.cpp
--- a/Engine/Plugins/Runtime/ReplicationSystemTestPlugin/Source/Private/Tests/ReplicationSystem/TestObjectSplitting.cpp
+++ b/Engine/Plugins/Runtime/ReplicationSystemTestPlugin/Source/Private/Tests/ReplicationSystem/TestObjectSplitting.cpp
@@ -122,7 +122,7 @@ UE_NET_TEST_FIXTURE(FSplitObjectTestFixture, SplitHugeObjectOnCreation)
 	UTestReplicatedIrisObject* ServerObject = CreateHugeObject(Server);
 
 	// Send and deliver packet
-	Server->PreSendUpdate();
+	Server->NetUpdate();
 	Server->SendAndDeliverTo(Client, DeliverPacket);
 	Server->PostSendUpdate();
 
@@ -132,7 +132,7 @@ UE_NET_TEST_FIXTURE(FSplitObjectTestFixture, SplitHugeObjectOnCreation)
 
 	for (uint32 RetryIt = 0; RetryIt != HugeObjectMaxNetTickCountToArrive && ClientObject == nullptr; ++RetryIt)
 	{
-		Server->PreSendUpdate();
+		Server->NetUpdate();
 		Server->SendAndDeliverTo(Client, DeliverPacket);
 		Server->PostSendUpdate();
 		ClientObject = Cast<UTestReplicatedIrisObject>(Client->GetReplicationBridge()->GetReplicatedObject(ServerObject->NetRefHandle));
@@ -148,7 +148,7 @@ UE_NET_TEST_FIXTURE(FSplitObjectTestFixture, SplitHugeObjectAfterCreation)
 	UTestReplicatedIrisObject* ServerObject = CreateObject(Server);
 
 	// Send and deliver packet
-	Server->PreSendUpdate();
+	Server->NetUpdate();
 	Server->SendAndDeliverTo(Client, DeliverPacket);
 	Server->PostSendUpdate();
 
@@ -163,7 +163,7 @@ UE_NET_TEST_FIXTURE(FSplitObjectTestFixture, SplitHugeObjectAfterCreation)
 
 	for (uint32 RetryIt = 0; RetryIt != HugeObjectMaxNetTickCountToArrive && Component->CallCounts.IntArrayRepNotifyCounter == 0; ++RetryIt)
 	{
-		Server->PreSendUpdate();
+		Server->NetUpdate();
 		Server->SendAndDeliverTo(Client, DeliverPacket);
 		Server->PostSendUpdate();
 	}
@@ -185,7 +185,7 @@ UE_NET_TEST_FIXTURE(FSplitObjectTestFixture, SplitObjectWithHugeSubObjectsOnCrea
 	}
 
 	// Send and deliver packet
-	Server->PreSendUpdate();
+	Server->NetUpdate();
 	Server->SendAndDeliverTo(Client, DeliverPacket);
 	Server->PostSendUpdate();
 
@@ -195,7 +195,7 @@ UE_NET_TEST_FIXTURE(FSplitObjectTestFixture, SplitObjectWithHugeSubObjectsOnCrea
 
 	for (uint32 RetryIt = 0; RetryIt != HugeObjectMaxNetTickCountToArrive && ClientObject == nullptr; ++RetryIt)
 	{
-		Server->PreSendUpdate();
+		Server->NetUpdate();
 		Server->SendAndDeliverTo(Client, DeliverPacket);
 		Server->PostSendUpdate();
 		ClientObject = Cast<UTestReplicatedIrisObject>(Client->GetReplicationBridge()->GetReplicatedObject(ServerObject->NetRefHandle));
@@ -245,7 +245,7 @@ UE_NET_TEST_FIXTURE(FSplitObjectTestFixture, SplitObjectWithSubObjectsWithHugeAt
 	}
 
 	// Send and deliver packet
-	Server->PreSendUpdate();
+	Server->NetUpdate();
 	Server->SendAndDeliverTo(Client, DeliverPacket);
 	Server->PostSendUpdate();
 
@@ -255,7 +255,7 @@ UE_NET_TEST_FIXTURE(FSplitObjectTestFixture, SplitObjectWithSubObjectsWithHugeAt
 
 	for (uint32 RetryIt = 0; RetryIt != HugeObjectMaxNetTickCountToArrive && ClientObject == nullptr; ++RetryIt)
 	{
-		Server->PreSendUpdate();
+		Server->NetUpdate();
 		Server->SendAndDeliverTo(Client, DeliverPacket);
 		Server->PostSendUpdate();
 		ClientObject = Cast<UTestReplicatedIrisObject>(Client->GetReplicationBridge()->GetReplicatedObject(ServerObject->NetRefHandle));
@@ -272,7 +272,7 @@ UE_NET_TEST_FIXTURE(FSplitObjectTestFixture, SplitObjectWithSubObjectsWithHugeAt
 	// Wait for attachments
 	for (uint32 RetryIt = 0; RetryIt != HugeObjectMaxNetTickCountToArrive; ++RetryIt)
 	{
-		Server->PreSendUpdate();
+		Server->NetUpdate();
 		Server->SendAndDeliverTo(Client, DeliverPacket);
 		Server->PostSendUpdate();
 	}
@@ -302,7 +302,7 @@ UE_NET_TEST_FIXTURE(FSplitObjectTestFixture, SplitObjectWithSubObjectsWithHugeAt
 	}
 
 	// Send and deliver packet
-	Server->PreSendUpdate();
+	Server->NetUpdate();
 	Server->SendAndDeliverTo(Client, DeliverPacket);
 	Server->PostSendUpdate();
 
@@ -341,7 +341,7 @@ UE_NET_TEST_FIXTURE(FSplitObjectTestFixture, SplitObjectWithSubObjectsWithHugeAt
 	bool bHasReceivedHugeState = false;
 	for (uint32 RetryIt = 0; RetryIt != HugeObjectMaxNetTickCountToArrive && !bHasReceivedHugeState; ++RetryIt)
 	{
-		Server->PreSendUpdate();
+		Server->NetUpdate();
 		Server->SendAndDeliverTo(Client, DeliverPacket);
 		Server->PostSendUpdate();
 
@@ -379,7 +379,7 @@ UE_NET_TEST_FIXTURE(FSplitObjectTestFixture, HugeObjectStateCanBeSentBackToBack)
 	}
 
 	// Send and deliver packet
-	Server->PreSendUpdate();
+	Server->NetUpdate();
 	Server->SendAndDeliverTo(Client, DeliverPacket);
 	Server->PostSendUpdate();
 
@@ -418,7 +418,7 @@ UE_NET_TEST_FIXTURE(FSplitObjectTestFixture, HugeObjectStateCanBeSentBackToBack)
 	bool bHasReceivedHugeState = false;
 	for (uint32 RetryIt = 0; RetryIt != HugeObjectMaxNetTickCountToArrive && !bHasReceivedHugeState; ++RetryIt)
 	{
-		Server->PreSendUpdate();
+		Server->NetUpdate();
 		Server->SendAndDeliverTo(Client, DeliverPacket);
 		Server->PostSendUpdate();
 
@@ -444,7 +444,7 @@ UE_NET_TEST_FIXTURE(FSplitObjectTestFixture, SplitObjectCanBeSentBackToBack)
 	UTestReplicatedIrisObject* ServerObject = CreateHugeObject(Server);
 
 	// Send and deliver packet. This will initiate huge object transfer.
-	Server->PreSendUpdate();
+	Server->NetUpdate();
 	Server->SendAndDeliverTo(Client, DeliverPacket);
 	Server->PostSendUpdate();
 
@@ -459,7 +459,7 @@ UE_NET_TEST_FIXTURE(FSplitObjectTestFixture, SplitObjectCanBeSentBackToBack)
 
 	for (uint32 RetryIt = 0; RetryIt != HugeObjectMaxNetTickCountToArrive && ClientObject == nullptr; ++RetryIt)
 	{
-		Server->PreSendUpdate();
+		Server->NetUpdate();
 		Server->SendAndDeliverTo(Client, DeliverPacket);
 		Server->PostSendUpdate();
 		ClientObject = Cast<UTestReplicatedIrisObject>(Client->GetReplicationBridge()->GetReplicatedObject(ServerObject->NetRefHandle));
@@ -470,7 +470,7 @@ UE_NET_TEST_FIXTURE(FSplitObjectTestFixture, SplitObjectCanBeSentBackToBack)
 	bool bHasReceivedSecondHugeState = false;
 	for (uint32 RetryIt = 0; RetryIt != HugeObjectMaxNetTickCountToArrive && !bHasReceivedSecondHugeState; ++RetryIt)
 	{
-		Server->PreSendUpdate();
+		Server->NetUpdate();
 		Server->SendAndDeliverTo(Client, DeliverPacket);
 		Server->PostSendUpdate();
 
@@ -491,7 +491,7 @@ UE_NET_TEST_FIXTURE(FSplitObjectTestFixture, SplitObjectIsDeletedAfterBeingCreat
 	const FNetRefHandle ServerNetRefHandle = ServerObject->NetRefHandle;
 
 	// Send and deliver packet
-	Server->PreSendUpdate();
+	Server->NetUpdate();
 	Server->SendAndDeliverTo(Client, DeliverPacket);
 	Server->PostSendUpdate();
 
@@ -503,7 +503,7 @@ UE_NET_TEST_FIXTURE(FSplitObjectTestFixture, SplitObjectIsDeletedAfterBeingCreat
 
 	for (uint32 RetryIt = 0; RetryIt != HugeObjectMaxNetTickCountToArrive && ClientObject == nullptr; ++RetryIt)
 	{
-		Server->PreSendUpdate();
+		Server->NetUpdate();
 		Server->SendAndDeliverTo(Client, DeliverPacket);
 		Server->PostSendUpdate();
 		ClientObject = Cast<UTestReplicatedIrisObject>(Client->GetReplicationBridge()->GetReplicatedObject(ServerNetRefHandle));
@@ -511,7 +511,7 @@ UE_NET_TEST_FIXTURE(FSplitObjectTestFixture, SplitObjectIsDeletedAfterBeingCreat
 	UE_NET_ASSERT_NE(ClientObject, nullptr);
 
 	// The object should be destroyed after the next net update.
-	Server->PreSendUpdate();
+	Server->NetUpdate();
 	Server->SendAndDeliverTo(Client, DeliverPacket);
 	Server->PostSendUpdate();
 
@@ -533,7 +533,7 @@ UE_NET_TEST_FIXTURE(FSplitObjectTestFixture, SubObjectToHugeObjectCanBeDeleted)
 	}
 
 	// Send and deliver packet
-	Server->PreSendUpdate();
+	Server->NetUpdate();
 	Server->SendAndDeliverTo(Client, DeliverPacket);
 	Server->PostSendUpdate();
 
@@ -552,7 +552,7 @@ UE_NET_TEST_FIXTURE(FSplitObjectTestFixture, SubObjectToHugeObjectCanBeDeleted)
 	}
 
 	// Initiate sending so that we have huge data in flight with the subobject we are going to destroy.
-	Server->PreSendUpdate();
+	Server->NetUpdate();
 	Server->SendAndDeliverTo(Client, DeliverPacket);
 	Server->PostSendUpdate();
 
@@ -561,7 +561,7 @@ UE_NET_TEST_FIXTURE(FSplitObjectTestFixture, SubObjectToHugeObjectCanBeDeleted)
 	bool bHasReceivedHugeState = false;
 	for (uint32 RetryIt = 0; RetryIt != HugeObjectMaxNetTickCountToArrive && !bHasReceivedHugeState; ++RetryIt)
 	{
-		Server->PreSendUpdate();
+		Server->NetUpdate();
 		Server->SendAndDeliverTo(Client, DeliverPacket);
 		Server->PostSendUpdate();
 
@@ -576,7 +576,7 @@ UE_NET_TEST_FIXTURE(FSplitObjectTestFixture, SubObjectToHugeObjectCanBeDeleted)
 	UE_NET_ASSERT_TRUE(bHasReceivedHugeState);
 
 	// Now the subobject can safely be destroyed
-	Server->PreSendUpdate();
+	Server->NetUpdate();
 	Server->SendAndDeliverTo(Client, DeliverPacket);
 	Server->PostSendUpdate();
 
@@ -599,7 +599,7 @@ UE_NET_TEST_FIXTURE(FSplitObjectTestFixture, TearOffOnCreation)
 	bool bHasHugeObjectBeenCreated = false;
 	for (uint32 RetryIt = 0; RetryIt != HugeObjectMaxNetTickCountToArrive && !bHasHugeObjectBeenCreated; ++RetryIt)
 	{
-		Server->PreSendUpdate();
+		Server->NetUpdate();
 		Server->SendAndDeliverTo(Client, DeliverPacket);
 		Server->PostSendUpdate();
 
@@ -624,7 +624,7 @@ UE_NET_TEST_FIXTURE(FSplitObjectTestFixture, TearOffCreatedObjectWithHugePayload
 	UTestReplicatedIrisObject* ServerObject = CreateObject(Server);
 
 	// Send and deliver packet
-	Server->PreSendUpdate();
+	Server->NetUpdate();
 	Server->SendAndDeliverTo(Client, true);
 	Server->PostSendUpdate();
 
@@ -639,7 +639,7 @@ UE_NET_TEST_FIXTURE(FSplitObjectTestFixture, TearOffCreatedObjectWithHugePayload
 	bool bHasReceivedHugeState = false;
 	for (uint32 RetryIt = 0; RetryIt != HugeObjectMaxNetTickCountToArrive && !bHasReceivedHugeState; ++RetryIt)
 	{
-		Server->PreSendUpdate();
+		Server->NetUpdate();
 		Server->SendAndDeliverTo(Client, DeliverPacket);
 		Server->PostSendUpdate();
 
@@ -662,7 +662,7 @@ UE_NET_TEST_FIXTURE(FSplitObjectTestFixture, TearOffWhileHugeObjectStateIsSendin
 	UTestReplicatedIrisObject* ServerObject = CreateHugeObject(Server);
 
 	// Send and deliver packet
-	Server->PreSendUpdate();
+	Server->NetUpdate();
 	Server->SendAndDeliverTo(Client, true);
 	Server->PostSendUpdate();
 
@@ -675,7 +675,7 @@ UE_NET_TEST_FIXTURE(FSplitObjectTestFixture, TearOffWhileHugeObjectStateIsSendin
 	bool bHasHugeObjectBeenCreated = false;
 	for (uint32 RetryIt = 0; RetryIt != HugeObjectMaxNetTickCountToArrive && !bHasHugeObjectBeenCreated; ++RetryIt)
 	{
-		Server->PreSendUpdate();
+		Server->NetUpdate();
 		Server->SendAndDeliverTo(Client, DeliverPacket);
 		Server->PostSendUpdate();
 
@@ -690,7 +690,7 @@ UE_NET_TEST_FIXTURE(FSplitObjectTestFixture, TearOffWhileHugeObjectStateIsSendin
 	const UTestReplicatedIrisObject* ClientObjectThatWillBeTornOff = Cast<UTestReplicatedIrisObject>(Client->GetReplicationBridge()->GetReplicatedObject(ServerObject->NetRefHandle));
 	UE_NET_ASSERT_EQ((ClientObjectThatWillBeTornOff->IntA ^ 1), ServerObject->IntA);
 
-	Server->PreSendUpdate();
+	Server->NetUpdate();
 	Server->SendAndDeliverTo(Client, true);
 	Server->PostSendUpdate();
 
@@ -710,19 +710,19 @@ UE_NET_TEST_FIXTURE(FSplitObjectTestFixture, TestCancelPendingDestroyOfHugeObjec
 	// Write packets
 	for (uint32 RetryIt = 0; RetryIt != HugeObjectMaxNetTickCountToArrive; ++RetryIt)
 	{
-		Server->PreSendUpdate();
+		Server->NetUpdate();
 		Server->SendUpdate(Client->ConnectionIdOnServer);
 		Server->PostSendUpdate();
 	}
 
 	// Filter out object to cause a PendingDestroy
 	Server->GetReplicationSystem()->AddToGroup(Server->GetReplicationSystem()->GetNotReplicatedNetObjectGroup(), ServerObject->NetRefHandle);
-	Server->PreSendUpdate();
+	Server->NetUpdate();
 	Server->PostSendUpdate();
 
 	// Remove object from filter to cause object to end up in CancelPendingDestroy
 	Server->GetReplicationSystem()->RemoveFromGroup(Server->GetReplicationSystem()->GetNotReplicatedNetObjectGroup(), ServerObject->NetRefHandle);
-	Server->PreSendUpdate();
+	Server->NetUpdate();
 	Server->PostSendUpdate();
 
 	// Deliver object creation packets
@@ -742,7 +742,7 @@ UE_NET_TEST_FIXTURE(FSplitObjectTestFixture, TestCancelPendingDestroyOfHugeObjec
 	// Modify a property on the object and make sure it's replicated as the object should now be confirmed created
 	ServerObject->IntA ^= 1;
 
-	Server->PreSendUpdate();
+	Server->NetUpdate();
 	Server->SendAndDeliverTo(Client, DeliverPacket);
 	Server->PostSendUpdate();
 
@@ -763,19 +763,19 @@ UE_NET_TEST_FIXTURE(FSplitObjectTestFixture, TestCancelPendingDestroyOfHugeObjec
 	// Write packets
 	for (uint32 RetryIt = 0; RetryIt != HugeObjectMaxNetTickCountToArrive; ++RetryIt)
 	{
-		Server->PreSendUpdate();
+		Server->NetUpdate();
 		Server->SendUpdate(Client->ConnectionIdOnServer);
 		Server->PostSendUpdate();
 	}
 
 	// Filter out object to cause a PendingDestroy
 	Server->GetReplicationSystem()->AddToGroup(Server->GetReplicationSystem()->GetNotReplicatedNetObjectGroup(), ServerObject->NetRefHandle);
-	Server->PreSendUpdate();
+	Server->NetUpdate();
 	Server->PostSendUpdate();
 
 	// Remove object from filter to cause object to end up in CancelPendingDestroy
 	Server->GetReplicationSystem()->RemoveFromGroup(Server->GetReplicationSystem()->GetNotReplicatedNetObjectGroup(), ServerObject->NetRefHandle);
-	Server->PreSendUpdate();
+	Server->NetUpdate();
 	Server->PostSendUpdate();
 
 	// Cause packet loss on object creation
@@ -795,7 +795,7 @@ UE_NET_TEST_FIXTURE(FSplitObjectTestFixture, TestCancelPendingDestroyOfHugeObjec
 
 		for (uint32 RetryIt = 0; RetryIt != HugeObjectMaxNetTickCountToArrive; ++RetryIt)
 		{
-			Server->PreSendUpdate();
+			Server->NetUpdate();
 			Server->SendAndDeliverTo(Client, DeliverPacket);
 			Server->PostSendUpdate();
 
@@ -818,7 +818,7 @@ UE_NET_TEST_FIXTURE(FSplitObjectTestFixture, TestCancelPendingDestroyDuringHugeO
 	UTestReplicatedIrisObject* ServerObject = CreateObject(Server);
 
 	// Write and send packet
-	Server->PreSendUpdate();
+	Server->NetUpdate();
 	Server->SendAndDeliverTo(Client, DeliverPacket);
 	Server->PostSendUpdate();
 
@@ -828,19 +828,19 @@ UE_NET_TEST_FIXTURE(FSplitObjectTestFixture, TestCancelPendingDestroyDuringHugeO
 	// Write packets
 	for (uint32 RetryIt = 0; RetryIt != HugeObjectMaxNetTickCountToArrive; ++RetryIt)
 	{
-		Server->PreSendUpdate();
+		Server->NetUpdate();
 		Server->SendUpdate(Client->ConnectionIdOnServer);
 		Server->PostSendUpdate();
 	}
 
 	// Filter out object to cause a PendingDestroy
 	Server->GetReplicationSystem()->AddToGroup(Server->GetReplicationSystem()->GetNotReplicatedNetObjectGroup(), ServerObject->NetRefHandle);
-	Server->PreSendUpdate();
+	Server->NetUpdate();
 	Server->PostSendUpdate();
 
 	// Remove object from filter to cause object to end up in CancelPendingDestroy
 	Server->GetReplicationSystem()->RemoveFromGroup(Server->GetReplicationSystem()->GetNotReplicatedNetObjectGroup(), ServerObject->NetRefHandle);
-	Server->PreSendUpdate();
+	Server->NetUpdate();
 	Server->PostSendUpdate();
 
 	// Modify a property on the object and make sure it's replicated as the object should still be created
@@ -858,7 +858,7 @@ UE_NET_TEST_FIXTURE(FSplitObjectTestFixture, TestCancelPendingDestroyDuringHugeO
 	}
 
 	// Deliver latest state
-	Server->PreSendUpdate();
+	Server->NetUpdate();
 	Server->SendAndDeliverTo(Client, DeliverPacket);
 	Server->PostSendUpdate();
 
@@ -877,7 +877,7 @@ UE_NET_TEST_FIXTURE(FSplitObjectTestFixture, TestReliableAttachmentIsDeliveredDe
 	UTestReplicatedIrisObject* ServerObject = CreateObject(Server);
 
 	// Send and deliver packet
-	Server->PreSendUpdate();
+	Server->NetUpdate();
 	Server->SendAndDeliverTo(Client, DeliverPacket);
 	Server->PostSendUpdate();
 
@@ -895,14 +895,14 @@ UE_NET_TEST_FIXTURE(FSplitObjectTestFixture, TestReliableAttachmentIsDeliveredDe
 	// Write packets
 	for (uint32 RetryIt = 0; RetryIt != HugeObjectMaxNetTickCountToArrive; ++RetryIt)
 	{
-		Server->PreSendUpdate();
+		Server->NetUpdate();
 		Server->SendUpdate(Client->ConnectionIdOnServer);
 		Server->PostSendUpdate();
 	}
 
 	// Filter out object to cause object to be set in state WaitOnFlush
 	Server->GetReplicationSystem()->AddToGroup(Server->GetReplicationSystem()->GetNotReplicatedNetObjectGroup(), ServerObject->NetRefHandle);
-	Server->PreSendUpdate();
+	Server->NetUpdate();
 	Server->SendUpdate(Client->ConnectionIdOnServer);
 	Server->PostSendUpdate();
 
@@ -918,7 +918,7 @@ UE_NET_TEST_FIXTURE(FSplitObjectTestFixture, TestReliableAttachmentIsDeliveredDe
 	}
 
 	// Deliver latest state
-	Server->PreSendUpdate();
+	Server->NetUpdate();
 	Server->SendAndDeliverTo(Client, DeliverPacket);
 	Server->PostSendUpdate();
 
@@ -941,7 +941,7 @@ UE_NET_TEST_FIXTURE(FSplitObjectTestFixture, TestHugeObjectIsFlushedAndNotDestro
 	UTestReplicatedIrisObject* ServerObject = CreateObject(Server);
 
 	// Send and deliver packet
-	Server->PreSendUpdate();
+	Server->NetUpdate();
 	Server->SendAndDeliverTo(Client, DeliverPacket);
 	Server->PostSendUpdate();
 
@@ -959,20 +959,20 @@ UE_NET_TEST_FIXTURE(FSplitObjectTestFixture, TestHugeObjectIsFlushedAndNotDestro
 	// Write packets
 	for (uint32 RetryIt = 0; RetryIt != HugeObjectMaxNetTickCountToArrive; ++RetryIt)
 	{
-		Server->PreSendUpdate();
+		Server->NetUpdate();
 		Server->SendUpdate(Client->ConnectionIdOnServer);
 		Server->PostSendUpdate();
 	}
 
 	// Filter out object to cause object to be set in state WaitOnFlush
 	Server->GetReplicationSystem()->AddToGroup(Server->GetReplicationSystem()->GetNotReplicatedNetObjectGroup(), ServerObject->NetRefHandle);
-	Server->PreSendUpdate();
+	Server->NetUpdate();
 	Server->SendUpdate(Client->ConnectionIdOnServer);
 	Server->PostSendUpdate();
 
 	// Remove object from filter to cause object to be set in state Created
 	Server->GetReplicationSystem()->RemoveFromGroup(Server->GetReplicationSystem()->GetNotReplicatedNetObjectGroup(), ServerObject->NetRefHandle);
-	Server->PreSendUpdate();
+	Server->NetUpdate();
 	Server->SendUpdate(Client->ConnectionIdOnServer);
 	Server->PostSendUpdate();
 
@@ -991,7 +991,7 @@ UE_NET_TEST_FIXTURE(FSplitObjectTestFixture, TestHugeObjectIsFlushedAndNotDestro
 	ServerObject->IntA += 1;
 
 	// Deliver latest state
-	Server->PreSendUpdate();
+	Server->NetUpdate();
 	Server->SendAndDeliverTo(Client, DeliverPacket);
 	Server->PostSendUpdate();
 
@@ -1021,7 +1021,7 @@ UE_NET_TEST_FIXTURE(FSplitObjectTestFixture, SplitManyHugeObjectsOnCreation)
 	uint32 ClientObjectCount = 0;
 	for (uint32 RetryIt = 0; RetryIt != (HugeObjectMaxNetTickCountToArrive*HugeObjectCount) && ClientObjectCount < HugeObjectCount; ++RetryIt)
 	{
-		Server->PreSendUpdate();
+		Server->NetUpdate();
 		Server->SendAndDeliverTo(Client, DeliverPacket);
 		Server->PostSendUpdate();
 
@@ -1058,7 +1058,7 @@ UE_NET_TEST_FIXTURE(FSplitObjectTestFixture, SplitManyHugeObjectsAfterCreation)
 	uint32 ClientObjectCount = 0;
 	for (uint32 RetryIt = 0; RetryIt != HugeObjectCount && ClientObjectCount < HugeObjectCount; ++RetryIt)
 	{
-		Server->PreSendUpdate();
+		Server->NetUpdate();
 		Server->SendAndDeliverTo(Client, DeliverPacket);
 		Server->PostSendUpdate();
 
@@ -1093,7 +1093,7 @@ UE_NET_TEST_FIXTURE(FSplitObjectTestFixture, SplitManyHugeObjectsAfterCreation)
 	uint32 ClientObjectsWithHugeArraysCount = 0;
 	for (uint32 RetryIt = 0; RetryIt != (HugeObjectMaxNetTickCountToArrive*HugeObjectCount) && ClientObjectsWithHugeArraysCount < HugeObjectCount; ++RetryIt)
 	{
-		Server->PreSendUpdate();
+		Server->NetUpdate();
 		Server->SendAndDeliverTo(Client, DeliverPacket);
 		Server->PostSendUpdate();
 
@@ -1119,14 +1119,14 @@ UE_NET_TEST_FIXTURE(FSplitObjectTestFixture, TestDependentObjectCannotBeDestroye
 	ServerBridge->AddDependentObject(ServerObject->NetRefHandle, ServerDependentObject->NetRefHandle);
 
 	// Introduce latency by not immediately delivering packets.
-	Server->PreSendUpdate();
+	Server->NetUpdate();
 	Server->SendTo(Client, TEXT("Create HugeObject + Dependent"));
 	Server->PostSendUpdate();
 
 	// Filter out dependent object to cause it to end up being destroyed
 	Server->ReplicationSystem->AddToGroup(Server->GetReplicationSystem()->GetNotReplicatedNetObjectGroup(), ServerDependentObject->NetRefHandle);
 
-	Server->PreSendUpdate();
+	Server->NetUpdate();
 	Server->SendTo(Client, TEXT("Try destroy Dependent"));
 	Server->PostSendUpdate();
 
@@ -1177,13 +1177,13 @@ UE_NET_TEST_FIXTURE(FSplitObjectTestFixture, TestHugeObjectCanBeSentDuringReplic
 	UTestReplicatedIrisObject* ServerHugeObject = CreateHugeObject(Server);
 
 	// Write packet
-	Server->PreSendUpdate();
+	Server->NetUpdate();
 	Server->SendUpdate(Client->ConnectionIdOnServer);
 	Server->PostSendUpdate();
 
 	for (uint32 RetryIt = 0; RetryIt != HugeObjectMaxNetTickCountToArrive; ++RetryIt)
 	{
-		Server->PreSendUpdate();
+		Server->NetUpdate();
 		const bool bPacketWasWritten = Server->SendUpdate(Client->ConnectionIdOnServer, TEXT("Create HugeObject"));
 		Server->PostSendUpdate();
 		if (!bPacketWasWritten)
diff --git a/Engine/Plugins/Runtime/ReplicationSystemTestPlugin/Source/Private/Tests/ReplicationSystem/TestReplicationOperations.cpp b/Engine/Plugins/Runtime/ReplicationSystemTestPlugin/Source/Private/Tests/ReplicationSystem/TestReplicationOperations.cpp
--- a/Engine/Plugins/Runtime/ReplicationSystemTestPlugin/Source/Private/Tests/ReplicationSystem/TestReplicationOperations.cpp
+++ b/Engine/Plugins/Runtime/ReplicationSystemTestPlugin/Source/Private/Tests/ReplicationSystem/TestReplicationOperations.cpp
@@ -532,7 +532,7 @@ UE_NET_TEST_FIXTURE(FTestReplicationOperationsForObjectsFixture, StaleObjectPoin
 	SerializationContext.SetIsInitState(true);
 
 	// Copy data from test object to buffer
-	ReplicationSystem->PreSendUpdate(0.033f);
+	ReplicationSystem->NetUpdate(0.033f);
 	FReplicationInstanceOperations::Quantize(SerializationContext, StateBuffers1[0], GetChangeMaskWriter(Protocol1->ChangeMaskBitCount), InstanceProtocol1, Protocol1);
 	ReplicationSystem->PostSendUpdate();
 
@@ -542,7 +542,7 @@ UE_NET_TEST_FIXTURE(FTestReplicationOperationsForObjectsFixture, StaleObjectPoin
 	constexpr bool bPerformFullPurge = false;
 	CollectGarbage(GARBAGE_COLLECTION_KEEPFLAGS, bPerformFullPurge);
 
-	ReplicationSystem->PreSendUpdate(0.033f);
+	ReplicationSystem->NetUpdate(0.033f);
 	FReplicationInstanceOperations::Quantize(SerializationContext, StateBuffers1[1], GetChangeMaskWriter(Protocol1->ChangeMaskBitCount), InstanceProtocol1, Protocol1);
 	ReplicationSystem->PostSendUpdate();
 
diff --git a/Engine/Plugins/Runtime/ReplicationSystemTestPlugin/Source/Private/Tests/ReplicationSystem/TestReplicationSystem.cpp b/Engine/Plugins/Runtime/ReplicationSystemTestPlugin/Source/Private/Tests/ReplicationSystem/TestReplicationSystem.cpp
--- a/Engine/Plugins/Runtime/ReplicationSystemTestPlugin/Source/Private/Tests/ReplicationSystem/TestReplicationSystem.cpp
+++ b/Engine/Plugins/Runtime/ReplicationSystemTestPlugin/Source/Private/Tests/ReplicationSystem/TestReplicationSystem.cpp
@@ -105,7 +105,7 @@ UE_NET_TEST_FIXTURE(FReplicationSystemServerClientTestFixture, TestDestroyWhileW
 	const FNetRefHandle ServerObjectRefHandle = ServerObject->NetRefHandle;
 
 	// Write packet with create
-	Server->PreSendUpdate();
+	Server->NetUpdate();
 	Server->SendTo(Client);
 	Server->PostSendUpdate();
 
@@ -113,7 +113,7 @@ UE_NET_TEST_FIXTURE(FReplicationSystemServerClientTestFixture, TestDestroyWhileW
 	Server->DestroyObject(ServerObject);
 
 	// Write packet with destroy
-	Server->PreSendUpdate();
+	Server->NetUpdate();
 	Server->SendTo(Client);
 	Server->PostSendUpdate();
 
@@ -150,7 +150,7 @@ UE_NET_TEST_FIXTURE(FReplicationSystemServerClientTestFixture, TestDestroyWithDa
 	ServerObject->IntA = 13;
 
 	// Write a packet with updated data
-	Server->PreSendUpdate();
+	Server->NetUpdate();
 	UE_NET_ASSERT_TRUE(Server->SendTo(Client));
 	Server->PostSendUpdate();
 
@@ -158,7 +158,7 @@ UE_NET_TEST_FIXTURE(FReplicationSystemServerClientTestFixture, TestDestroyWithDa
 	Server->DestroyObject(ServerObject);
 
 	// Write packet with destroy
-	Server->PreSendUpdate();
+	Server->NetUpdate();
 	UE_NET_ASSERT_TRUE(Server->SendTo(Client));
 	Server->PostSendUpdate();
 
@@ -369,7 +369,7 @@ UE_NET_TEST_FIXTURE(FReplicationSystemServerClientTestFixture, ReplicateAndDestr
 
 	uint32 NumUnAcknowledgedPackets = 0;
 	// Write a packet
-	Server->PreSendUpdate();
+	Server->NetUpdate();
 	Server->SendTo(Client);
 	Server->PostSendUpdate();
 
@@ -377,7 +377,7 @@ UE_NET_TEST_FIXTURE(FReplicationSystemServerClientTestFixture, ReplicateAndDestr
 	Server->DestroyObject(ServerSubObject);
 
 	// We have no data to send but we want to tick ReplicationSystem to capture state change
-	Server->PreSendUpdate();
+	Server->NetUpdate();
 	UE_NET_ASSERT_FALSE(Server->SendTo(Client));
 	Server->PostSendUpdate();
 
@@ -388,7 +388,7 @@ UE_NET_TEST_FIXTURE(FReplicationSystemServerClientTestFixture, ReplicateAndDestr
 	//Server->DeliverTo(Client, true);
 
 	// Send and deliver packet
-	Server->PreSendUpdate();
+	Server->NetUpdate();
 	Server->SendAndDeliverTo(Client, true);
 	Server->PostSendUpdate();
 
@@ -419,7 +419,7 @@ UE_NET_TEST_FIXTURE(FReplicationSystemServerClientTestFixture, LateCreatedSubObj
 	const FNetRefHandle ServerSubObjectRefHandle = ServerSubObject->NetRefHandle;
 
 	// Put subobject in WaitOnCreateConfirmation state
-	Server->PreSendUpdate();
+	Server->NetUpdate();
 	Server->SendUpdate(Client->ConnectionIdOnServer);
 	Server->PostSendUpdate();
 
@@ -427,7 +427,7 @@ UE_NET_TEST_FIXTURE(FReplicationSystemServerClientTestFixture, LateCreatedSubObj
 	Server->DestroyObject(ServerObject);
 
 	// Write a packet
-	Server->PreSendUpdate();
+	Server->NetUpdate();
 	Server->SendUpdate(Client->ConnectionIdOnServer);
 	Server->PostSendUpdate();
 
@@ -536,7 +536,7 @@ UE_NET_TEST_FIXTURE(FReplicationSystemServerClientTestFixture, TestSubObjectDefa
 	UReplicatedSubObjectOrderObject::RepOrderCounter = 0U;
 
 	// Send and deliver packet
-	Server->PreSendUpdate();
+	Server->NetUpdate();
 	Server->SendAndDeliverTo(Client, true);
 	Server->PostSendUpdate();
 
@@ -573,7 +573,7 @@ UE_NET_TEST_FIXTURE(FReplicationSystemServerClientTestFixture, TestSubObjectSpec
 	UReplicatedSubObjectOrderObject::RepOrderCounter = 0U;
 
 	// Send and deliver packet
-	Server->PreSendUpdate();
+	Server->NetUpdate();
 	Server->SendAndDeliverTo(Client, true);
 	Server->PostSendUpdate();
 
@@ -693,7 +693,7 @@ UE_NET_TEST_FIXTURE(FTestNetTokensFixture, NetToken)
 	FNetToken StringTokenA = CreateAndExportNetToken(TokenStringA);
 
 	// Send and drop packet
-	Server->PreSendUpdate();
+	Server->NetUpdate();
 	Server->SendAndDeliverTo(Client, false);
 	Server->PostSendUpdate();
 
@@ -705,7 +705,7 @@ UE_NET_TEST_FIXTURE(FTestNetTokensFixture, NetToken)
 	}
 
 	// Send and deliver packet
-	Server->PreSendUpdate();
+	Server->NetUpdate();
 	Server->SendAndDeliverTo(Client, true);
 	Server->PostSendUpdate();
 
@@ -720,7 +720,7 @@ UE_NET_TEST_FIXTURE(FTestNetTokensFixture, NetTokenResendWithFullPacket)
 	FNetToken StringTokenA = CreateAndExportNetToken(TokenStringA);
 
 	// Send and drop packet
-	Server->PreSendUpdate();
+	Server->NetUpdate();
 	Server->SendAndDeliverTo(Client, false);
 	Server->PostSendUpdate();
 
@@ -732,7 +732,7 @@ UE_NET_TEST_FIXTURE(FTestNetTokensFixture, NetTokenResendWithFullPacket)
 	FNetToken StringTokenB = CreateAndExportNetToken(TokenStringB);
 
 	// Send and deliver packet
-	Server->PreSendUpdate();
+	Server->NetUpdate();
 	Server->SendAndDeliverTo(Client, true);
 	Server->PostSendUpdate();
 
@@ -747,7 +747,7 @@ UE_NET_TEST_FIXTURE(FTestNetTokensFixture, NetTokenResendWithFullPacket)
 	Server->SetMaxSendPacketSize(1024U);
 
 	// Send and deliver packet
-	Server->PreSendUpdate();
+	Server->NetUpdate();
 	Server->SendAndDeliverTo(Client, true);
 	Server->PostSendUpdate();
 
@@ -761,7 +761,7 @@ UE_NET_TEST_FIXTURE(FTestNetTokensFixture, NetTokenResendWithFullPacketAfterFirs
 	FNetToken StringTokenA = CreateAndExportNetToken(TestStringA);	
 
 	// Send and delay delivery
-	Server->PreSendUpdate();
+	Server->NetUpdate();
 	Server->SendTo(Client);
 	Server->PostSendUpdate();
 
@@ -770,7 +770,7 @@ UE_NET_TEST_FIXTURE(FTestNetTokensFixture, NetTokenResendWithFullPacketAfterFirs
 	FNetToken StringTokenB = CreateAndExportNetToken(TestStringB);
 
 	// Send and delay delivery
-	Server->PreSendUpdate();
+	Server->NetUpdate();
 	Server->SendTo(Client);
 	Server->PostSendUpdate();
 
@@ -788,7 +788,7 @@ UE_NET_TEST_FIXTURE(FTestNetTokensFixture, NetTokenResendWithFullPacketAfterFirs
 	// Send and deliver packet which now should contain two entries in the resend queue
 	Server->SetMaxSendPacketSize(1024);
 
-	Server->PreSendUpdate();
+	Server->NetUpdate();
 	Server->SendAndDeliverTo(Client, true);
 	Server->PostSendUpdate();
 
@@ -815,7 +815,7 @@ UE_NET_TEST_FIXTURE(FTestNetTokensFixture, NetTokenSequenceTest)
 	FNetToken StringTokenB = CreateAndExportNetToken(TestStrings[1]);
 
 	// Send packet
-	Server->PreSendUpdate();
+	Server->NetUpdate();
 	Server->SendTo(Client);
 	Server->PostSendUpdate();
 
@@ -826,7 +826,7 @@ UE_NET_TEST_FIXTURE(FTestNetTokensFixture, NetTokenSequenceTest)
 	FNetToken StringTokenD = CreateAndExportNetToken(TestStrings[3]);
 
 	// Send packet
-	Server->PreSendUpdate();
+	Server->NetUpdate();
 	Server->SendTo(Client);
 	Server->PostSendUpdate();
 
@@ -841,7 +841,7 @@ UE_NET_TEST_FIXTURE(FTestNetTokensFixture, NetTokenSequenceTest)
 	ClientStringTokenStore->GetOrCreateToken(TEXT("LocalTokenB"));
 
 	// Send packet with resend data
-	Server->PreSendUpdate();
+	Server->NetUpdate();
 	Server->SendAndDeliverTo(Client, true);
 	Server->PostSendUpdate();
 
@@ -865,7 +865,7 @@ UE_NET_TEST_FIXTURE(FTestNetTokensFixture, NetTokenResendAndDataInSamePacketTest
 	FNetToken StringTokenA = CreateAndExportNetToken(TestStrings[0]);
 
 	// Send packet
-	Server->PreSendUpdate();
+	Server->NetUpdate();
 	Server->SendTo(Client);
 	Server->PostSendUpdate();
 
@@ -876,7 +876,7 @@ UE_NET_TEST_FIXTURE(FTestNetTokensFixture, NetTokenResendAndDataInSamePacketTest
 	FNetToken StringTokenB = CreateAndExportNetToken(TestStrings[1]);
 
 	// Send and deliver packet
-	Server->PreSendUpdate();
+	Server->NetUpdate();
 	Server->SendAndDeliverTo(Client, true);
 	Server->PostSendUpdate();
 
@@ -962,7 +962,7 @@ UE_NET_TEST_FIXTURE(FReplicationSystemServerClientTestFixture, AddRemoveFromConn
 	
 	// Send packet
 	// Expected state to be WaitOnCreateConfirmation
-	Server->PreSendUpdate();
+	Server->NetUpdate();
 	Server->SendTo(Client);
 	Server->PostSendUpdate();
 
@@ -977,7 +977,7 @@ UE_NET_TEST_FIXTURE(FReplicationSystemServerClientTestFixture, AddRemoveFromConn
 	UE_NET_ASSERT_NE(Client->GetReplicationBridge()->GetReplicatedObject(ServerObject->NetRefHandle), nullptr);
 
 	// Send packet
-	Server->PreSendUpdate();
+	Server->NetUpdate();
 	Server->SendTo(Client);
 	Server->PostSendUpdate();
 
@@ -1015,7 +1015,7 @@ UE_NET_TEST_FIXTURE(FReplicationSystemServerClientTestFixture, TestNetTemporary)
 	ServerObject1->IntA = 1;
 
 	// Send and deliver packet
-	Server->PreSendUpdate();
+	Server->NetUpdate();
 	Server->SendAndDeliverTo(Client, true);
 	Server->PostSendUpdate();
 
@@ -1043,7 +1043,7 @@ UE_NET_TEST_FIXTURE(FReplicationSystemServerClientTestFixture, TestNetTemporary)
 	ServerObject1->IntA = 2;
 
 	// Send and deliver packet
-	Server->PreSendUpdate();
+	Server->NetUpdate();
 	Server->SendAndDeliverTo(Client, true);
 	Server->PostSendUpdate();
 
@@ -1068,7 +1068,7 @@ UE_NET_TEST_FIXTURE(FReplicationSystemServerClientTestFixture, TestNetTemporary)
 	FReplicationSystemTestClient* Client2 = CreateClient();
 
 	// Send and deliver packet
-	Server->PreSendUpdate();
+	Server->NetUpdate();
 	Server->SendAndDeliverTo(Client2, true);
 	Server->PostSendUpdate();
 
@@ -1106,7 +1106,7 @@ UE_NET_TEST_FIXTURE(FReplicationSystemServerClientTestFixture, TestTearOffExisti
 	ServerObject->IntA = 1;
 
 	// Send and deliver packet
-	Server->PreSendUpdate();
+	Server->NetUpdate();
 	Server->SendAndDeliverTo(Client, true);
 	Server->PostSendUpdate();
 
@@ -1123,7 +1123,7 @@ UE_NET_TEST_FIXTURE(FReplicationSystemServerClientTestFixture, TestTearOffExisti
 	Server->ReplicationBridge->EndReplication(ServerObject, EEndReplicationFlags::TearOff);
 
 	// Send and deliver packet
-	Server->PreSendUpdate();
+	Server->NetUpdate();
 	Server->SendAndDeliverTo(Client, true);
 	Server->PostSendUpdate();
 
@@ -1153,7 +1153,7 @@ UE_NET_TEST_FIXTURE(FReplicationSystemServerClientTestFixture, TestTearOffOnNewl
 	Server->ReplicationBridge->EndReplication(ServerObject, EEndReplicationFlags::TearOff);
 
 	// Send and deliver packet
-	Server->PreSendUpdate();
+	Server->NetUpdate();
 	Server->SendAndDeliverTo(Client, true);
 	Server->PostSendUpdate();
 
@@ -1184,7 +1184,7 @@ UE_NET_TEST_FIXTURE(FReplicationSystemServerClientTestFixture, TestTearOffResend
 	ServerObject->IntA = 1;
 
 	// Send and deliver packet
-	Server->PreSendUpdate();
+	Server->NetUpdate();
 	Server->SendAndDeliverTo(Client, DeliverPacket);
 	Server->PostSendUpdate();
 
@@ -1198,14 +1198,14 @@ UE_NET_TEST_FIXTURE(FReplicationSystemServerClientTestFixture, TestTearOffResend
 	Server->ReplicationBridge->EndReplication(ServerObject, EEndReplicationFlags::TearOff);
 
 	// Send and do not deliver packet
-	Server->PreSendUpdate();
+	Server->NetUpdate();
 	Server->SendAndDeliverTo(Client, DoNotDeliverPacket);
 	Server->PostSendUpdate();
 
 	// The ClientObject should still be found using the NetRefHandle
 	UE_NET_ASSERT_NE(Cast<UTestReplicatedIrisObject>(Client->GetReplicationBridge()->GetReplicatedObject(ServerObject->NetRefHandle)), nullptr);
 
-	Server->PreSendUpdate();
+	Server->NetUpdate();
 	Server->SendAndDeliverTo(Client, DeliverPacket);
 	Server->PostSendUpdate();
 
@@ -1320,12 +1320,12 @@ UE_NET_TEST_FIXTURE(FReplicationSystemServerClientTestFixture, TestDefferedTearO
 	Server->ReplicationSystem->TearOffNextUpdate(ServerObject->NetRefHandle);
 
 	// Send and deliver packet
-	Server->PreSendUpdate();
+	Server->NetUpdate();
 	Server->SendAndDeliverTo(Client, false);
 	Server->PostSendUpdate();
 
 	// Send and deliver packet
-	Server->PreSendUpdate();
+	Server->NetUpdate();
 	Server->SendAndDeliverTo(Client, true);
 	Server->PostSendUpdate();
 	
@@ -1359,7 +1359,7 @@ UE_NET_TEST_FIXTURE(FReplicationSystemServerClientTestFixture, TestTearOffObject
 	ServerObject->IntA = 1;
 
 	// Send packet to get put the object in flight
-	Server->PreSendUpdate();
+	Server->NetUpdate();
 	Server->SendTo(Client);
 	Server->PostSendUpdate();
 
@@ -1376,7 +1376,7 @@ UE_NET_TEST_FIXTURE(FReplicationSystemServerClientTestFixture, TestTearOffObject
 	UE_NET_ASSERT_EQ(ServerObject->IntA, ClientObjectThatWillBeTornOff->IntA);
 
 	// Send and deliver packet
-	Server->PreSendUpdate();
+	Server->NetUpdate();
 	Server->SendAndDeliverTo(Client, true);
 	Server->PostSendUpdate();
 
@@ -1400,7 +1400,7 @@ UE_NET_TEST_FIXTURE(FReplicationSystemServerClientTestFixture, TestTearOffExisti
 	ServerObject->IntA = 1;
 
 	// Send and deliver packet
-	Server->PreSendUpdate();
+	Server->NetUpdate();
 	Server->SendAndDeliverTo(Client, true);
 	Server->PostSendUpdate();
 
@@ -1420,7 +1420,7 @@ UE_NET_TEST_FIXTURE(FReplicationSystemServerClientTestFixture, TestTearOffExisti
 	Server->ReplicationBridge->EndReplication(ServerObject, EEndReplicationFlags::TearOff);
 
 	// Send and deliver packet
-	Server->PreSendUpdate();
+	Server->NetUpdate();
 	Server->SendAndDeliverTo(Client, true);
 	Server->PostSendUpdate();
 
@@ -1444,7 +1444,7 @@ UE_NET_TEST_FIXTURE(FReplicationSystemServerClientTestFixture, TestTearOffResend
 	ServerObject->IntA = 1;
 
 	// Send and deliver packet
-	Server->PreSendUpdate();
+	Server->NetUpdate();
 	Server->SendAndDeliverTo(Client, true);
 	Server->PostSendUpdate();
 
@@ -1461,7 +1461,7 @@ UE_NET_TEST_FIXTURE(FReplicationSystemServerClientTestFixture, TestTearOffResend
 	Server->ReplicationSystem->TearOffNextUpdate(ServerObject->NetRefHandle);
 
 	// Send and deliver packet, in this case the packet containing 2 was lost, but, we did not know that when we 
-	Server->PreSendUpdate();
+	Server->NetUpdate();
 	Server->SendAndDeliverTo(Client, false);
 	Server->PostSendUpdate();
 
@@ -1469,7 +1469,7 @@ UE_NET_TEST_FIXTURE(FReplicationSystemServerClientTestFixture, TestTearOffResend
 	ServerObject->IntA = 3;
 
 	// Send and deliver packet
-	Server->PreSendUpdate();
+	Server->NetUpdate();
 	Server->SendAndDeliverTo(Client, true);
 	Server->PostSendUpdate();
 
@@ -1493,7 +1493,7 @@ UE_NET_TEST_FIXTURE(FReplicationSystemServerClientTestFixture, TestTornedOffObje
 	ServerObject->IntA = 1;
 
 	// Send and deliver packet
-	Server->PreSendUpdate();
+	Server->NetUpdate();
 	Server->SendAndDeliverTo(Client, true);
 	Server->PostSendUpdate();
 
@@ -1510,7 +1510,7 @@ UE_NET_TEST_FIXTURE(FReplicationSystemServerClientTestFixture, TestTornedOffObje
 	Server->ReplicationSystem->TearOffNextUpdate(ServerObject->NetRefHandle);
 
 	// Send and drop packet containing the value 2
-	Server->PreSendUpdate();
+	Server->NetUpdate();
 	Server->SendAndDeliverTo(Client, false);
 	Server->PostSendUpdate();
 
@@ -1518,7 +1518,7 @@ UE_NET_TEST_FIXTURE(FReplicationSystemServerClientTestFixture, TestTornedOffObje
 	ServerObject->IntA = 3;
 
 	// Send and deliver packet
-	Server->PreSendUpdate();
+	Server->NetUpdate();
 	Server->SendAndDeliverTo(Client, true);
 	Server->PostSendUpdate();
 
@@ -1546,7 +1546,7 @@ UE_NET_TEST_FIXTURE(FReplicationSystemServerClientTestFixture, TestImmediateTear
 	ServerSubObject->IntA = 1;
 
 	// Send and deliver packet
-	Server->PreSendUpdate();
+	Server->NetUpdate();
 	Server->SendAndDeliverTo(Client, true);
 	Server->PostSendUpdate();
 
@@ -1566,7 +1566,7 @@ UE_NET_TEST_FIXTURE(FReplicationSystemServerClientTestFixture, TestImmediateTear
 	Server->ReplicationBridge->EndReplication(ServerObject, EEndReplicationFlags::TearOff);
 
 	// Send and deliver packet
-	Server->PreSendUpdate();
+	Server->NetUpdate();
 	Server->SendAndDeliverTo(Client, true);
 	Server->PostSendUpdate();
 
@@ -1594,7 +1594,7 @@ UE_NET_TEST_FIXTURE(FReplicationSystemServerClientTestFixture, TestImmediateTear
 	ServerSubObject->IntA = 1;
 
 	// Send and deliver packet
-	Server->PreSendUpdate();
+	Server->NetUpdate();
 	Server->SendAndDeliverTo(Client, true);
 	Server->PostSendUpdate();
 
@@ -1614,12 +1614,12 @@ UE_NET_TEST_FIXTURE(FReplicationSystemServerClientTestFixture, TestImmediateTear
 	Server->ReplicationBridge->EndReplication(ServerObject, EEndReplicationFlags::TearOff);
 
 	// Send and do not deliver packet
-	Server->PreSendUpdate();
+	Server->NetUpdate();
 	Server->SendAndDeliverTo(Client, false);
 	Server->PostSendUpdate();
 
 	// Send and deliver packet
-	Server->PreSendUpdate();
+	Server->NetUpdate();
 	Server->SendAndDeliverTo(Client, true);
 	Server->PostSendUpdate();
 
@@ -1640,7 +1640,7 @@ UE_NET_TEST_FIXTURE(FReplicationSystemServerClientTestFixture, TestDroppedCreati
 	UTestReplicatedIrisObject* ServerObject = Server->CreateObject(0,0);
 
 	// Send and deliver packet
-	Server->PreSendUpdate();
+	Server->NetUpdate();
 	Server->SendAndDeliverTo(Client, true);
 	Server->PostSendUpdate();
 
@@ -1652,12 +1652,12 @@ UE_NET_TEST_FIXTURE(FReplicationSystemServerClientTestFixture, TestDroppedCreati
 	UTestReplicatedIrisObject* ServerSubObject = Server->CreateSubObject(ServerObject->NetRefHandle, 0, 0);
 
 	// Send and do not deliver packet
-	Server->PreSendUpdate();
+	Server->NetUpdate();
 	Server->SendAndDeliverTo(Client, false);
 	Server->PostSendUpdate();
 
 	// Send and deliver packet
-	Server->PreSendUpdate();
+	Server->NetUpdate();
 	Server->SendAndDeliverTo(Client, true);
 	Server->PostSendUpdate();
 
@@ -1678,7 +1678,7 @@ UE_NET_TEST_FIXTURE(FReplicationSystemServerClientTestFixture, TestReplicatedDes
 	UTestReplicatedIrisObject* ServerObject = Server->CreateObject(0,0);
 
 	// Update and delay delivery
-	Server->PreSendUpdate();
+	Server->NetUpdate();
 	Server->SendTo(Client);
 	Server->PostSendUpdate();
 
@@ -1686,7 +1686,7 @@ UE_NET_TEST_FIXTURE(FReplicationSystemServerClientTestFixture, TestReplicatedDes
 	Server->ReplicationBridge->EndReplication(ServerObject);
 
 	// Update and delay delivery
-	Server->PreSendUpdate();
+	Server->NetUpdate();
 	Server->SendTo(Client);
 	Server->PostSendUpdate();
 
@@ -1713,7 +1713,7 @@ UE_NET_TEST_FIXTURE(FReplicationSystemServerClientTestFixture, TestReplicatedSub
 	UTestReplicatedIrisObject* ServerObject = Server->CreateObject(0,0);
 
 	// Replicate object
-	Server->PreSendUpdate();
+	Server->NetUpdate();
 	Server->SendAndDeliverTo(Client, true);
 	Server->PostSendUpdate();
 
@@ -1721,7 +1721,7 @@ UE_NET_TEST_FIXTURE(FReplicationSystemServerClientTestFixture, TestReplicatedSub
 	UTestReplicatedIrisObjectWithNoReplicatedMembers* ServerSubObject = Server->CreateSubObject<UTestReplicatedIrisObjectWithNoReplicatedMembers>(ServerObject->NetRefHandle);
 
 	// Update and delay delivery
-	Server->PreSendUpdate();
+	Server->NetUpdate();
 	Server->SendTo(Client);
 	Server->PostSendUpdate();
 
@@ -1729,7 +1729,7 @@ UE_NET_TEST_FIXTURE(FReplicationSystemServerClientTestFixture, TestReplicatedSub
 	Server->ReplicationBridge->EndReplication(ServerSubObject);
 
 	// Update and delay delivery
-	Server->PreSendUpdate();
+	Server->NetUpdate();
 	Server->SendTo(Client);
 	Server->PostSendUpdate();
 
@@ -1830,7 +1830,7 @@ UE_NET_TEST_FIXTURE(FReplicationSystemServerClientTestFixture, TestTearOffObject
 	// Spawn object on server
 	UTestReplicatedIrisObjectWithNoReplicatedMembers* ServerObject = Server->CreateObject<UTestReplicatedIrisObjectWithNoReplicatedMembers>();
 
-	Server->PreSendUpdate();
+	Server->NetUpdate();
 	Server->SendAndDeliverTo(Client, false);
 	Server->PostSendUpdate();
 
@@ -1838,7 +1838,7 @@ UE_NET_TEST_FIXTURE(FReplicationSystemServerClientTestFixture, TestTearOffObject
 	Server->ReplicationBridge->EndReplication(ServerObject, EEndReplicationFlags::TearOff);
 
 	// Trigger the next update but avoid sending any data so that we keep the object in the PendingCreation state while we flush the Handles PendingTearOff Array which occurs in PostSendUpdate
-	Server->PreSendUpdate();
+	Server->NetUpdate();
 	Server->PostSendUpdate();
 }
 
@@ -1857,7 +1857,7 @@ UE_NET_TEST_FIXTURE(FReplicationSystemServerClientTestFixture, TestTearOffSubObj
 	UTestReplicatedIrisObjectWithNoReplicatedMembers* ServerSubObject = Server->CreateSubObject<UTestReplicatedIrisObjectWithNoReplicatedMembers>(ServerObject->NetRefHandle);
 
 	// Update and drop
-	Server->PreSendUpdate();
+	Server->NetUpdate();
 	Server->SendAndDeliverTo(Client, false);
 	Server->PostSendUpdate();
 
@@ -1865,7 +1865,7 @@ UE_NET_TEST_FIXTURE(FReplicationSystemServerClientTestFixture, TestTearOffSubObj
 	Server->ReplicationBridge->EndReplication(ServerSubObject, EEndReplicationFlags::TearOff);
 
 	// Trigger the next update but avoid sending any data so that we keep the sub-object in the PendingCreation state while we flush the Handles PendingTearOff Array which occurs in PostSendUpdate
-	Server->PreSendUpdate();
+	Server->NetUpdate();
 	Server->PostSendUpdate();
 }
 
@@ -1888,7 +1888,7 @@ UE_NET_TEST_FIXTURE(FReplicationSystemServerClientTestFixture, TestTearOffNextUp
 	ServerSubObject->IntA = 1;
 
 	// Send and deliver packet
-	Server->PreSendUpdate();
+	Server->NetUpdate();
 	Server->SendAndDeliverTo(Client, true);
 	Server->PostSendUpdate();
 
@@ -1908,7 +1908,7 @@ UE_NET_TEST_FIXTURE(FReplicationSystemServerClientTestFixture, TestTearOffNextUp
 	Server->ReplicationSystem->TearOffNextUpdate(ServerObject->NetRefHandle);
 
 	// Send and deliver packet
-	Server->PreSendUpdate();
+	Server->NetUpdate();
 	Server->SendAndDeliverTo(Client, true);
 	Server->PostSendUpdate();
 
@@ -1940,7 +1940,7 @@ UE_NET_TEST_FIXTURE(FReplicationSystemServerClientTestFixture, TestTearOffNextUp
 	const FInternalNetRefIndex SubObjectObjectInternalIndex = NetRefHandleManager->GetInternalIndex(ServerSubObject->NetRefHandle);
 
 	// Trigger presend without send to add the objects to scope
-	Server->PreSendUpdate();
+	Server->NetUpdate();
 	Server->PostSendUpdate();
 
 	UE_NET_ASSERT_EQ(uint16(1), NetRefHandleManager->GetNetObjectRefCount(ServerObjectInternalIndex));
@@ -1950,13 +1950,13 @@ UE_NET_TEST_FIXTURE(FReplicationSystemServerClientTestFixture, TestTearOffNextUp
 	Server->ReplicationSystem->TearOffNextUpdate(ServerObject->NetRefHandle);
 
 	// Update logic, object should be removed from scope but still exist as pending create in
-	Server->PreSendUpdate();
+	Server->NetUpdate();
 	Server->PostSendUpdate();
 
 	UE_NET_ASSERT_EQ(uint16(1), NetRefHandleManager->GetNetObjectRefCount(ServerObjectInternalIndex));
 	UE_NET_ASSERT_EQ(uint16(1), NetRefHandleManager->GetNetObjectRefCount(SubObjectObjectInternalIndex));
 
-	Server->PreSendUpdate();
+	Server->NetUpdate();
 	Server->SendAndDeliverTo(Client, true);
 	Server->PostSendUpdate();
 
@@ -1983,7 +1983,7 @@ UE_NET_TEST_FIXTURE(FReplicationSystemServerClientTestFixture, TestReplicatedObj
 	UE_NET_ASSERT_TRUE(ServerHandle.IsValid());
 
 	// Send and deliver packet
-	Server->PreSendUpdate();
+	Server->NetUpdate();
 	Server->SendAndDeliverTo(Client, true);
 	Server->PostSendUpdate();
 
@@ -1994,7 +1994,7 @@ UE_NET_TEST_FIXTURE(FReplicationSystemServerClientTestFixture, TestReplicatedObj
 	Server->DestroyObject(ServerObject);
 
 	// Send and deliver packet
-	Server->PreSendUpdate();
+	Server->NetUpdate();
 	Server->SendAndDeliverTo(Client, true);
 	Server->PostSendUpdate();
 
@@ -2018,7 +2018,7 @@ UE_NET_TEST_FIXTURE(FReplicationSystemServerClientTestFixture, TestObjectPollFra
 	UTestReplicatedIrisObject* ServerObjectPolledEveryOtherFrame = Server->CreateObject(Params);
 
 	// Send and deliver packet
-	Server->PreSendUpdate();
+	Server->NetUpdate();
 	Server->SendAndDeliverTo(Client, true);
 	Server->PostSendUpdate();
 
@@ -2040,7 +2040,7 @@ UE_NET_TEST_FIXTURE(FReplicationSystemServerClientTestFixture, TestObjectPollFra
 	ServerObjectPolledEveryOtherFrame->IntA += 1;
 
 	// Send and deliver packet
-	Server->PreSendUpdate();
+	Server->NetUpdate();
 	Server->SendAndDeliverTo(Client, true);
 	Server->PostSendUpdate();
 
@@ -2053,7 +2053,7 @@ UE_NET_TEST_FIXTURE(FReplicationSystemServerClientTestFixture, TestObjectPollFra
 	ServerObjectPolledEveryOtherFrame->IntA += 1;
 
 	// Send and deliver packet
-	Server->PreSendUpdate();
+	Server->NetUpdate();
 	Server->SendAndDeliverTo(Client, true);
 	Server->PostSendUpdate();
 
@@ -2090,7 +2090,7 @@ UE_NET_TEST_FIXTURE(FReplicationSystemServerClientTestFixture, TestClientCanSkip
 		LogIris.SetVerbosity(ELogVerbosity::NoLogging);
 
 		// Send and deliver packet
-		Server->PreSendUpdate();
+		Server->NetUpdate();
 		Server->SendAndDeliverTo(Client, true);
 		Server->PostSendUpdate();
 
@@ -2115,7 +2115,7 @@ UE_NET_TEST_FIXTURE(FReplicationSystemServerClientTestFixture, TestClientCanSkip
 	++ServerObjectB->IntA;
 
 	// Send and deliver packet to verify that client ignores the broken object
-	Server->PreSendUpdate();
+	Server->NetUpdate();
 	Server->SendAndDeliverTo(Client, true);
 	Server->PostSendUpdate();
 
@@ -2129,7 +2129,7 @@ UE_NET_TEST_FIXTURE(FReplicationSystemServerClientTestFixture, TestClientCanSkip
 	ReplicationSystem->AddToGroup(ReplicationSystem->GetNotReplicatedNetObjectGroup(), ServerObjectA->NetRefHandle);
 
 	// Send and deliver packet, the client should now remove the broken object from the list of broken objects
-	Server->PreSendUpdate();
+	Server->NetUpdate();
 	Server->SendAndDeliverTo(Client, true);
 	Server->PostSendUpdate();
 
@@ -2140,7 +2140,7 @@ UE_NET_TEST_FIXTURE(FReplicationSystemServerClientTestFixture, TestClientCanSkip
 	ServerObjectA->bForceFailToInstantiateOnRemote = false;
 
 	// Client should now be able to instantiate the object
-	Server->PreSendUpdate();
+	Server->NetUpdate();
 	Server->SendAndDeliverTo(Client, true);
 	Server->PostSendUpdate();
 
@@ -2181,7 +2181,7 @@ UE_NET_TEST_FIXTURE(FReplicationSystemServerClientTestFixture, TestPartialDequan
 	UTestReplicatedObjectWithRepNotifies* ServerObjectA = Server->CreateObject<UTestReplicatedObjectWithRepNotifies>();
 	
 	// Send and deliver packet
-	Server->PreSendUpdate();
+	Server->NetUpdate();
 	Server->SendAndDeliverTo(Client, true);
 	Server->PostSendUpdate();
 
@@ -2324,7 +2324,7 @@ UE_NET_TEST_FIXTURE(FReplicationSystemServerClientTestFixture, TestReplicationRe
 	UReplicatedTestObject* FirstObject = Server->CreateObject(UTestReplicatedIrisObject::FComponents{});
 	const FNetRefHandle FirstObjectRefHandle = FirstObject->NetRefHandle;
 
-	Server->PreSendUpdate();
+	Server->NetUpdate();
 	Server->SendUpdate(Client->ConnectionIdOnServer);
 	Server->PostSendUpdate();
 
@@ -2332,7 +2332,7 @@ UE_NET_TEST_FIXTURE(FReplicationSystemServerClientTestFixture, TestReplicationRe
 	UReplicatedTestObject* SecondObject = Server->CreateObject(UTestReplicatedIrisObject::FComponents{});
 	const FNetRefHandle SecondObjectRefHandle = SecondObject->NetRefHandle;
 
-	Server->PreSendUpdate();
+	Server->NetUpdate();
 	Server->SendUpdate(Client->ConnectionIdOnServer);
 	Server->PostSendUpdate();
 
@@ -2344,7 +2344,7 @@ UE_NET_TEST_FIXTURE(FReplicationSystemServerClientTestFixture, TestReplicationRe
 	UE_NET_ASSERT_FALSE(Client->IsResolvableNetRefHandle(SecondObjectRefHandle));
 
 	// Now we should be able replicate objects again. Retry sending the second object.
-	Server->PreSendUpdate();
+	Server->NetUpdate();
 	Server->SendUpdate(Client->ConnectionIdOnServer);
 	Server->PostSendUpdate();
 
diff --git a/Engine/Plugins/Runtime/ReplicationSystemTestPlugin/Source/Private/Tests/ReplicationSystem/TestStructNetTokenDataStore.cpp b/Engine/Plugins/Runtime/ReplicationSystemTestPlugin/Source/Private/Tests/ReplicationSystem/TestStructNetTokenDataStore.cpp
--- a/Engine/Plugins/Runtime/ReplicationSystemTestPlugin/Source/Private/Tests/ReplicationSystem/TestStructNetTokenDataStore.cpp
+++ b/Engine/Plugins/Runtime/ReplicationSystemTestPlugin/Source/Private/Tests/ReplicationSystem/TestStructNetTokenDataStore.cpp
@@ -126,7 +126,7 @@ UE_NET_TEST_FIXTURE(FTestStructNetTokensFixture, StructAsNetToken)
 	}
 
 	// Send and deliver packet
-	Server->PreSendUpdate();
+	Server->NetUpdate();
 	Server->SendAndDeliverTo(Client, true);
 	Server->PostSendUpdate();
 
@@ -146,7 +146,7 @@ UE_NET_TEST_FIXTURE(FTestStructNetTokensFixture, StructAsNetTokenPropertyWithIri
 	ServerObject->NetTokenStoreStruct = StructA;
 
 	// Send and deliver packet
-	Server->PreSendUpdate();
+	Server->NetUpdate();
 	Server->SendAndDeliverTo(Client, true);
 	Server->PostSendUpdate();
 
@@ -181,7 +181,7 @@ UE_NET_TEST_FIXTURE(FTestStructNetTokensFixture, StructAsNetTokenUsingLastResort
 	ServerObject->DerivedNetTokenStoreStruct = StructA;
 
 	// Send and deliver packet
-	Server->PreSendUpdate();
+	Server->NetUpdate();
 	Server->SendAndDeliverTo(Client, true);
 	Server->PostSendUpdate();
 
diff --git a/Engine/Plugins/Runtime/ReplicationSystemTestPlugin/Source/Private/Tests/Serialization/TestObjectNetSerializer.cpp b/Engine/Plugins/Runtime/ReplicationSystemTestPlugin/Source/Private/Tests/Serialization/TestObjectNetSerializer.cpp
--- a/Engine/Plugins/Runtime/ReplicationSystemTestPlugin/Source/Private/Tests/Serialization/TestObjectNetSerializer.cpp
+++ b/Engine/Plugins/Runtime/ReplicationSystemTestPlugin/Source/Private/Tests/Serialization/TestObjectNetSerializer.cpp
@@ -33,7 +33,7 @@ UE_NET_TEST_FIXTURE(FReplicationSystemServerClientTestFixture, TestObjectReferen
 	ObjectD->RawObjectPtrRef = ObjectA;
 
 	// Send and deliver packet
-	Server->PreSendUpdate();
+	Server->NetUpdate();
 	Server->SendAndDeliverTo(Client, true);
 	Server->PostSendUpdate();
 
@@ -73,7 +73,7 @@ UE_NET_TEST_FIXTURE(FReplicationSystemServerClientTestFixture, TestCircularRefer
 	ObjectD->RawObjectPtrRef = ObjectA;
 
 	// Send and deliver packet
-	Server->PreSendUpdate();
+	Server->NetUpdate();
 	Server->SendAndDeliverTo(Client, true);
 	Server->PostSendUpdate();
 
@@ -116,27 +116,27 @@ UE_NET_TEST_FIXTURE(FReplicationSystemServerClientTestFixture, TestCircularRefer
 	ObjectD->RawObjectPtrRef = ObjectA;
 
 	// Send and deliver packets
-	Server->PreSendUpdate();
+	Server->NetUpdate();
 	Server->SendAndDeliverTo(Client, true);
 	Server->PostSendUpdate();
 
 	// Send and deliver packets
-	Server->PreSendUpdate();
+	Server->NetUpdate();
 	Server->SendAndDeliverTo(Client, true);
 	Server->PostSendUpdate();
 
 	// Send and deliver packets
-	Server->PreSendUpdate();
+	Server->NetUpdate();
 	Server->SendAndDeliverTo(Client, true);
 	Server->PostSendUpdate();
 
 	// Send and deliver packets
-	Server->PreSendUpdate();
+	Server->NetUpdate();
 	Server->SendAndDeliverTo(Client, true);
 	Server->PostSendUpdate();
 
 	// Send and deliver packets
-	Server->PreSendUpdate();
+	Server->NetUpdate();
 	Server->SendAndDeliverTo(Client, true);
 	Server->PostSendUpdate();
 
@@ -171,7 +171,7 @@ UE_NET_TEST_FIXTURE(FReplicationSystemServerClientTestFixture, TestSoftObjectRef
 	ObjectA->SoftObjectPtrRef = ObjectB;
 
 	// Send and deliver packet
-	Server->PreSendUpdate();
+	Server->NetUpdate();
 	Server->SendAndDeliverTo(Client, DeliverPacket);
 	Server->PostSendUpdate();
 
diff --git a/Engine/Plugins/Runtime/ReplicationSystemTestPlugin/Source/Private/Tests/Serialization/TestPolymorphicStructNetSerializer.cpp b/Engine/Plugins/Runtime/ReplicationSystemTestPlugin/Source/Private/Tests/Serialization/TestPolymorphicStructNetSerializer.cpp
--- a/Engine/Plugins/Runtime/ReplicationSystemTestPlugin/Source/Private/Tests/Serialization/TestPolymorphicStructNetSerializer.cpp
+++ b/Engine/Plugins/Runtime/ReplicationSystemTestPlugin/Source/Private/Tests/Serialization/TestPolymorphicStructNetSerializer.cpp
@@ -727,7 +727,7 @@ UE_NET_TEST_FIXTURE(FTestPolymorphicArrayStructNetSerializerFixture, TestPolymor
 	ServerObject->PolyStruct.GetAs<FExamplePolymorphicStructB>().SomeFloat = 12.0f;
 
 	// Replicate
-	Server->PreSendUpdate();
+	Server->NetUpdate();
 	Server->SendAndDeliverTo(Client, DeliverPacket);
 	Server->PostSendUpdate();
 	
@@ -740,7 +740,7 @@ UE_NET_TEST_FIXTURE(FTestPolymorphicArrayStructNetSerializerFixture, TestPolymor
 	ServerObject->PolyStruct.GetAs<FExamplePolymorphicStructB>().SomeFloat += 1.0f;
 
 	// Replicate
-	Server->PreSendUpdate();
+	Server->NetUpdate();
 	Server->SendAndDeliverTo(Client, DeliverPacket);
 	Server->PostSendUpdate();
 
@@ -752,7 +752,7 @@ UE_NET_TEST_FIXTURE(FTestPolymorphicArrayStructNetSerializerFixture, TestPolymor
 	ServerObject->PolyStruct.GetAs<FExamplePolymorphicStructD>().SomeValue = 100;
 
 	// Replicate
-	Server->PreSendUpdate();
+	Server->NetUpdate();
 	Server->SendAndDeliverTo(Client, DeliverPacket);
 	Server->PostSendUpdate();
 
@@ -787,7 +787,7 @@ UE_NET_TEST_FIXTURE(FTestPolymorphicArrayStructNetSerializerFixture, TestPolymor
 	ServerObject->PolyStructFastArray.Edit()[3].PolyStruct.GetAs<FExamplePolymorphicStructD>().SomeValue = 100;
 	
 	// Replicate
-	Server->PreSendUpdate();
+	Server->NetUpdate();
 	Server->SendAndDeliverTo(Client, DeliverPacket);
 	Server->PostSendUpdate();
 	
@@ -801,7 +801,7 @@ UE_NET_TEST_FIXTURE(FTestPolymorphicArrayStructNetSerializerFixture, TestPolymor
 	ServerObject->PolyStructFastArray.Edit()[3].PolyStruct.GetAs<FExamplePolymorphicStructD>().SomeValue += 3;
 
 	// Replicate
-	Server->PreSendUpdate();
+	Server->NetUpdate();
 	Server->SendAndDeliverTo(Client, DeliverPacket);
 	Server->PostSendUpdate();
 
@@ -813,7 +813,7 @@ UE_NET_TEST_FIXTURE(FTestPolymorphicArrayStructNetSerializerFixture, TestPolymor
 	ServerObject->PolyStructFastArray.Edit()[2].PolyStruct.GetAs<FExamplePolymorphicStructD>().SomeValue = 1;
 
 	// Replicate
-	Server->PreSendUpdate();
+	Server->NetUpdate();
 	Server->SendAndDeliverTo(Client, DeliverPacket);
 	Server->PostSendUpdate();
 
diff --git a/Engine/Plugins/Runtime/ReplicationSystemTestPlugin/Source/Private/Tests/Serialization/TestRemoteObjectReferenceNetSerializer.cpp b/Engine/Plugins/Runtime/ReplicationSystemTestPlugin/Source/Private/Tests/Serialization/TestRemoteObjectReferenceNetSerializer.cpp
--- a/Engine/Plugins/Runtime/ReplicationSystemTestPlugin/Source/Private/Tests/Serialization/TestRemoteObjectReferenceNetSerializer.cpp
+++ b/Engine/Plugins/Runtime/ReplicationSystemTestPlugin/Source/Private/Tests/Serialization/TestRemoteObjectReferenceNetSerializer.cpp
@@ -63,7 +63,7 @@ UE_NET_TEST_FIXTURE(FRPCTestFixture, TestRemoteObjectReference)
 	TestObject->RemoteReferenceProperty = FRemoteObjectReference(ReferencedObject);
 	TestObject->RemoteRPCWithRemoteReferenceParam(FRemoteObjectReference(ReferencedObject));
 	
-	Server->PreSendUpdate();
+	Server->NetUpdate();
 	Server->SendAndDeliverTo(Client, true);
 	Server->PostSendUpdate();
 
@@ -107,7 +107,7 @@ UE_NET_TEST_FIXTURE(FRPCTestFixture, TestRemoteObjectReferencePaths)
 	TestObject->RemoteReferenceProperty = FRemoteObjectReference(ServerReferencedObject);
 	TestObject->RemoteRPCWithRemoteReferenceParam(FRemoteObjectReference(ServerReferencedObject));
 
-	Server->PreSendUpdate();
+	Server->NetUpdate();
 	Server->SendTo(Client);
 	Server->PostSendUpdate();
 
diff --git a/Engine/Plugins/Runtime/ReplicationSystemTestPlugin/Source/Private/Tests/Serialization/TestScriptInterfaceNetSerializer.cpp b/Engine/Plugins/Runtime/ReplicationSystemTestPlugin/Source/Private/Tests/Serialization/TestScriptInterfaceNetSerializer.cpp
--- a/Engine/Plugins/Runtime/ReplicationSystemTestPlugin/Source/Private/Tests/Serialization/TestScriptInterfaceNetSerializer.cpp
+++ b/Engine/Plugins/Runtime/ReplicationSystemTestPlugin/Source/Private/Tests/Serialization/TestScriptInterfaceNetSerializer.cpp
@@ -29,7 +29,7 @@ UE_NET_TEST_FIXTURE(FTestScriptInterfaceNetSerializer, TestNoScriptInterface)
 	UTestObjectReferencingScriptInterface* Object = Server->CreateObject<UTestObjectReferencingScriptInterface>();
 
 	// Send and deliver packet
-	Server->PreSendUpdate();
+	Server->NetUpdate();
 	Server->SendAndDeliverTo(Client, true);
 	Server->PostSendUpdate();
 
@@ -56,7 +56,7 @@ UE_NET_TEST_FIXTURE(FTestScriptInterfaceNetSerializer, TestInterfaceToReplicated
 	Object->InterfaceObject = InterfaceObject;
 
 	// Send and deliver packet
-	Server->PreSendUpdate();
+	Server->NetUpdate();
 	Server->SendAndDeliverTo(Client, true);
 	Server->PostSendUpdate();
 
@@ -88,7 +88,7 @@ UE_NET_TEST_FIXTURE(FTestScriptInterfaceNetSerializer, TestInterfaceToNamedObjec
 	Object->InterfaceObject = ObjectWithInterfaceSubObject->DefaultSubobjectWithInterface.Get();
 
 	// Send and deliver packet
-	Server->PreSendUpdate();
+	Server->NetUpdate();
 	Server->SendAndDeliverTo(Client, true);
 	Server->PostSendUpdate();
 
diff --git a/Engine/Source/Runtime/Engine/Private/Net/Iris/ReplicationSystem/EngineReplicationBridge.cpp b/Engine/Source/Runtime/Engine/Private/Net/Iris/ReplicationSystem/EngineReplicationBridge.cpp
--- a/Engine/Source/Runtime/Engine/Private/Net/Iris/ReplicationSystem/EngineReplicationBridge.cpp
+++ b/Engine/Source/Runtime/Engine/Private/Net/Iris/ReplicationSystem/EngineReplicationBridge.cpp
@@ -488,25 +488,12 @@ void UEngineReplicationBridge::StopReplicatingActor(AActor* Actor, EEndPlayReaso
 
 	UE_LOG(LogIrisBridge, Verbose, TEXT("StopReplicatingActor %s. Reason %s "), *PrintObjectFromNetRefHandle(RefHandle), *UEnum::GetValueAsString(TEXT("Engine.EEndPlayReason"), EndPlayReason));
 	ensureMsgf(IsValid(Actor), TEXT("StopReplicatingActor called on invalid actor tied to handle: %s."), *PrintObjectFromNetRefHandle(RefHandle));
-	
-	EEndReplicationFlags Flags = EEndReplicationFlags::None;
-	const bool bShouldDestroyObject = EndPlayReason == EEndPlayReason::Destroyed;		
-	if (bShouldDestroyObject)
-	{ 
-		Flags |= EEndReplicationFlags::Destroy | EEndReplicationFlags::DestroyNetHandle | EEndReplicationFlags::ClearNetPushId;
-	}
 
-	// If we are shutting down or the actor is a nettemporary we do not need to validate that we are not detaching remote instances by accident.
-	const bool bIsShuttingDown = (EndPlayReason == EEndPlayReason::EndPlayInEditor) || (EndPlayReason == EEndPlayReason::Quit) || NetDriver->GetPendingDestruction();
-	// Streaming out a level can happen prior to all actors in it being destroyed/ending replication through replication so we don't want to validate they're detached by accident.
-	const bool bIsStreamingOutLevel = (EndPlayReason == EEndPlayReason::RemovedFromWorld);
-	if (bIsShuttingDown || bIsStreamingOutLevel || Actor->bNetTemporary)
-	{
-		Flags |= EEndReplicationFlags::SkipPendingEndReplicationValidation;
-	}
+	const bool bIsServer = GetReplicationSystem()->IsServer();
+	const bool bIsObjectDestroyed = EndPlayReason == EEndPlayReason::Destroyed;
 
-	const bool bShouldCreateDestructionInfo = RefHandle.IsStatic() && bShouldDestroyObject && GetReplicationSystem()->IsServer();
-	if (bShouldCreateDestructionInfo)
+	// Store destruction info of static actors so their destruction can be replicated when a client loads their level in the future.
+	if (bIsServer && bIsObjectDestroyed && RefHandle.IsStatic())
 	{
 		const UObjectReplicationBridge::FDestructionParameters EndReplicationParameters
 		{
@@ -515,14 +502,39 @@ void UEngineReplicationBridge::StopReplicatingActor(AActor* Actor, EEndPlayReaso
 			.bUseDistanceBasedPrioritization = (Actor->bAlwaysRelevant == false),
 			.NetFactoryId = ActorFactoryId,
 		};
-
-		// Store this static actor destruction info so it can be replicated and destroyed when a client loads this level in the future.
 		StoreDestructionInfo(RefHandle, EndReplicationParameters);
 	}
-
+		
+	EEndReplicationFlags Flags = ConvertEndPlayIntoEndReplication(EndPlayReason);
+	if (Actor->bNetTemporary)
+	{
+		Flags |= EEndReplicationFlags::SkipPendingEndReplicationValidation;
+	}
 	StopReplicatingNetRefHandle(RefHandle, Flags);
 }
 
+EEndReplicationFlags UEngineReplicationBridge::ConvertEndPlayIntoEndReplication(EEndPlayReason::Type EndPlayReason) const
+{
+	EEndReplicationFlags EndReplicationFlags = EEndReplicationFlags::None;
+	if (EndPlayReason == EEndPlayReason::Destroyed)
+	{
+		EndReplicationFlags |= EEndReplicationFlags::Destroy | EEndReplicationFlags::DestroyNetHandle | EEndReplicationFlags::ClearNetPushId;
+	}
+
+	// If we are shutting down we do not need to validate that we are not detaching remote instances by accident.
+	const bool bIsShuttingDown = (EndPlayReason == EEndPlayReason::EndPlayInEditor) || (EndPlayReason == EEndPlayReason::Quit) || NetDriver->GetPendingDestruction();
+
+	// Streaming out a level can happen prior to all actors in it being destroyed/ending replication through replication so we don't want to validate they're detached by accident.
+	const bool bIsStreamingOutLevel = (EndPlayReason == EEndPlayReason::RemovedFromWorld);
+
+	if (bIsShuttingDown || bIsStreamingOutLevel)
+	{
+		EndReplicationFlags |= EEndReplicationFlags::SkipPendingEndReplicationValidation;
+	}
+
+	return EndReplicationFlags;
+}
+
 void UEngineReplicationBridge::StopReplicatingComponent(UActorComponent* ActorComponent, EEndReplicationFlags EndReplicationFlags)
 {
 	using namespace UE::Net;
diff --git a/Engine/Source/Runtime/Engine/Private/NetDriver.cpp b/Engine/Source/Runtime/Engine/Private/NetDriver.cpp
--- a/Engine/Source/Runtime/Engine/Private/NetDriver.cpp
+++ b/Engine/Source/Runtime/Engine/Private/NetDriver.cpp
@@ -1061,7 +1061,7 @@ void UNetDriver::TickFlush(float DeltaSeconds)
 		{
 			UpdateIrisReplicationViews();
 			SendClientMoveAdjustments();
-			ReplicationSystem->PreSendUpdate(UReplicationSystem::FSendUpdateParams { .SendPass = UE::Net::EReplicationSystemSendPass::TickFlush, .DeltaSeconds = DeltaSeconds });
+			ReplicationSystem->NetUpdate(DeltaSeconds);
 		}
 		else
 #endif // UE_WITH_IRIS
@@ -1091,7 +1091,7 @@ void UNetDriver::TickFlush(float DeltaSeconds)
 		if (ReplicationSystem)
 		{
 			UpdateIrisReplicationViews();
-			ReplicationSystem->PreSendUpdate(UReplicationSystem::FSendUpdateParams { .SendPass = UE::Net::EReplicationSystemSendPass::TickFlush, .DeltaSeconds = DeltaSeconds });
+			ReplicationSystem->NetUpdate(DeltaSeconds);
 		}
 	}
 #endif // UE_WITH_IRIS
@@ -2187,7 +2187,22 @@ void UNetDriver::UnregisterTickEvents(class UWorld* InWorld)
 void UNetDriver::InternalTickDispatch(float DeltaSeconds)
 {
 	TGuardValue<bool> GuardInNetTick(bInTick, true);
+
+#if UE_WITH_IRIS
+	if (ReplicationSystem)
+	{
+		ReplicationSystem->PreReceiveUpdate();
+	}
+#endif
+
 	TickDispatch(DeltaSeconds);
+
+#if UE_WITH_IRIS
+	if (ReplicationSystem)
+	{
+		ReplicationSystem->PostReceiveUpdate();
+	}
+#endif
 }
 
 void UNetDriver::InternalTickFlush(float DeltaSeconds)
@@ -7456,40 +7471,36 @@ void UNetDriver::SendClientMoveAdjustments()
 
 void UNetDriver::PostDispatchSendUpdate()
 {
-	if (ReplicationSystem)
+	if (!ReplicationSystem)
 	{
-		if (!IsKnownChannelName(NAME_DataStream))
-		{
-			return;
-		}
+		return;
+	}
+	
+	if (!IsKnownChannelName(NAME_DataStream))
+	{
+		return;
+	}
 
-		const UReplicationSystem::FSendUpdateParams Params
-		{ 
-			.SendPass = UE::Net::EReplicationSystemSendPass::PostTickDispatch 
-			//no delta time passed for PostTickDispatch
-		};
+	ReplicationSystem->TickPostReceive();
 
-		ReplicationSystem->PreSendUpdate(Params);
+	ReplicationSystem->SendUpdate([this](TArrayView<uint32> ConnectionsToSend)
+	{
+		const int32 DataStreamChannelIndex = ChannelDefinitionMap[NAME_DataStream].StaticChannelIndex;
 
-		ReplicationSystem->SendUpdate([this](TArrayView<uint32> ConnectionsToSend)
+		for (uint32 ConnId : ConnectionsToSend)
 		{
-			const int32 DataStreamChannelIndex = ChannelDefinitionMap[NAME_DataStream].StaticChannelIndex;
-
-			for (uint32 ConnId : ConnectionsToSend)
+			UNetConnection* NetConnection = GetConnectionByHandle(UE::Net::FConnectionHandle(ConnId));
+			if (NetConnection && NetConnection->Channels.IsValidIndex(DataStreamChannelIndex))
 			{
-				UNetConnection* NetConnection = GetConnectionByHandle(UE::Net::FConnectionHandle(ConnId));
-				if (NetConnection && NetConnection->Channels.IsValidIndex(DataStreamChannelIndex))
+				if (UDataStreamChannel* DataStreamChannel = Cast<UDataStreamChannel>(NetConnection->Channels[DataStreamChannelIndex]))
 				{
-					if (UDataStreamChannel* DataStreamChannel = Cast<UDataStreamChannel>(NetConnection->Channels[DataStreamChannelIndex]))
-					{
-						DataStreamChannel->PostTickDispatch();
-					}
+					DataStreamChannel->PostTickDispatch();
 				}
 			}
-		});
+		}
+	});
 
-		ReplicationSystem->PostSendUpdate();
-	}
+	ReplicationSystem->PostSendUpdate();
 }
 
 #endif // UE_WITH_IRIS
diff --git a/Engine/Source/Runtime/Engine/Public/Net/Iris/ReplicationSystem/EngineReplicationBridge.h b/Engine/Source/Runtime/Engine/Public/Net/Iris/ReplicationSystem/EngineReplicationBridge.h
--- a/Engine/Source/Runtime/Engine/Public/Net/Iris/ReplicationSystem/EngineReplicationBridge.h
+++ b/Engine/Source/Runtime/Engine/Public/Net/Iris/ReplicationSystem/EngineReplicationBridge.h
@@ -78,8 +78,11 @@ public:
 	/** Begin replication of an actor and its registered ActorComponents and SubObjects. */
 	ENGINE_API FNetRefHandle StartReplicatingActor(AActor* Instance, const FActorReplicationParams& Params);
 
-	/** Stop replicating an actor. Will destroy handle for actor and registered subobjects. */
+	/** Stop replicating an actor. This will destroy the handle of the actor and those of his subobjects. */
 	ENGINE_API void StopReplicatingActor(AActor* Actor, EEndPlayReason::Type EndPlayReason);
+
+	/** Convert EndPlayReason types into the proper EndReplicationFlags */
+	ENGINE_API EEndReplicationFlags ConvertEndPlayIntoEndReplication(EEndPlayReason::Type EndPlayReason) const;
 		
 	/**
 	 * Begin replication of an ActorComponent and its registered SubObjects, 
diff --git a/Engine/Source/Runtime/Experimental/Iris/Core/Private/Iris/ReplicationSystem/NetObjectFactory.cpp b/Engine/Source/Runtime/Experimental/Iris/Core/Private/Iris/ReplicationSystem/NetObjectFactory.cpp
--- a/Engine/Source/Runtime/Experimental/Iris/Core/Private/Iris/ReplicationSystem/NetObjectFactory.cpp
+++ b/Engine/Source/Runtime/Experimental/Iris/Core/Private/Iris/ReplicationSystem/NetObjectFactory.cpp
@@ -36,9 +36,9 @@ void UNetObjectFactory::Deinit()
 	Bridge = nullptr;
 }
 
-void UNetObjectFactory::PostTickDispatch()
+void UNetObjectFactory::PostReceiveUpdate()
 {
-	OnPostTickDispatch();
+	OnPostReceiveUpdate();
 }
 
 TUniquePtr<UE::Net::FNetObjectCreationHeader> UNetObjectFactory::CreateHeader(UE::Net::FNetRefHandle Handle, UE::Net::FReplicationProtocolIdentifier ProtocolId)
diff --git a/Engine/Source/Runtime/Experimental/Iris/Core/Private/Iris/ReplicationSystem/ObjectReplicationBridge.cpp b/Engine/Source/Runtime/Experimental/Iris/Core/Private/Iris/ReplicationSystem/ObjectReplicationBridge.cpp
--- a/Engine/Source/Runtime/Experimental/Iris/Core/Private/Iris/ReplicationSystem/ObjectReplicationBridge.cpp
+++ b/Engine/Source/Runtime/Experimental/Iris/Core/Private/Iris/ReplicationSystem/ObjectReplicationBridge.cpp
@@ -1113,11 +1113,11 @@ void UObjectReplicationBridge::OnPostSendUpdate()
 	bBlockStartRootObjectReplication = false;
 }
 
-void UObjectReplicationBridge::OnPostTickDispatch()
+void UObjectReplicationBridge::OnPostReceiveUpdate()
 {
 	for (UNetObjectFactory* NetFactory : NetObjectFactories)
 	{
-		NetFactory->PostTickDispatch();
+		NetFactory->PostReceiveUpdate();
 	}
 }
 
diff --git a/Engine/Source/Runtime/Experimental/Iris/Core/Private/Iris/ReplicationSystem/ReplicationBridge.cpp b/Engine/Source/Runtime/Experimental/Iris/Core/Private/Iris/ReplicationSystem/ReplicationBridge.cpp
--- a/Engine/Source/Runtime/Experimental/Iris/Core/Private/Iris/ReplicationSystem/ReplicationBridge.cpp
+++ b/Engine/Source/Runtime/Experimental/Iris/Core/Private/Iris/ReplicationSystem/ReplicationBridge.cpp
@@ -31,6 +31,7 @@
 DEFINE_LOG_CATEGORY(LogIrisBridge)
 
 #define UE_LOG_REPLICATIONBRIDGE(Category, Format, ...)  UE_LOG(LogIrisBridge, Category, TEXT("ReplicationBridge(%u)::") Format, ReplicationSystem->GetId(), ##__VA_ARGS__)
+#define UE_LOG_BRIDGEID(Category, Verbosity, Format, ...)  UE_LOG(Category, Verbosity, TEXT("ReplicationBridge(%u)::") Format, ReplicationSystem->GetId(), ##__VA_ARGS__)
 
 static bool bEnableFlushReliableRPCOnDestroy = true;
 static FAutoConsoleVariableRef CVarEnableFlushReliableRPCOnDestroy(
@@ -65,44 +66,34 @@ UReplicationBridge::UReplicationBridge()
 {
 }
 
-bool UReplicationBridge::WriteNetRefHandleCreationInfo(FReplicationBridgeSerializationContext& Context, FNetRefHandle Handle)
+void UReplicationBridge::PreReceiveUpdate()
 {
-	return true;
+	check(bInReceiveUpdate == false);
+	bInReceiveUpdate = true;
 }
 
-bool UReplicationBridge::CacheNetRefHandleCreationInfo(FNetRefHandle Handle)
+void UReplicationBridge::PostReceiveUpdate()
 {
-	return false;
-}
+	check(bInReceiveUpdate == true);
+	bInReceiveUpdate = false;
 
-FReplicationBridgeCreateNetRefHandleResult UReplicationBridge::CreateNetRefHandleFromRemote(FNetRefHandle RootObjectOfSubObject, FNetRefHandle WantedNetHandle, FReplicationBridgeSerializationContext& Context)
-{
-	return FReplicationBridgeCreateNetRefHandleResult();
-};
-
-void UReplicationBridge::PreSendUpdateSingleHandle(FNetRefHandle Handle)
-{
+	OnPostReceiveUpdate();
 }
 
-void UReplicationBridge::PreSendUpdate()
-{
-}
-
-void UReplicationBridge::OnStartPreSendUpdate()
-{
-}
-
-void UReplicationBridge::OnPostSendUpdate()
+bool UReplicationBridge::WriteNetRefHandleCreationInfo(FReplicationBridgeSerializationContext& Context, FNetRefHandle Handle)
 {
+	return true;
 }
 
-void UReplicationBridge::OnPostTickDispatch()
+bool UReplicationBridge::CacheNetRefHandleCreationInfo(FNetRefHandle Handle)
 {
+	return false;
 }
 
-void UReplicationBridge::UpdateInstancesWorldLocation()
+FReplicationBridgeCreateNetRefHandleResult UReplicationBridge::CreateNetRefHandleFromRemote(FNetRefHandle RootObjectOfSubObject, FNetRefHandle WantedNetHandle, FReplicationBridgeSerializationContext& Context)
 {
-}
+	return FReplicationBridgeCreateNetRefHandleResult();
+};
 
 void UReplicationBridge::SubObjectCreatedFromReplication(UE::Net::Private::FInternalNetRefIndex RootObjectIndex, FNetRefHandle SubObjectRefHandle)
 {
diff --git a/Engine/Source/Runtime/Experimental/Iris/Core/Private/Iris/ReplicationSystem/ReplicationSystem.cpp b/Engine/Source/Runtime/Experimental/Iris/Core/Private/Iris/ReplicationSystem/ReplicationSystem.cpp
--- a/Engine/Source/Runtime/Experimental/Iris/Core/Private/Iris/ReplicationSystem/ReplicationSystem.cpp
+++ b/Engine/Source/Runtime/Experimental/Iris/Core/Private/Iris/ReplicationSystem/ReplicationSystem.cpp
@@ -581,9 +581,14 @@ public:
 		NetBlobManager.ProcessOOBNetObjectAttachmentSendQueue(ConnectionsPendingPostTickDispatchSend);
 	}
 
-	void PostTickDispatch()
+	void PreReceiveUpdate()
 	{
-		ReplicationSystemInternal.GetReplicationBridge()->OnPostTickDispatch();
+		ReplicationSystemInternal.GetReplicationBridge()->PreReceiveUpdate();
+	}
+
+	void PostReceiveUpdate()
+	{
+		ReplicationSystemInternal.GetReplicationBridge()->PostReceiveUpdate();
 	}
 
 	void ResetNetObjectAttachmentSendQueue()
@@ -832,124 +837,117 @@ const UE::Net::Private::FReplicationSystemInternal* UReplicationSystem::GetRepli
 	return &Impl->ReplicationSystemInternal;
 }
 
-void UReplicationSystem::PreSendUpdate(const FSendUpdateParams& Params)
+void UReplicationSystem::NetUpdate(float DeltaSeconds)
 {
 	using namespace UE::Net;
 	using namespace UE::Net::Private;
 
-	IRIS_CSV_PROFILER_SCOPE(Iris, ReplicationSystem_PreSendUpdate);
-
-	ensureMsgf(Impl->CurrentSendPass == EReplicationSystemSendPass::Invalid, TEXT("PostSendUpdate was not called after the last PreSendUpdate."));
-	Impl->CurrentSendPass = Params.SendPass;
+	IRIS_CSV_PROFILER_SCOPE(Iris, ReplicationSystem_NetUpdate);
 
 	FReplicationSystemInternal& InternalSys = Impl->ReplicationSystemInternal;
 
-	// Most systems are only updated during the normal TickFlush
-	if (Impl->CurrentSendPass ==  EReplicationSystemSendPass::TickFlush)
-	{
-		ElapsedTime += Params.DeltaSeconds;
+	ElapsedTime += DeltaSeconds;
 
-		// $IRIS TODO. There may be some throttling of connections to tick that we should take into account.
-		const FNetBitArrayView& ReplicatingConnections = MakeNetBitArrayView(Impl->ReplicationSystemInternal.GetConnections().GetValidConnections());
+	ensureMsgf(Impl->CurrentSendPass == EReplicationSystemSendPass::Invalid, TEXT("PostSendUpdate was not called after the last Tick."));
+	Impl->CurrentSendPass = EReplicationSystemSendPass::TickFlush;
 
-#if UE_NET_IRIS_CSV_STATS
-		{
-			FNetSendStats& SendStats = InternalSys.GetSendStats();
-			SendStats.Reset();
-			SendStats.SetNumberOfReplicatingConnections(ReplicatingConnections.CountSetBits());
-		}
-#endif
-		InternalSys.GetAccumulatedReplicationStats().Accumulate(InternalSys.GetTickReplicationStats());
-		InternalSys.GetTickReplicationStats() = {};
-	
-		// Force a integrity check of all replicated instances
-		if (bDoCollectGarbage || ReplicationSystemCVars::bForcePruneBeforeUpdate)
-		{
-			CollectGarbage();
-		}
+	// $IRIS TODO. There may be some throttling of connections to tick that we should take into account.
+	const FNetBitArrayView& ReplicatingConnections = MakeNetBitArrayView(Impl->ReplicationSystemInternal.GetConnections().GetValidConnections());
 
-		// DataStream presend update, this is similar to channel tick()
-		UDataStream::FUpdateParameters DataStreamUpdateParams = { .UpdateType = UDataStream::EUpdateType::PreSendUpdate };
-		Impl->UpdateDataStreams(DataStreamUpdateParams);
+	#if UE_NET_IRIS_CSV_STATS
+	{
+		FNetSendStats& SendStats = InternalSys.GetSendStats();
+		SendStats.Reset();
+		SendStats.SetNumberOfReplicatingConnections(ReplicatingConnections.CountSetBits());
+	}
+	#endif
+	InternalSys.GetAccumulatedReplicationStats().Accumulate(InternalSys.GetTickReplicationStats());
+	InternalSys.GetTickReplicationStats() = {};
+	
+	// Force a integrity check of all replicated instances
+	if (bDoCollectGarbage || ReplicationSystemCVars::bForcePruneBeforeUpdate)
+	{
+		CollectGarbage();
+	}
 
-		if (bAllowObjectReplication)
-		{
-			UE_NET_TRACE_FRAME_STATSCOUNTER(GetId(), ReplicationSystem.ReplicatedObjectCount, InternalSys.GetNetRefHandleManager().GetActiveObjectCount(), ENetTraceVerbosity::Verbose);
+	// DataStream presend update, this is similar to channel tick()
+	UDataStream::FUpdateParameters DataStreamUpdateParams = { .UpdateType = UDataStream::EUpdateType::PreSendUpdate };
+	Impl->UpdateDataStreams(DataStreamUpdateParams);
 
-			// Tell systems we are starting PreSendUpdate
-			Impl->StartPreSendUpdate();
+	if (bAllowObjectReplication)
+	{
+		UE_NET_TRACE_FRAME_STATSCOUNTER(GetId(), ReplicationSystem.ReplicatedObjectCount, InternalSys.GetNetRefHandleManager().GetActiveObjectCount(), ENetTraceVerbosity::Verbose);
 
-			// Refresh the dirty objects we were told about.
-			Impl->UpdateDirtyObjectList();
+		// Tell systems we are starting PreSendUpdate
+		Impl->StartPreSendUpdate();
 
-			// Update world locations. We need this to happen before both filtering and prioritization.
-			Impl->UpdateWorldLocations();
+		// Refresh the dirty objects we were told about.
+		Impl->UpdateDirtyObjectList();
 
-			// Update filters, reduce the top-level scoped object list and set each connection's scope.
-			Impl->UpdateFiltering();
+		// Update world locations. We need this to happen before both filtering and prioritization.
+		Impl->UpdateWorldLocations();
 
-			// Invoke any operations we need to do before copying state data
-			Impl->CallPreSendUpdate(Params.DeltaSeconds);
+		// Update filters, reduce the top-level scoped object list and set each connection's scope.
+		Impl->UpdateFiltering();
 
-			// Finalize the dirty list with objects set dirty during the poll phase
-			Impl->UpdateDirtyListPostPoll();
+		// Invoke any operations we need to do before copying state data
+		Impl->CallPreSendUpdate(DeltaSeconds);
 
-			// Update conditionals
-			Impl->UpdateConditionals();
+		// Finalize the dirty list with objects set dirty during the poll phase
+		Impl->UpdateDirtyListPostPoll();
 
-			// Quantize dirty state data. We need this to happen before both filtering and prioritization
-			Impl->QuantizeDirtyStateData();
+		// Update conditionals
+		Impl->UpdateConditionals();
 
-			// We must process all attachments to objects going out of scope before we update the scope
-			Impl->ProcessNetObjectAttachmentSendQueue(FNetBlobManager::EProcessMode::ProcessObjectsGoingOutOfScope);
+		// Quantize dirty state data. We need this to happen before both filtering and prioritization
+		Impl->QuantizeDirtyStateData();
 
-			// Update scope for all connections
-			Impl->UpdateObjectScopes();
+		// We must process all attachments to objects going out of scope before we update the scope
+		Impl->ProcessNetObjectAttachmentSendQueue(FNetBlobManager::EProcessMode::ProcessObjectsGoingOutOfScope);
 
-			// Propagate dirty changes to all connections
-			Impl->PropagateDirtyChanges();
-		}
+		// Update scope for all connections
+		Impl->UpdateObjectScopes();
 
-		// Forward attachments to the connections after scope update
-		Impl->ProcessNetObjectAttachmentSendQueue(FNetBlobManager::EProcessMode::ProcessObjectsInScope);
-		Impl->ResetNetObjectAttachmentSendQueue();
+		// Propagate dirty changes to all connections
+		Impl->PropagateDirtyChanges();
+	}
 
-		if (bAllowObjectReplication)
-		{
-			// Update object priorities
-			Impl->UpdatePrioritization(ReplicatingConnections);
+	// Forward attachments to the connections after scope update
+	Impl->ProcessNetObjectAttachmentSendQueue(FNetBlobManager::EProcessMode::ProcessObjectsInScope);
+	Impl->ResetNetObjectAttachmentSendQueue();
 
-			// Delta compression preparations before send
-			{
-				FDeltaCompressionBaselineManagerPreSendUpdateParams UpdateParams;
-				UpdateParams.ChangeMaskCache = &InternalSys.GetChangeMaskCache();
-				InternalSys.GetDeltaCompressionBaselineManager().PreSendUpdate(UpdateParams);
-			}
-		}
+	if (bAllowObjectReplication)
+	{
+		// Update object priorities
+		Impl->UpdatePrioritization(ReplicatingConnections);
 
-		// Destroy objects pending destroy
+		// Delta compression preparations before send
 		{
-			Impl->UpdateUnresolvableReferenceTracking();
-			InternalSys.GetNetRefHandleManager().DestroyObjectsPendingDestroy();
+			FDeltaCompressionBaselineManagerPreSendUpdateParams UpdateParams;
+			UpdateParams.ChangeMaskCache = &InternalSys.GetChangeMaskCache();
+			InternalSys.GetDeltaCompressionBaselineManager().PreSendUpdate(UpdateParams);
 		}
 	}
-	else if (Impl->CurrentSendPass == EReplicationSystemSendPass::PostTickDispatch)
-	{
-		// Forward attachments scheduled to use the OOBChannel and mark connetions needing immediate send
-		Impl->ProcessOOBNetObjectAttachmentSendQueue();
 
-		Impl->PostTickDispatch();
+	// Destroy objects pending destroy
+	{
+		Impl->UpdateUnresolvableReferenceTracking();
+		InternalSys.GetNetRefHandleManager().DestroyObjectsPendingDestroy();
 	}
 }
 
-void UReplicationSystem::PreSendUpdate(float DeltaSeconds)
+void UReplicationSystem::TickPostReceive()
 {
-	const FSendUpdateParams Params
-	{ 
-		.SendPass = UE::Net::EReplicationSystemSendPass::TickFlush, 
-		.DeltaSeconds = DeltaSeconds 
-	};
-	PreSendUpdate(Params);
+	using namespace UE::Net;
+	using namespace UE::Net::Private;
+
+	IRIS_CSV_PROFILER_SCOPE(Iris, ReplicationSystem_TickPostReceive);
+
+	ensureMsgf(Impl->CurrentSendPass == EReplicationSystemSendPass::Invalid, TEXT("PostSendUpdate was not called after the last Tick."));
+	Impl->CurrentSendPass = EReplicationSystemSendPass::PostTickDispatch;
+
+	// Forward attachments scheduled to use the OOBChannel and mark connections needing immediate send
+	Impl->ProcessOOBNetObjectAttachmentSendQueue();
 }
 
 IRISCORE_API void UReplicationSystem::SendUpdate(TFunctionRef<void(TArrayView<uint32>)> SendFunction)
@@ -978,7 +976,10 @@ IRISCORE_API void UReplicationSystem::SendUpdate(TFunctionRef<void(TArrayView<ui
 	{
 		// We only need to send data to connections that has data to send in PostTickDispatch
 
-		FNetBitArray::ForAllSetBits(Impl->ConnectionsPendingPostTickDispatchSend, ReplicatingConnections, FNetBitArray::AndOp, [&ConnectionToUpdate](uint32 ConnId) { ConnectionToUpdate.Add(ConnId);});
+		FNetBitArray::ForAllSetBits(Impl->ConnectionsPendingPostTickDispatchSend, ReplicatingConnections, FNetBitArray::AndOp, [&ConnectionToUpdate](uint32 ConnId)
+		{ 
+			ConnectionToUpdate.Add(ConnId);
+		});
 		Impl->ConnectionsPendingPostTickDispatchSend.ClearAllBits();
 	}
 
@@ -997,7 +998,7 @@ void UReplicationSystem::PostSendUpdate()
 		return;
 	}
 
-	// Most systems are only updated during the normal TickFlush
+	// Most systems are only updated during the normal NetUpdate
 	if (Impl->CurrentSendPass == EReplicationSystemSendPass::TickFlush)
 	{
 		if (bAllowObjectReplication)
@@ -1040,6 +1041,16 @@ void UReplicationSystem::PostSendUpdate()
 	Impl->CurrentSendPass = EReplicationSystemSendPass::Invalid;
 }
 
+void UReplicationSystem::PreReceiveUpdate()
+{
+	Impl->PreReceiveUpdate();
+}
+
+void UReplicationSystem::PostReceiveUpdate()
+{
+	Impl->PostReceiveUpdate();
+}
+
 void UReplicationSystem::PostGarbageCollection()
 {
 	bDoCollectGarbage = 1U;
diff --git a/Engine/Source/Runtime/Experimental/Iris/Core/Public/Iris/ReplicationSystem/NetObjectFactory.h b/Engine/Source/Runtime/Experimental/Iris/Core/Public/Iris/ReplicationSystem/NetObjectFactory.h
--- a/Engine/Source/Runtime/Experimental/Iris/Core/Public/Iris/ReplicationSystem/NetObjectFactory.h
+++ b/Engine/Source/Runtime/Experimental/Iris/Core/Public/Iris/ReplicationSystem/NetObjectFactory.h
@@ -82,7 +82,7 @@ public:
 
 	void Init(UE::Net::FNetObjectFactoryId InId, UObjectReplicationBridge* InBridge);
 	void Deinit();
-	void PostTickDispatch();
+	void PostReceiveUpdate();
 
 	struct FInstantiateResult;
 	struct FInstantiateContext;
@@ -271,7 +271,7 @@ protected:
 	virtual void OnDeinit() {}
 
 	/** Called after we finished processing all incoming packets */
-	virtual void OnPostTickDispatch() {}
+	virtual void OnPostReceiveUpdate() {}
 
 	/**
 	 * Create the correct header type for a given replicated object and fill the header with the information representing it.
diff --git a/Engine/Source/Runtime/Experimental/Iris/Core/Public/Iris/ReplicationSystem/ObjectReplicationBridge.h b/Engine/Source/Runtime/Experimental/Iris/Core/Public/Iris/ReplicationSystem/ObjectReplicationBridge.h
--- a/Engine/Source/Runtime/Experimental/Iris/Core/Public/Iris/ReplicationSystem/ObjectReplicationBridge.h
+++ b/Engine/Source/Runtime/Experimental/Iris/Core/Public/Iris/ReplicationSystem/ObjectReplicationBridge.h
@@ -287,7 +287,7 @@ protected:
 	IRISCORE_API virtual void PreSendUpdate() override;	
 	IRISCORE_API virtual void OnStartPreSendUpdate() override;
 	IRISCORE_API virtual void OnPostSendUpdate() override;
-	IRISCORE_API virtual void OnPostTickDispatch() override;
+	IRISCORE_API virtual void OnPostReceiveUpdate() override;
 	IRISCORE_API virtual void UpdateInstancesWorldLocation() override;
 	IRISCORE_API virtual void PruneStaleObjects() override;	
 	IRISCORE_API virtual bool WriteNetRefHandleCreationInfo(FReplicationBridgeSerializationContext& Context, FNetRefHandle Handle) override;
diff --git a/Engine/Source/Runtime/Experimental/Iris/Core/Public/Iris/ReplicationSystem/ReplicationBridge.h b/Engine/Source/Runtime/Experimental/Iris/Core/Public/Iris/ReplicationSystem/ReplicationBridge.h
--- a/Engine/Source/Runtime/Experimental/Iris/Core/Public/Iris/ReplicationSystem/ReplicationBridge.h
+++ b/Engine/Source/Runtime/Experimental/Iris/Core/Public/Iris/ReplicationSystem/ReplicationBridge.h
@@ -127,6 +127,9 @@ public:
 	/** Get the group associated with the level in order to control connection filtering for it. */
 	IRISCORE_API UE::Net::FNetObjectGroupHandle GetLevelGroup(const UObject* Level) const;
 
+	/** Returns true when we are in the middle of processing incoming data. */
+	bool IsInReceiveUpdate() const { return bInReceiveUpdate; }
+
 	/** Print common information about this handle and the object it is mapped to */
 	[[nodiscard]] IRISCORE_API FString PrintObjectFromNetRefHandle(FNetRefHandle RefHandle) const;
 
@@ -138,22 +141,22 @@ protected:
 	IRISCORE_API virtual void Deinitialize();
 
 	/** Invoked before ReplicationSystem copies dirty state data. */
-	IRISCORE_API virtual void PreSendUpdate();
+	virtual void PreSendUpdate() {}
 
 	/** Invoked when the ReplicationSystem starts the PreSendUpdate tick. */
-	IRISCORE_API virtual void OnStartPreSendUpdate();
+	virtual void OnStartPreSendUpdate() {}
 	
 	/** Invoked after we sent data to all connections. */
-	IRISCORE_API virtual void OnPostSendUpdate();
+	virtual void OnPostSendUpdate() {}
 
 	/** Invoked after we processed all incoming data */
-	IRISCORE_API virtual void OnPostTickDispatch();
+	virtual void OnPostReceiveUpdate() {}
 	
 	/** Invoked before ReplicationSystem copies dirty state data for a single replicated object. */
-	IRISCORE_API virtual void PreSendUpdateSingleHandle(FNetRefHandle Handle);
+	virtual void PreSendUpdateSingleHandle(FNetRefHandle Handle) {}
 
 	/** Update world locations in FWorldLocations for objects that support it. */
-	IRISCORE_API virtual void UpdateInstancesWorldLocation();
+	virtual void UpdateInstancesWorldLocation() {}
 
 	// Remote interface, invoked from Replication code during serialization
 	
@@ -281,6 +284,9 @@ private:
 	void CallDetachInstance(FNetRefHandle Handle);
 	void CallDetachInstanceFromRemote(FNetRefHandle Handle, EReplicationBridgeDestroyInstanceReason DestroyReason, EReplicationBridgeDestroyInstanceFlags DestroyFlags, UE::Net::FNetObjectFactoryId NetFactoryId);
 
+	void PreReceiveUpdate();
+	void PostReceiveUpdate();
+
 private:
 
 	void InternalFlushStateData(UE::Net::FNetSerializationContext& SerializationContext, UE::Net::Private::FChangeMaskCache& ChangeMaskCache, UE::Net::FNetBitStreamWriter& ChangeMaskWriter, uint32 InternalObjectIndex);
@@ -335,6 +341,9 @@ private:
 
 	TMap<FObjectKey, UE::Net::FNetObjectGroupHandle> LevelGroups;
 
+	/** Tracks if we are in the middle of processing incoming data */
+	bool bInReceiveUpdate = false;
+
 private:
 
 	// DestructionInfos
diff --git a/Engine/Source/Runtime/Experimental/Iris/Core/Public/Iris/ReplicationSystem/ReplicationSystem.h b/Engine/Source/Runtime/Experimental/Iris/Core/Public/Iris/ReplicationSystem/ReplicationSystem.h
--- a/Engine/Source/Runtime/Experimental/Iris/Core/Public/Iris/ReplicationSystem/ReplicationSystem.h
+++ b/Engine/Source/Runtime/Experimental/Iris/Core/Public/Iris/ReplicationSystem/ReplicationSystem.h
@@ -146,27 +146,23 @@ public:
 	/** @return Is this system configured to replicate object properties. */
 	bool AllowObjectReplication() { return bAllowObjectReplication; }
 
-	struct FSendUpdateParams
-	{
-		// Type of SendPass we want to do @see EReplicationSystemSendPass
-		UE::Net::EReplicationSystemSendPass SendPass = UE::Net::EReplicationSystemSendPass::TickFlush;
-
-		// DeltaTime, only relevant for EReplicationSystemSendPass::TickFlush
-		float DeltaSeconds = 0.f;		
-	};
-
 	/**
-	 * PreSendUpdate performs all the necessary work, such as filtering and prioritization of objects,
-	 * so that each connection will be properly updated with all the information needed in order to replicate.
-	 * @param Params, Parameters for the update pass, to @see FSendUpateParams 
+	 * Update all internal systems, such as filtering, dirty tracking, prioritization, etc.
+	 * Also copies all replicated data into the internal Iris protocols so they are ready to be sent to clients
 	 */
-	IRISCORE_API void PreSendUpdate(const FSendUpdateParams& Params);	
+	IRISCORE_API void NetUpdate(float DeltaSeconds);
 
 	/**
-	 * PreSendUpdate performs all the necessary work, such as filtering and prioritization of objects,
-	 * so that each connection will be properly updated with all the information needed in order to replicate.
+	 * Tick replication system after parsing all received data.
+	 * Used to check if immediate RPCs need to be sent before the rest of the engine is ticked
 	 */
-	IRISCORE_API void PreSendUpdate(float DeltaSeconds);
+	IRISCORE_API void TickPostReceive();
+
+	/** Callback triggered before connections start processing their received data */
+	IRISCORE_API void PreReceiveUpdate();
+
+	/** Callback triggered after the connections processed all received data */
+	IRISCORE_API void PostReceiveUpdate();
 
 	/**
 	 * SendUpdate is currently more of a placeholder for a future where the ReplicationSystem itself is responsible for
