Disable UE_WITH_REMOTE_OBJECT_HANDLE for replication system unit tests.

diff --git a/Engine/Plugins/Runtime/ReplicationSystemTestPlugin/Source/Private/Tests/Serialization/TestRemoteObjectReferenceNetSerializer.cpp b/Engine/Plugins/Runtime/ReplicationSystemTestPlugin/Source/Private/Tests/Serialization/TestRemoteObjectReferenceNetSerializer.cpp
--- a/Engine/Plugins/Runtime/ReplicationSystemTestPlugin/Source/Private/Tests/Serialization/TestRemoteObjectReferenceNetSerializer.cpp
+++ b/Engine/Plugins/Runtime/ReplicationSystemTestPlugin/Source/Private/Tests/Serialization/TestRemoteObjectReferenceNetSerializer.cpp
@@ -20,6 +20,9 @@ void UTestReplicatedObjectWithRemoteReference::RemoteRPCWithRemoteReferenceParam
 	LastReceivedRemoteReference = RemoteReference;
 }
 
+// These tests require remote reference support to be compiled in
+#if UE_WITH_REMOTE_OBJECT_REFERENCE
+
 namespace UE::Net::Private
 {
 
@@ -139,3 +142,5 @@ UE_NET_TEST_FIXTURE(FRPCTestFixture, TestRemoteObjectReferencePaths)
 }
 
 }
+
+#endif
diff --git a/Engine/Plugins/Runtime/ReplicationSystemTestPlugin/Source/Private/Tests/TestNetworkAutomationTestFramework.cpp b/Engine/Plugins/Runtime/ReplicationSystemTestPlugin/Source/Private/Tests/TestNetworkAutomationTestFramework.cpp
--- a/Engine/Plugins/Runtime/ReplicationSystemTestPlugin/Source/Private/Tests/TestNetworkAutomationTestFramework.cpp
+++ b/Engine/Plugins/Runtime/ReplicationSystemTestPlugin/Source/Private/Tests/TestNetworkAutomationTestFramework.cpp
@@ -242,11 +242,4 @@ UE_NET_TEST(NetworkAutomationTest, TestPushModelEnabled)
 	UE_NET_ASSERT_TRUE(PushModelCVar->GetBool());
 }
 
-UE_NET_TEST(NetworkAutomationTest, TestRemoteObjectHandleEnabled)
-{
-#if !UE_WITH_REMOTE_OBJECT_HANDLE
-	UE_NET_FAIL("Remote object handle support not compiled in.");
-#endif
-}
-
 }
diff --git a/Engine/Source/Programs/ReplicationSystemLowLevelTests/ReplicationSystemLowLevelTests.Target.cs b/Engine/Source/Programs/ReplicationSystemLowLevelTests/ReplicationSystemLowLevelTests.Target.cs
--- a/Engine/Source/Programs/ReplicationSystemLowLevelTests/ReplicationSystemLowLevelTests.Target.cs
+++ b/Engine/Source/Programs/ReplicationSystemLowLevelTests/ReplicationSystemLowLevelTests.Target.cs
@@ -24,6 +24,6 @@ public class ReplicationSystemLowLevelTestsTarget : TestTargetRules
 		GlobalDefinitions.Add("UE_NET_TEST_FAKE_REP_TAGS=1");
 		// Load time profiling brings object construction to a crawl.
 		GlobalDefinitions.Add("LOADTIMEPROFILERTRACE_ENABLED=0");
-		GlobalDefinitions.Add("UE_WITH_REMOTE_OBJECT_HANDLE=1");
+		GlobalDefinitions.Add("UE_WITH_REMOTE_OBJECT_HANDLE=0");
 	}
 }
diff --git a/Engine/Source/Programs/ReplicationSystemTest/ReplicationSystemTest.Target.cs b/Engine/Source/Programs/ReplicationSystemTest/ReplicationSystemTest.Target.cs
--- a/Engine/Source/Programs/ReplicationSystemTest/ReplicationSystemTest.Target.cs
+++ b/Engine/Source/Programs/ReplicationSystemTest/ReplicationSystemTest.Target.cs
@@ -43,7 +43,7 @@ public class ReplicationSystemTestTarget : TargetRules
 		GlobalDefinitions.Add("UE_NET_TEST_FAKE_REP_TAGS=1");
 		// Load time profiling brings object construction to a crawl.
 		GlobalDefinitions.Add("LOADTIMEPROFILERTRACE_ENABLED=0");
-		GlobalDefinitions.Add("UE_WITH_REMOTE_OBJECT_HANDLE=1");
+		GlobalDefinitions.Add("UE_WITH_REMOTE_OBJECT_HANDLE=0");
 	}
 }
 
