ReplicationSystemTest
* Downgraded some logs to Verbose

#rnx

diff --git a/Engine/Plugins/Runtime/ReplicationSystemTestPlugin/Source/Private/Tests/ReplicationSystem/ReplicationSystemServerClientTestFixture.cpp b/Engine/Plugins/Runtime/ReplicationSystemTestPlugin/Source/Private/Tests/ReplicationSystem/ReplicationSystemServerClientTestFixture.cpp
--- a/Engine/Plugins/Runtime/ReplicationSystemTestPlugin/Source/Private/Tests/ReplicationSystem/ReplicationSystemServerClientTestFixture.cpp
+++ b/Engine/Plugins/Runtime/ReplicationSystemTestPlugin/Source/Private/Tests/ReplicationSystem/ReplicationSystemServerClientTestFixture.cpp
@@ -413,7 +413,7 @@ bool FReplicationSystemTestNode::SendUpdate(uint32 ConnectionId, const TCHAR* De
 	{
 		UE_NET_TRACE_FLUSH_COLLECTOR(Context.GetTraceCollector(), GetNetTraceId(), Connection.ConnectionId, ENetTracePacketType::Outgoing);
 		UE_NET_TRACE_PACKET_SEND(GetNetTraceId(), Connection.ConnectionId, Packet.PacketId, Packet.BitCount);
-		UE_LOG(LogIris, Log, TEXT("ReplicationSystemTestFixture: Conn: %u Send PacketId: %u %s"), ConnectionId, Packet.PacketId, *Packet.Desc);
+		UE_LOG(LogIris, Verbose, TEXT("ReplicationSystemTestFixture: Conn: %u Send PacketId: %u %s"), ConnectionId, Packet.PacketId, *Packet.Desc);
 	}
 
 	UE_NET_TRACE_DESTROY_COLLECTOR(Context.GetTraceCollector());
@@ -432,7 +432,7 @@ void FReplicationSystemTestNode::DeliverTo(FReplicationSystemTestNode& Dest, uin
 	FConnectionInfo& Connection = GetConnectionInfo(LocalConnectionId);
 	if (Connection.WrittenPackets.IsEmpty())
 	{
-		UE_LOG(LogIris, Log, TEXT("ReplicationSystemTestFixture: Conn: %u Unable to %hs packet as there are no packets."), LocalConnectionId, (bDeliver ? "deliver" : "drop"));
+		UE_LOG(LogIris, Verbose, TEXT("ReplicationSystemTestFixture: Conn: %u Unable to %hs packet as there are no packets."), LocalConnectionId, (bDeliver ? "deliver" : "drop"));
 		return;
 	}
 
@@ -446,7 +446,7 @@ void FReplicationSystemTestNode::DeliverTo(FReplicationSystemTestNode& Dest, uin
 
 		Context.SetTraceCollector(UE_NET_TRACE_CREATE_COLLECTOR(ENetTraceVerbosity::Trace));
 
-		UE_LOG(LogIris, Log, TEXT("ReplicationSystemTestFixture: Conn: %u Deliver PacketId: %u %s"), LocalConnectionId, Packet.PacketId, *Packet.Desc);
+		UE_LOG(LogIris, Verbose, TEXT("ReplicationSystemTestFixture: Conn: %u Deliver PacketId: %u %s"), LocalConnectionId, Packet.PacketId, *Packet.Desc);
 		Dest.RecvUpdate(RemoteConnectionId, Context);
 
 		UE_NET_TRACE_FLUSH_COLLECTOR(Context.GetTraceCollector(), Dest.GetNetTraceId(), RemoteConnectionId, ENetTracePacketType::Incoming);
@@ -455,7 +455,7 @@ void FReplicationSystemTestNode::DeliverTo(FReplicationSystemTestNode& Dest, uin
 	}
 	else
 	{
-		UE_LOG(LogIris, Log, TEXT("ReplicationSystemTestFixture: Conn: %u Dropped PacketId: %u %s"), LocalConnectionId, Packet.PacketId, *Packet.Desc);
+		UE_LOG(LogIris, Verbose, TEXT("ReplicationSystemTestFixture: Conn: %u Dropped PacketId: %u %s"), LocalConnectionId, Packet.PacketId, *Packet.Desc);
 		UE_NET_TRACE_PACKET_DROPPED(Dest.GetNetTraceId(), RemoteConnectionId, Packet.PacketId, ENetTracePacketType::Incoming);
 		UE_NET_TRACE_PACKET_DROPPED(GetNetTraceId(), LocalConnectionId, Packet.PacketId, ENetTracePacketType::Outgoing);
 	}
