Iris - increase ReplicationSystemIdBits to 10 to support up to 1024 replication systems.
Remove FReplicationSystemFactory::MaxReplicationSystemCount in favor of FNetRefHandle::MaxReplicationSystemCount so there's a single source of truth.
Change the ReplicationSystems array to a TArray with inline allocator, maintaining stable indices, and as small a size as possible.
Add test cases.

diff --git a/Engine/Plugins/Runtime/ReplicationSystemTestPlugin/Source/Private/Tests/ReplicationSystem/MultiReplicationSystemsTestFixture.cpp b/Engine/Plugins/Runtime/ReplicationSystemTestPlugin/Source/Private/Tests/ReplicationSystem/MultiReplicationSystemsTestFixture.cpp
--- a/Engine/Plugins/Runtime/ReplicationSystemTestPlugin/Source/Private/Tests/ReplicationSystem/MultiReplicationSystemsTestFixture.cpp
+++ b/Engine/Plugins/Runtime/ReplicationSystemTestPlugin/Source/Private/Tests/ReplicationSystem/MultiReplicationSystemsTestFixture.cpp
@@ -65,6 +65,21 @@ FReplicationSystemTestServer* FMultiReplicationSystemsTestFixture::CreateServer(
 	return Server;
 }
 
+void FMultiReplicationSystemsTestFixture::DestroyServer(FReplicationSystemTestServer* Server)
+{
+	if (TMap<FReplicationSystemTestServer*, TArray<FReplicationSystemTestClient*>>::ValueType* Clients = ServerClients.Find(Server))
+	{
+		for (FReplicationSystemTestClient* Client : *Clients)
+		{
+			delete Client;
+		}
+	}
+	ServerClients.Remove(Server);
+
+	Servers.RemoveSingle(Server);
+	delete Server;
+}
+
 FReplicationSystemTestClient* FMultiReplicationSystemsTestFixture::CreateClientForServer(FReplicationSystemTestServer* Server)
 {
 	if (TMap<FReplicationSystemTestServer*, TArray<FReplicationSystemTestClient*>>::ValueType* Clients = ServerClients.Find(Server))
diff --git a/Engine/Plugins/Runtime/ReplicationSystemTestPlugin/Source/Private/Tests/ReplicationSystem/MultiReplicationSystemsTestFixture.h b/Engine/Plugins/Runtime/ReplicationSystemTestPlugin/Source/Private/Tests/ReplicationSystem/MultiReplicationSystemsTestFixture.h
--- a/Engine/Plugins/Runtime/ReplicationSystemTestPlugin/Source/Private/Tests/ReplicationSystem/MultiReplicationSystemsTestFixture.h
+++ b/Engine/Plugins/Runtime/ReplicationSystemTestPlugin/Source/Private/Tests/ReplicationSystem/MultiReplicationSystemsTestFixture.h
@@ -43,6 +43,9 @@ protected:
 	/** Creates a server and returns it. */
 	FReplicationSystemTestServer* CreateServer();
 
+	/** Destroys a server and all of its clients */
+	void DestroyServer(FReplicationSystemTestServer* Server);
+
 	/** Creates a client that is connected to a specific server. */
 	FReplicationSystemTestClient* CreateClientForServer(FReplicationSystemTestServer*);
 
diff --git a/Engine/Plugins/Runtime/ReplicationSystemTestPlugin/Source/Private/Tests/ReplicationSystem/TestMultiReplicationSystems.cpp b/Engine/Plugins/Runtime/ReplicationSystemTestPlugin/Source/Private/Tests/ReplicationSystem/TestMultiReplicationSystems.cpp
--- a/Engine/Plugins/Runtime/ReplicationSystemTestPlugin/Source/Private/Tests/ReplicationSystem/TestMultiReplicationSystems.cpp
+++ b/Engine/Plugins/Runtime/ReplicationSystemTestPlugin/Source/Private/Tests/ReplicationSystem/TestMultiReplicationSystems.cpp
@@ -1,6 +1,7 @@
 // Copyright Epic Games, Inc. All Rights Reserved.
 
 #include "Tests/ReplicationSystem/MultiReplicationSystemsTestFixture.h"
+#include "Iris/ReplicationSystem/ReplicationSystem.h"
 #include "Net/Core/NetHandle/NetHandleManager.h"
 #include "Containers/StaticArray.h"
 
@@ -27,6 +28,92 @@ UE_NET_TEST_FIXTURE(FMultiReplicationSystemsTestFixture, CanCreateMultipleReplic
 	UE_NET_ASSERT_EQ(ValidReplicationSystemCount, ExpectedReplicationSystemCount);
 }
 
+UE_NET_TEST_FIXTURE(FMultiReplicationSystemsTestFixture, CanCreateMaxReplicationSystemCount)
+{
+	CreateServers(FNetRefHandle::MaxReplicationSystemCount);
+
+	TArrayView<FReplicationSystemTestServer*> ServersView = GetAllServers();
+	UE_NET_ASSERT_EQ(unsigned(ServersView.Num()), unsigned(FNetRefHandle::MaxReplicationSystemCount));
+
+	const unsigned ExpectedReplicationSystemCount = FNetRefHandle::MaxReplicationSystemCount;
+	unsigned ValidReplicationSystemCount = 0U;
+	for (FReplicationSystemTestServer* Server : ServersView)
+	{
+		if (UReplicationSystem* ReplicationSystem = Server->GetReplicationSystem())
+		{
+			++ValidReplicationSystemCount;
+		}
+	}
+
+	UE_NET_ASSERT_EQ(ValidReplicationSystemCount, ExpectedReplicationSystemCount);
+}
+
+UE_NET_TEST_FIXTURE(FMultiReplicationSystemsTestFixture, ReplicationSystemIdsAreStable)
+{
+	constexpr uint32 InitialNumServers = 5;
+	TStaticArray<FReplicationSystemTestServer*, InitialNumServers> TestServers;
+	
+	for (uint32 i = 0; i < InitialNumServers; ++i)
+	{
+		TestServers[i] = CreateServer();
+		UE_NET_ASSERT_NE(TestServers[i]->GetReplicationSystem(), nullptr);
+	}
+
+	TArrayView<UReplicationSystem*> ReplicationSystemsView = FReplicationSystemFactory::GetAllReplicationSystems();
+	UE_NET_ASSERT_EQ(unsigned(ReplicationSystemsView.Num()), unsigned(InitialNumServers));
+
+	// Destroy some replication systems in the middle of the array
+	DestroyServer(TestServers[1]);
+	TestServers[1] = nullptr;
+
+	DestroyServer(TestServers[3]);
+	TestServers[3] = nullptr;
+
+	// Remaining replication systems should maintain their IDs
+	UReplicationSystem* System0 = UE::Net::GetReplicationSystem(TestServers[0]->GetReplicationSystem()->GetId());
+	UE_NET_ASSERT_EQ(System0, TestServers[0]->GetReplicationSystem());
+
+	UReplicationSystem* System2 = UE::Net::GetReplicationSystem(TestServers[2]->GetReplicationSystem()->GetId());
+	UE_NET_ASSERT_EQ(System2, TestServers[2]->GetReplicationSystem());
+
+	UReplicationSystem* System4 = UE::Net::GetReplicationSystem(TestServers[4]->GetReplicationSystem()->GetId());
+	UE_NET_ASSERT_EQ(System4, TestServers[4]->GetReplicationSystem());
+}
+
+UE_NET_TEST_FIXTURE(FMultiReplicationSystemsTestFixture, ReplicationSystemArrayMaintainsMinimumSize)
+{
+	constexpr uint32 InitialNumServers = 5;
+	TStaticArray<FReplicationSystemTestServer*, InitialNumServers> TestServers;
+
+	for (uint32 i = 0; i < InitialNumServers; ++i)
+	{
+		TestServers[i] = CreateServer();
+		UE_NET_ASSERT_NE(TestServers[i]->GetReplicationSystem(), nullptr);
+	}
+
+	TArrayView<UReplicationSystem*> ReplicationSystemsView = FReplicationSystemFactory::GetAllReplicationSystems();
+	UE_NET_ASSERT_EQ(unsigned(ReplicationSystemsView.Num()), unsigned(InitialNumServers));
+
+	// Destroy some replication systems in the middle of the array
+	DestroyServer(TestServers[2]);
+	TestServers[2] = nullptr;
+
+	DestroyServer(TestServers[3]);
+	TestServers[3] = nullptr;
+
+	// View should still be the same size because the last ReplicationSystem still exists
+	TArrayView<UReplicationSystem*> ReplicationSystemsView2 = FReplicationSystemFactory::GetAllReplicationSystems();
+	UE_NET_ASSERT_EQ(unsigned(ReplicationSystemsView2.Num()), unsigned(InitialNumServers));
+
+	// Destroy the last replication system, leaving the only valid ones at indices 0 and 1
+	DestroyServer(TestServers[4]);
+	TestServers[4] = nullptr;
+
+	// View should shrink as all the null entries at the end are removed
+	TArrayView<UReplicationSystem*> ReplicationSystemsView3 = FReplicationSystemFactory::GetAllReplicationSystems();
+	UE_NET_ASSERT_EQ(unsigned(ReplicationSystemsView3.Num()), 2);
+}
+
 UE_NET_TEST_FIXTURE(FMultiReplicationSystemsTestFixture, CanReplicateObjectOnMultipleReplicationSystems)
 {
 	CreateSomeServers();
diff --git a/Engine/Source/Runtime/Experimental/Iris/Core/Private/Iris/Core/IrisDebugging.cpp b/Engine/Source/Runtime/Experimental/Iris/Core/Private/Iris/Core/IrisDebugging.cpp
--- a/Engine/Source/Runtime/Experimental/Iris/Core/Private/Iris/Core/IrisDebugging.cpp
+++ b/Engine/Source/Runtime/Experimental/Iris/Core/Private/Iris/Core/IrisDebugging.cpp
@@ -217,9 +217,8 @@ void SetIrisDebugInternalNetRefIndexViaObject(UObject* Instance)
 	}
 
 	// See if we can find the instance in any replication system
-	for (uint32 RepSystemIt = 0; RepSystemIt < FReplicationSystemFactory::MaxReplicationSystemCount; ++RepSystemIt)
+	for (const UReplicationSystem* ReplicationSystem : FReplicationSystemFactory::GetAllReplicationSystems())
 	{
-		const UReplicationSystem* ReplicationSystem = GetReplicationSystemForDebug(RepSystemIt);
 		if (ReplicationSystem)
 		{
 			if (const UObjectReplicationBridge* Bridge = ReplicationSystem->GetReplicationBridgeAs<UObjectReplicationBridge>())
@@ -359,9 +358,8 @@ FNetReplicatedObjectDebugInfo DebugNetObject(UObject* Instance)
 	FNetReplicatedObjectDebugInfo Info = {};
 
 	// See if we can find the instance in any replication system
-	for (uint32 RepSystemIt = 0; RepSystemIt < FReplicationSystemFactory::MaxReplicationSystemCount; ++RepSystemIt)
+	for (const UReplicationSystem* ReplicationSystem : FReplicationSystemFactory::GetAllReplicationSystems())
 	{
-		const UReplicationSystem* ReplicationSystem = GetReplicationSystemForDebug(RepSystemIt);
 		if (ReplicationSystem)
 		{
 			if (const UObjectReplicationBridge* Bridge = ReplicationSystem->GetReplicationBridgeAs<UObjectReplicationBridge>())
diff --git a/Engine/Source/Runtime/Experimental/Iris/Core/Private/Iris/ReplicationSystem/ReplicationSystem.cpp b/Engine/Source/Runtime/Experimental/Iris/Core/Private/Iris/ReplicationSystem/ReplicationSystem.cpp
--- a/Engine/Source/Runtime/Experimental/Iris/Core/Private/Iris/ReplicationSystem/ReplicationSystem.cpp
+++ b/Engine/Source/Runtime/Experimental/Iris/Core/Private/Iris/ReplicationSystem/ReplicationSystem.cpp
@@ -2084,8 +2084,7 @@ FReplicationSystemDestroyedDelegate& FReplicationSystemFactory::GetReplicationSy
 }
 
 // FReplicationSystemFactory
-UReplicationSystem* FReplicationSystemFactory::ReplicationSystems[MaxReplicationSystemCount];
-uint32 FReplicationSystemFactory::MaxReplicationSystemId = 0;
+FReplicationSystemFactory::FReplicationSystemArray FReplicationSystemFactory::ReplicationSystems;
 
 UReplicationSystem* FReplicationSystemFactory::CreateReplicationSystem(const UReplicationSystem::FReplicationSystemParams& Params)
 {
@@ -2097,22 +2096,19 @@ UReplicationSystem* FReplicationSystemFactory::CreateReplicationSystem(const URe
 		return nullptr;
 	}
 
-	for (uint32 It = 0, EndIt = MaxReplicationSystemCount; It != EndIt; ++It)
+	if (static_cast<uint32>(ReplicationSystems.Num()) < FNetRefHandle::MaxReplicationSystemCount)
 	{
-		if (ReplicationSystems[It] != nullptr)
+		int32 NewSystemIndex = ReplicationSystems.Find(nullptr);
+		if (NewSystemIndex == INDEX_NONE)
 		{
-			continue;
+			NewSystemIndex = ReplicationSystems.Add(nullptr);
 		}
 
 		UReplicationSystem* ReplicationSystem = NewObject<UReplicationSystem>();
-		ReplicationSystems[It] = ReplicationSystem;
+		ReplicationSystems[NewSystemIndex] = ReplicationSystem;
 		ReplicationSystem->AddToRoot();
 
-		const uint32 ReplicationSystemId = It;
-		if (ReplicationSystemId > MaxReplicationSystemId)
-		{
-			MaxReplicationSystemId = ReplicationSystemId;
-		}
+		const uint32 ReplicationSystemId = static_cast<uint32>(NewSystemIndex);
 
 		UE_LOG(LogIris, Display, TEXT("Iris ReplicationSystem[%i]: %s (0x%p) is created"), ReplicationSystemId, *ReplicationSystem->GetName(), ReplicationSystem);
 
@@ -2126,7 +2122,7 @@ UReplicationSystem* FReplicationSystemFactory::CreateReplicationSystem(const URe
 		return ReplicationSystem;
 	}
 
-	LowLevelFatalError(TEXT("Too many ReplicationSystems have already been created (%u)"), MaxReplicationSystemCount);
+	LowLevelFatalError(TEXT("Too many ReplicationSystems have already been created (%u)"), FNetRefHandle::MaxReplicationSystemCount);
 	return nullptr;
 }
 
@@ -2141,19 +2137,18 @@ void FReplicationSystemFactory::DestroyReplicationSystem(UReplicationSystem* Sys
 
 	UE_LOG(LogIris, Display, TEXT("Iris ReplicationSystem[%i]: %s (0x%p) is about to be destroyed"), Id, *System->GetName(), System);
 
-	if (Id < MaxReplicationSystemCount)
+	if (Id < static_cast<uint32>(ReplicationSystems.Num()))
 	{
 		ReplicationSystems[Id] = nullptr;
 
-		uint32 NewMaxReplicationSystemId = 0;
-		for (uint32 It = 0, EndIt = MaxReplicationSystemCount; It != EndIt; ++It)
+		// Remove all null entries from the end of the array
+		const int32 LastValidIndex = ReplicationSystems.FindLastByPredicate([](const UReplicationSystem* It)
 		{
-			if (ReplicationSystems[It] != nullptr)
-			{
-				NewMaxReplicationSystemId = Id;
-			}
-		}
-		MaxReplicationSystemId = NewMaxReplicationSystemId;
+			return It != nullptr;
+		});
+
+		const int32 MaxValidSystemNum = LastValidIndex == INDEX_NONE ? 0 : LastValidIndex + 1;
+		ReplicationSystems.SetNum(MaxValidSystemNum, EAllowShrinking::No);
 	}
 
 	if (GetReplicationSystemDestroyedDelegate().IsBound())
@@ -2168,7 +2163,7 @@ void FReplicationSystemFactory::DestroyReplicationSystem(UReplicationSystem* Sys
 
 TArrayView<UReplicationSystem*> FReplicationSystemFactory::GetAllReplicationSystems()
 {
-	return MakeArrayView(ReplicationSystems, MaxReplicationSystemId + 1);
+	return TArrayView<UReplicationSystem*>(ReplicationSystems);
 }
 
 } //end namespace UE::Net
diff --git a/Engine/Source/Runtime/Experimental/Iris/Core/Public/Iris/ReplicationSystem/NetRefHandle.h b/Engine/Source/Runtime/Experimental/Iris/Core/Public/Iris/ReplicationSystem/NetRefHandle.h
--- a/Engine/Source/Runtime/Experimental/Iris/Core/Public/Iris/ReplicationSystem/NetRefHandle.h
+++ b/Engine/Source/Runtime/Experimental/Iris/Core/Public/Iris/ReplicationSystem/NetRefHandle.h
@@ -29,11 +29,14 @@ public:
 
 private:
 	enum { InvalidValue = 0 };
-	enum { IdBits = 60 };
-	enum { ReplicationSystemIdBits = 4 };
+	enum { IdBits = 54 };
+	enum { ReplicationSystemIdBits = 10 };
 
 public:
 
+	static constexpr uint64 MaxReplicationSystemId = (1ULL << ReplicationSystemIdBits) - 1;
+	static constexpr uint64 MaxReplicationSystemCount = MaxReplicationSystemId + 1;
+
 	FNetRefHandle() : Value(InvalidValue) {}
 
 	uint64 GetId() const { return Id; }
@@ -63,7 +66,6 @@ private:
 
 	static constexpr uint64 StaticIdMask = 1;
 	static constexpr uint64 IdMask = (1ULL << IdBits) - 1;
-	static constexpr uint64 MaxReplicationSystemId = (1ULL << ReplicationSystemIdBits) - 1;
 
 	union 
 	{
diff --git a/Engine/Source/Runtime/Experimental/Iris/Core/Public/Iris/ReplicationSystem/ReplicationSystem.h b/Engine/Source/Runtime/Experimental/Iris/Core/Public/Iris/ReplicationSystem/ReplicationSystem.h
--- a/Engine/Source/Runtime/Experimental/Iris/Core/Public/Iris/ReplicationSystem/ReplicationSystem.h
+++ b/Engine/Source/Runtime/Experimental/Iris/Core/Public/Iris/ReplicationSystem/ReplicationSystem.h
@@ -15,6 +15,12 @@
 
 #include "ReplicationSystem.generated.h"
 
+// The number of inline elements the TInlineAllocator for ReplicationSystems can store.
+// Projects may define this to a value best suited for their needs.
+#ifndef UE_NUM_INLINE_REPLICATIONSYSTEMS
+	#define UE_NUM_INLINE_REPLICATIONSYSTEMS 8
+#endif
+
 class UDataStreamManager;
 class UDataStream;
 class UFunction;
@@ -806,22 +812,17 @@ public:
 	/** Static delegate that is triggered before we destroy a replication system. */
 	IRISCORE_API static FReplicationSystemDestroyedDelegate& GetReplicationSystemDestroyedDelegate();
 
-	enum ReplicationSystemConstants : uint32
-	{
-		MaxReplicationSystemCount = 16,
-	};
-
 private:
 
 	friend UReplicationSystem* GetReplicationSystem(uint32 Id);
 
-	IRISCORE_API static UReplicationSystem* ReplicationSystems[];
-	static uint32 MaxReplicationSystemId;
+	using FReplicationSystemArray = TArray<UReplicationSystem*, TInlineAllocator<UE_NUM_INLINE_REPLICATIONSYSTEMS>>;
+	IRISCORE_API static FReplicationSystemArray ReplicationSystems;
 };
 
 inline UReplicationSystem* GetReplicationSystem(uint32 Id)
 {
-	return Id >= FReplicationSystemFactory::MaxReplicationSystemCount ? nullptr : FReplicationSystemFactory::ReplicationSystems[Id];
+	return Id >= static_cast<uint32>(FReplicationSystemFactory::ReplicationSystems.Num()) ? nullptr : FReplicationSystemFactory::ReplicationSystems[Id];
 }
 
 } // end namespace UE::Net
