Reduce memory usage of replication systems in FMultiReplicationSystemsTestFixture::CanCreateMaxReplicationSystemCount.

diff --git a/Engine/Plugins/Runtime/ReplicationSystemTestPlugin/Source/Private/Tests/ReplicationSystem/MultiReplicationSystemsTestFixture.cpp b/Engine/Plugins/Runtime/ReplicationSystemTestPlugin/Source/Private/Tests/ReplicationSystem/MultiReplicationSystemsTestFixture.cpp
--- a/Engine/Plugins/Runtime/ReplicationSystemTestPlugin/Source/Private/Tests/ReplicationSystem/MultiReplicationSystemsTestFixture.cpp
+++ b/Engine/Plugins/Runtime/ReplicationSystemTestPlugin/Source/Private/Tests/ReplicationSystem/MultiReplicationSystemsTestFixture.cpp
@@ -38,11 +38,11 @@ void FMultiReplicationSystemsTestFixture::TearDown()
 	FNetworkAutomationTestSuiteFixture::TearDown();
 }
 
-void FMultiReplicationSystemsTestFixture::CreateServers(unsigned ServerCount)
+void FMultiReplicationSystemsTestFixture::CreateServers(unsigned ServerCount, const FReplicationSystemTestNode::FReplicationSystemParamsOverride* ParamsOverride)
 {
 	for (unsigned It = 0, EndIt = ServerCount; It != EndIt; ++It)
 	{
-		CreateServer();
+		CreateServer(ParamsOverride);
 	}
 }
 
@@ -56,9 +56,10 @@ TArrayView<FReplicationSystemTestServer*> FMultiReplicationSystemsTestFixture::G
 	return MakeArrayView(Servers);
 }
 
-FReplicationSystemTestServer* FMultiReplicationSystemsTestFixture::CreateServer()
+FReplicationSystemTestServer* FMultiReplicationSystemsTestFixture::CreateServer(const FReplicationSystemTestNode::FReplicationSystemParamsOverride* ParamsOverride)
 {
-	FReplicationSystemTestServer* Server = new FReplicationSystemTestServer(GetName());
+	FReplicationSystemTestServer* Server = new FReplicationSystemTestServer(FReplicationSystemTestNode::DelaySetup);
+	Server->Setup(true, GetName(), ParamsOverride);
 	Servers.Add(Server);
 	ServerClients.Emplace(Server);
 
diff --git a/Engine/Plugins/Runtime/ReplicationSystemTestPlugin/Source/Private/Tests/ReplicationSystem/MultiReplicationSystemsTestFixture.h b/Engine/Plugins/Runtime/ReplicationSystemTestPlugin/Source/Private/Tests/ReplicationSystem/MultiReplicationSystemsTestFixture.h
--- a/Engine/Plugins/Runtime/ReplicationSystemTestPlugin/Source/Private/Tests/ReplicationSystem/MultiReplicationSystemsTestFixture.h
+++ b/Engine/Plugins/Runtime/ReplicationSystemTestPlugin/Source/Private/Tests/ReplicationSystem/MultiReplicationSystemsTestFixture.h
@@ -32,7 +32,7 @@ protected:
 	virtual void TearDown() override;
 
 	/** Create a specific amount of servers. */
-	void CreateServers(unsigned ServerCount);
+	void CreateServers(unsigned ServerCount, const FReplicationSystemTestNode::FReplicationSystemParamsOverride* ParamsOverride = nullptr);
 
 	/** Creates the default amount of servers. */
 	void CreateSomeServers();
@@ -41,7 +41,7 @@ protected:
 	TArrayView<FReplicationSystemTestServer*> GetAllServers();
 
 	/** Creates a server and returns it. */
-	FReplicationSystemTestServer* CreateServer();
+	FReplicationSystemTestServer* CreateServer(const FReplicationSystemTestNode::FReplicationSystemParamsOverride* ParamsOverride = nullptr);
 
 	/** Destroys a server and all of its clients */
 	void DestroyServer(FReplicationSystemTestServer* Server);
diff --git a/Engine/Plugins/Runtime/ReplicationSystemTestPlugin/Source/Private/Tests/ReplicationSystem/TestMultiReplicationSystems.cpp b/Engine/Plugins/Runtime/ReplicationSystemTestPlugin/Source/Private/Tests/ReplicationSystem/TestMultiReplicationSystems.cpp
--- a/Engine/Plugins/Runtime/ReplicationSystemTestPlugin/Source/Private/Tests/ReplicationSystem/TestMultiReplicationSystems.cpp
+++ b/Engine/Plugins/Runtime/ReplicationSystemTestPlugin/Source/Private/Tests/ReplicationSystem/TestMultiReplicationSystems.cpp
@@ -30,7 +30,12 @@ UE_NET_TEST_FIXTURE(FMultiReplicationSystemsTestFixture, CanCreateMultipleReplic
 
 UE_NET_TEST_FIXTURE(FMultiReplicationSystemsTestFixture, CanCreateMaxReplicationSystemCount)
 {
-	CreateServers(FNetRefHandle::MaxReplicationSystemCount);
+	// Use a lower max object count for this test to keep memory usage under control
+	FReplicationSystemTestNode::FReplicationSystemParamsOverride ParamsOverride;
+	ParamsOverride.MaxReplicatedObjectCount = 1024;
+	ParamsOverride.InitialNetObjectListCount = 1024;
+
+	CreateServers(FNetRefHandle::MaxReplicationSystemCount, &ParamsOverride);
 
 	TArrayView<FReplicationSystemTestServer*> ServersView = GetAllServers();
 	UE_NET_ASSERT_EQ(unsigned(ServersView.Num()), unsigned(FNetRefHandle::MaxReplicationSystemCount));
