Added TestParallelTask to ReplicationSystemTest, which checks the Parallel implementation of Iris Object Polling polls the correct number of objects when running multiple tasks

#tests Tested with ReplicationSystemTest
#rnx

diff --git a/Engine/Plugins/Runtime/ReplicationSystemTestPlugin/Source/Private/Tests/Core/TestParallelTask.cpp b/Engine/Plugins/Runtime/ReplicationSystemTestPlugin/Source/Private/Tests/Core/TestParallelTask.cpp
new file mode 100644
--- /dev/null
+++ b/Engine/Plugins/Runtime/ReplicationSystemTestPlugin/Source/Private/Tests/Core/TestParallelTask.cpp
@@ -0,0 +1,61 @@
+// Copyright Epic Games, Inc. All Rights Reserved.
+
+#include "NetworkAutomationTest.h"
+#include "NetworkAutomationTestMacros.h"
+#include "Tests/ReplicationSystem/ReplicationSystemServerClientTestFixture.h"
+#include "Iris/ReplicationSystem/ReplicationSystemInternal.h"
+#include "Iris/ReplicationSystem/Polling/ObjectPoller.h"
+#include "Async/TaskGraphInterfaces.h"
+
+namespace UE::Net::Private
+{
+
+class FTestParallelTaskFixture : public FReplicationSystemServerClientTestFixture
+{
+protected:
+	virtual void SetUp() override
+	{
+		ServerParamsOverride.bAllowParallelTasks = true;
+
+		FReplicationSystemServerClientTestFixture::SetUp();
+
+		// Add a client
+		Client = CreateClient();
+	}
+	
+	FReplicationSystemTestClient* Client = nullptr;
+};
+
+UE_NET_TEST_FIXTURE(FTestParallelTaskFixture, PollTask)
+{	
+	//This test is designed to make sure that we're polling the correct number of objects when running Parallel Tasks
+	
+	const uint32 NumReplicatedTestObjects = 2048;// We need enough objects that several tasks have work do
+	for (uint32 i = 0; i < NumReplicatedTestObjects; ++i)
+	{
+		UReplicatedTestObject* ServerObject = Server->CreateObject(0, 0);
+	}
+	
+	Server->NetUpdate();
+
+	const FNetRefHandleManager& NetRefHandleManager = Server->GetReplicationSystem()->GetReplicationSystemInternal()->GetNetRefHandleManager();
+	FNetBitArrayView ObjectsConsideredForPolling = NetRefHandleManager.GetAssignedInternalIndices();
+	ensure(ObjectsConsideredForPolling.IsAnyBitSet());
+
+
+	//Set up object poller and track number of objects polled
+	FObjectPoller::FInitParams PollerInitParams;
+	PollerInitParams.ObjectReplicationBridge = Server->GetReplicationBridge();
+	PollerInitParams.ReplicationSystemInternal = Server->GetReplicationSystem()->GetReplicationSystemInternal();
+
+	FObjectPoller Poller(PollerInitParams);
+	Poller.PollAndCopyObjects(ObjectsConsideredForPolling);
+	
+	FObjectPoller::FPreUpdateAndPollStats PollStats = Poller.GetPollStats();
+
+	UE_NET_ASSERT_EQ(NumReplicatedTestObjects, PollStats.PolledObjectCount);
+
+	Server->PostSendUpdate();
+}
+
+} // end namespace UE::Net::Private
diff --git a/Engine/Plugins/Runtime/ReplicationSystemTestPlugin/Source/Private/Tests/ReplicationSystem/ReplicationSystemServerClientTestFixture.cpp b/Engine/Plugins/Runtime/ReplicationSystemTestPlugin/Source/Private/Tests/ReplicationSystem/ReplicationSystemServerClientTestFixture.cpp
--- a/Engine/Plugins/Runtime/ReplicationSystemTestPlugin/Source/Private/Tests/ReplicationSystem/ReplicationSystemServerClientTestFixture.cpp
+++ b/Engine/Plugins/Runtime/ReplicationSystemTestPlugin/Source/Private/Tests/ReplicationSystem/ReplicationSystemServerClientTestFixture.cpp
@@ -109,9 +109,9 @@ void FDataStreamTestUtil::AddDataStreamDefinition(const TCHAR* StreamName, const
 }
 
 // FReplicationSystemTestNode implementation
-FReplicationSystemTestNode::FReplicationSystemTestNode(bool bIsServer, const TCHAR* Name)
+FReplicationSystemTestNode::FReplicationSystemTestNode(bool bIsServer, const TCHAR* Name, const FReplicationSystemTestNode::FReplicationSystemParamsOverride* ParamsOverride)
 {
-	Setup(bIsServer, Name);
+	Setup(bIsServer, Name, ParamsOverride);
 }
 
 FReplicationSystemTestNode::FReplicationSystemTestNode(FReplicationSystemTestNode::EDelaySetup)
@@ -155,6 +155,7 @@ void FReplicationSystemTestNode::Setup(bool bIsServer, const TCHAR* Name, const
 		Params.InitialNetObjectListCount = ParamsOverride->InitialNetObjectListCount > 0 ? ParamsOverride->InitialNetObjectListCount : Params.InitialNetObjectListCount;
 		Params.NetObjectListGrowCount = ParamsOverride->NetObjectListGrowCount > 0 ? ParamsOverride->NetObjectListGrowCount : Params.NetObjectListGrowCount;
 		Params.bUseRemoteObjectReferences = ParamsOverride->bUseRemoteObjectReferences;
+		Params.bAllowParallelTasks = ParamsOverride->bAllowParallelTasks;
 	}
 
 	LOG_SCOPE_VERBOSITY_OVERRIDE(LogIris, ELogVerbosity::Error);
@@ -519,8 +520,8 @@ bool FReplicationSystemTestClient::UpdateAndSend(FReplicationSystemTestServer* S
 // Class FReplicationSystemTestServer
 //*****************************************************************************
 
-FReplicationSystemTestServer::FReplicationSystemTestServer(const TCHAR* Name)
-	: FReplicationSystemTestNode(true, Name)
+FReplicationSystemTestServer::FReplicationSystemTestServer(const TCHAR* Name, const FReplicationSystemTestNode::FReplicationSystemParamsOverride* ParamsOverride)
+	: FReplicationSystemTestNode(true, Name, ParamsOverride)
 {
 }
 
@@ -581,7 +582,7 @@ void FReplicationSystemServerClientTestFixture::SetUp()
 
 	NetTokenDataStoreUtil.SetUp();
 
-	Server = new FReplicationSystemTestServer(GetName());
+	Server = new FReplicationSystemTestServer(GetName(), &ServerParamsOverride);
 }
 
 void FReplicationSystemServerClientTestFixture::TearDown()
diff --git a/Engine/Plugins/Runtime/ReplicationSystemTestPlugin/Source/Private/Tests/ReplicationSystem/ReplicationSystemServerClientTestFixture.h b/Engine/Plugins/Runtime/ReplicationSystemTestPlugin/Source/Private/Tests/ReplicationSystem/ReplicationSystemServerClientTestFixture.h
--- a/Engine/Plugins/Runtime/ReplicationSystemTestPlugin/Source/Private/Tests/ReplicationSystem/ReplicationSystemServerClientTestFixture.h
+++ b/Engine/Plugins/Runtime/ReplicationSystemTestPlugin/Source/Private/Tests/ReplicationSystem/ReplicationSystemServerClientTestFixture.h
@@ -100,12 +100,13 @@ public:
 		uint32 InitialNetObjectListCount = 0;
 		uint32 NetObjectListGrowCount = 0;
 		bool bUseRemoteObjectReferences = false;
+		bool bAllowParallelTasks = false;
 	};
 
 	enum EDelaySetup { DelaySetup=0 };
 
 public:
-	FReplicationSystemTestNode(bool bIsServer, const TCHAR* Name);
+	FReplicationSystemTestNode(bool bIsServer, const TCHAR* Name, const FReplicationSystemTestNode::FReplicationSystemParamsOverride* ParamsOverride = nullptr);
 	explicit FReplicationSystemTestNode(FReplicationSystemTestNode::EDelaySetup);
 	virtual ~FReplicationSystemTestNode();
 
@@ -223,7 +224,7 @@ public:
 class FReplicationSystemTestServer : public FReplicationSystemTestNode
 {
 public:
-	explicit FReplicationSystemTestServer(const TCHAR* Name);
+	explicit FReplicationSystemTestServer(const TCHAR* Name, const FReplicationSystemTestNode::FReplicationSystemParamsOverride* ParamsOverride = nullptr);
 	explicit FReplicationSystemTestServer(FReplicationSystemTestNode::EDelaySetup Delay) : FReplicationSystemTestNode(Delay) {}
 
 	// Send data and deliver to the client if bDeliver is true
@@ -261,6 +262,7 @@ protected:
 	FNetTokenDataStoreTestUtil NetTokenDataStoreUtil;
 	FReplicationSystemTestServer* Server;
 	TArray<FReplicationSystemTestClient*> Clients;
+	FReplicationSystemTestNode::FReplicationSystemParamsOverride ServerParamsOverride;
 };
 
 }
