Suppressing a warning that complains about the intended behaviour of a test that ensures that a subobject is not replicated when the root object is torn off

diff --git a/Engine/Plugins/Runtime/ReplicationSystemTestPlugin/Source/Private/Tests/ReplicationSystem/TestReplicationSystem.cpp b/Engine/Plugins/Runtime/ReplicationSystemTestPlugin/Source/Private/Tests/ReplicationSystem/TestReplicationSystem.cpp
--- a/Engine/Plugins/Runtime/ReplicationSystemTestPlugin/Source/Private/Tests/ReplicationSystem/TestReplicationSystem.cpp
+++ b/Engine/Plugins/Runtime/ReplicationSystemTestPlugin/Source/Private/Tests/ReplicationSystem/TestReplicationSystem.cpp
@@ -1632,7 +1632,6 @@ UE_NET_TEST_FIXTURE(FReplicationSystemServerClientTestFixture, TestImmediateTear
 // Test creating a subobject after the root is already torn off
 UE_NET_TEST_FIXTURE(FReplicationSystemServerClientTestFixture, TestCreatingSubobjectAfterRootObjectTearOff)
 {
-	//LOG_SCOPE_VERBOSITY_OVERRIDE(LogIrisBridge, ELogVerbosity::Verbose);
 
 	// Add a client
 	FReplicationSystemTestClient* Client = CreateClient();
@@ -1653,10 +1652,14 @@ UE_NET_TEST_FIXTURE(FReplicationSystemServerClientTestFixture, TestCreatingSubob
 
 	{
 		FTestEnsureScope EnsureScope;
-
+		UTestReplicatedIrisObject* ServerSubObject1 = nullptr;
 		// Add a new subobject. This shouldn't become a replicated subobject of the already torn off root object.
-		UTestReplicatedIrisObject* ServerSubObject1 = Server->CreateSubObject(ServerObject->NetRefHandle);
-		UE_NET_ASSERT_EQ(EnsureScope.GetCount(), 0);
+		{
+			LOG_SCOPE_VERBOSITY_OVERRIDE(LogIris, ELogVerbosity::Error);
+			ServerSubObject1 = Server->CreateSubObject(ServerObject->NetRefHandle);
+			UE_NET_ASSERT_EQ(EnsureScope.GetCount(), 0);
+		}
+
 		// The subobject should not start replicating.
 		UE_NET_ASSERT_FALSE(ServerSubObject1->NetRefHandle.IsValid());
 
