Iris
- Fixed compilation error in test, fallout from CL 43192071

diff --git a/Engine/Plugins/Runtime/ReplicationSystemTestPlugin/Source/Private/Tests/Serialization/TestFieldPathNetSerializer.cpp b/Engine/Plugins/Runtime/ReplicationSystemTestPlugin/Source/Private/Tests/Serialization/TestFieldPathNetSerializer.cpp
--- a/Engine/Plugins/Runtime/ReplicationSystemTestPlugin/Source/Private/Tests/Serialization/TestFieldPathNetSerializer.cpp
+++ b/Engine/Plugins/Runtime/ReplicationSystemTestPlugin/Source/Private/Tests/Serialization/TestFieldPathNetSerializer.cpp
@@ -449,6 +449,9 @@ void FTestFieldPathNetSerializer::TestRoundtripSerialize()
 	DeserializeArgs.NetSerializerConfig = NetSerializerConfigParam(&SerializerConfig);
 	DeserializeArgs.Target = NetSerializerValuePointer(StateBuffer1);
 
+	EIrisAsyncLoadingPriority IrisAsyncLoadingPriority = EIrisAsyncLoadingPriority::Default;
+
+	TArray<FNetRefHandle> TempMustBeMappedReferences;
 	TArrayView<const FFieldPath> AllTheValues[] = { MakeArrayView(TestValues), MakeArrayView(UnresolvableTestValues) };
 	for (const TArrayView<const FFieldPath>& Values : AllTheValues)
 	{
@@ -481,7 +484,8 @@ void FTestFieldPathNetSerializer::TestRoundtripSerialize()
 
 			// Read exports
 			ClientNetSerializationContext.GetBitStreamReader()->Seek(ExportsPos);
-			ClientNetSerializationContext.GetInternalContext()->ObjectReferenceCache->ReadExports(ClientNetSerializationContext, nullptr);
+			TempMustBeMappedReferences.Reset();
+			ClientNetSerializationContext.GetInternalContext()->ObjectReferenceCache->ReadExports(FNetRefHandle(), ClientNetSerializationContext, &TempMustBeMappedReferences, IrisAsyncLoadingPriority);
 			ClientNetSerializationContext.GetBitStreamReader()->Seek(0U);
 
 			// Deserialize data
