Fix ReplicationSystemTest compile error.

diff --git a/Engine/Plugins/Runtime/ReplicationSystemTestPlugin/Source/Private/Tests/ReplicationSystem/TestMultiReplicationSystems.cpp b/Engine/Plugins/Runtime/ReplicationSystemTestPlugin/Source/Private/Tests/ReplicationSystem/TestMultiReplicationSystems.cpp
--- a/Engine/Plugins/Runtime/ReplicationSystemTestPlugin/Source/Private/Tests/ReplicationSystem/TestMultiReplicationSystems.cpp
+++ b/Engine/Plugins/Runtime/ReplicationSystemTestPlugin/Source/Private/Tests/ReplicationSystem/TestMultiReplicationSystems.cpp
@@ -116,7 +116,7 @@ UE_NET_TEST_FIXTURE(FMultiReplicationSystemsTestFixture, ReplicationSystemArrayM
 
 	// View should shrink as all the null entries at the end are removed
 	TArrayView<UReplicationSystem*> ReplicationSystemsView3 = FReplicationSystemFactory::GetAllReplicationSystems();
-	UE_NET_ASSERT_EQ(unsigned(ReplicationSystemsView3.Num()), 2);
+	UE_NET_ASSERT_EQ(unsigned(ReplicationSystemsView3.Num()), 2U);
 }
 
 UE_NET_TEST_FIXTURE(FMultiReplicationSystemsTestFixture, CanReplicateObjectOnMultipleReplicationSystems)
