Iris
- Fixed broken test, that now fails after the issue has been fixed.

diff --git a/Engine/Plugins/Runtime/ReplicationSystemTestPlugin/Source/Private/Tests/ReplicationSystem/TestReplicationSystem.cpp b/Engine/Plugins/Runtime/ReplicationSystemTestPlugin/Source/Private/Tests/ReplicationSystem/TestReplicationSystem.cpp
--- a/Engine/Plugins/Runtime/ReplicationSystemTestPlugin/Source/Private/Tests/ReplicationSystem/TestReplicationSystem.cpp
+++ b/Engine/Plugins/Runtime/ReplicationSystemTestPlugin/Source/Private/Tests/ReplicationSystem/TestReplicationSystem.cpp
@@ -2188,7 +2188,7 @@ UE_NET_TEST_FIXTURE(FReplicationSystemServerClientTestFixture, TestClientCanSkip
 	// We expect replication of ObjectA to have succeeded this time
 	{
 		UTestReplicatedIrisObject* ClientObjectA = Cast<UTestReplicatedIrisObject>(Client->GetReplicationBridge()->GetReplicatedObject(ServerObjectA->NetRefHandle));
-		UE_NET_ASSERT_TRUE(ClientObjectA == nullptr);
+		UE_NET_ASSERT_TRUE(ClientObjectA != nullptr);
 	}
 }
 
