Fixed INVTEXT literals using hex escape sequences, in favor of using Unicode escape sequences

#rnx

diff --git a/Engine/Plugins/Runtime/ReplicationSystemTestPlugin/Source/Private/Tests/Serialization/TestLastResortPropertyNetSerializer.cpp b/Engine/Plugins/Runtime/ReplicationSystemTestPlugin/Source/Private/Tests/Serialization/TestLastResortPropertyNetSerializer.cpp
--- a/Engine/Plugins/Runtime/ReplicationSystemTestPlugin/Source/Private/Tests/Serialization/TestLastResortPropertyNetSerializer.cpp
+++ b/Engine/Plugins/Runtime/ReplicationSystemTestPlugin/Source/Private/Tests/Serialization/TestLastResortPropertyNetSerializer.cpp
@@ -58,7 +58,7 @@ const FText FTestLastResortPropertyNetSerializer::TestValues[] =
 {
 	FText(),
 	INVTEXT("Just a regular ASCII string"),
-	INVTEXT("\xD83D\xDE06\xD83C\xDC1C\x20AC\x00A9"),
+	INVTEXT("\uD83D\uDE06\uD83C\uDC1C\u20AC\u00A9"),
 };
 
 const SIZE_T FTestLastResortPropertyNetSerializer::TestValueCount = sizeof(TestValues)/sizeof(TestValues[0]);
