Iris serializers for FRemoteServerId and FRemoteObjectId.

diff --git a/Engine/Plugins/Runtime/ReplicationSystemTestPlugin/Source/Private/TestMessage.cpp b/Engine/Plugins/Runtime/ReplicationSystemTestPlugin/Source/Private/TestMessage.cpp
--- a/Engine/Plugins/Runtime/ReplicationSystemTestPlugin/Source/Private/TestMessage.cpp
+++ b/Engine/Plugins/Runtime/ReplicationSystemTestPlugin/Source/Private/TestMessage.cpp
@@ -5,6 +5,7 @@
 #include "Misc/StringBuilder.h"
 #include "Net/Core/NetHandle/NetHandle.h"
 #include "Net/Core/NetToken/NetToken.h"
+#include "UObject/RemoteObjectTypes.h"
 
 namespace UE::Net
 {
@@ -28,4 +29,14 @@ FTestMessage& operator<<(FTestMessage& Message, const FNetToken& NetToken)
 	return Message << NetToken.ToString();
 }
 
+FTestMessage& operator<<(FTestMessage& Message, const FRemoteServerId& RemoteServerId)
+{
+	return Message << RemoteServerId.ToString();
+}
+
+FTestMessage& operator<<(FTestMessage& Message, const FRemoteObjectId& RemoteObjectId)
+{
+	return Message << RemoteObjectId.ToString();
+}
+
 }
diff --git a/Engine/Plugins/Runtime/ReplicationSystemTestPlugin/Source/Private/TestMessage.h b/Engine/Plugins/Runtime/ReplicationSystemTestPlugin/Source/Private/TestMessage.h
--- a/Engine/Plugins/Runtime/ReplicationSystemTestPlugin/Source/Private/TestMessage.h
+++ b/Engine/Plugins/Runtime/ReplicationSystemTestPlugin/Source/Private/TestMessage.h
@@ -8,6 +8,9 @@
 #include "Math/Vector.h"
 #include <cinttypes>
 
+struct FRemoteServerId;
+struct FRemoteObjectId;
+
 namespace UE::Net
 {
 	class FNetHandle;
@@ -145,6 +148,10 @@ FTestMessage& operator<<(FTestMessage& Message, const FNetRefHandle& NetRefHandl
 
 FTestMessage& operator<<(FTestMessage& Message, const FNetToken& NetToken);
 
+FTestMessage& operator<<(FTestMessage& Message, const FRemoteServerId& RemoteServerId);
+
+FTestMessage& operator<<(FTestMessage& Message, const FRemoteObjectId& RemoteObjectId);
+
 #if defined(_WIN32) && !defined(_WIN64)
 inline FTestMessage& operator<<(FTestMessage& Message, SIZE_T Value)
 {
diff --git a/Engine/Plugins/Runtime/ReplicationSystemTestPlugin/Source/Private/Tests/Serialization/TestRemoteObjectIdNetSerializer.cpp b/Engine/Plugins/Runtime/ReplicationSystemTestPlugin/Source/Private/Tests/Serialization/TestRemoteObjectIdNetSerializer.cpp
new file mode 100644
--- /dev/null
+++ b/Engine/Plugins/Runtime/ReplicationSystemTestPlugin/Source/Private/Tests/Serialization/TestRemoteObjectIdNetSerializer.cpp
@@ -0,0 +1,151 @@
+// Copyright Epic Games, Inc. All Rights Reserved.
+
+#include "TestNetSerializerFixture.h"
+#include "Iris/Serialization/NetSerializers.h"
+#include "Iris/Serialization/RemoteObjectIdNetSerializer.h"
+
+namespace UE::Net::Private
+{
+
+static FTestMessage& PrintRemoteObjectIdNetSerializerConfig(FTestMessage& Message, const FNetSerializerConfig& InConfig)
+{
+	return Message;
+}
+
+class FTestRemoteObjectIdNetSerializer : public TTestNetSerializerFixture<PrintRemoteObjectIdNetSerializerConfig, FRemoteObjectId>
+{
+public:
+	FTestRemoteObjectIdNetSerializer() : TTestNetSerializerFixture<PrintRemoteObjectIdNetSerializerConfig, FRemoteObjectId>(UE_NET_GET_SERIALIZER(FRemoteObjectIdNetSerializer)) {}
+
+	void TestIsEqual();
+	void TestValidate();
+	void TestQuantize();
+	void TestSerialize();
+
+protected:
+	virtual void SetUp() override;
+	virtual void TearDown() override;
+
+	void SetUpTestValues();
+
+	typedef TTestNetSerializerFixture<PrintRemoteObjectIdNetSerializerConfig, FRemoteObjectId> Super;
+
+	TArray<FRemoteObjectId> TestValues;
+
+	static FRemoteObjectIdNetSerializerConfig SerializerConfig;
+};
+
+FRemoteObjectIdNetSerializerConfig FTestRemoteObjectIdNetSerializer::SerializerConfig;
+
+UE_NET_TEST_FIXTURE(FTestRemoteObjectIdNetSerializer, HasTestValues)
+{
+	UE_NET_ASSERT_GT_MSG(TestValues.Num(), SIZE_T(0), "No test values found");
+}
+
+UE_NET_TEST_FIXTURE(FTestRemoteObjectIdNetSerializer, TestIsEqual)
+{
+	TestIsEqual();
+}
+
+UE_NET_TEST_FIXTURE(FTestRemoteObjectIdNetSerializer, TestValidate)
+{
+	TestValidate();
+}
+
+UE_NET_TEST_FIXTURE(FTestRemoteObjectIdNetSerializer, TestQuantize)
+{
+	TestQuantize();
+}
+
+UE_NET_TEST_FIXTURE(FTestRemoteObjectIdNetSerializer, TestSerialize)
+{
+	TestSerialize();
+}
+
+void FTestRemoteObjectIdNetSerializer::SetUp()
+{
+	Super::SetUp();
+
+	if (TestValues.Num() == 0)
+	{
+		SetUpTestValues();
+	}
+}
+
+void FTestRemoteObjectIdNetSerializer::TearDown()
+{
+	TestValues.Reset();
+
+	Super::TearDown();
+}
+
+void FTestRemoteObjectIdNetSerializer::SetUpTestValues()
+{
+	TestValues.Add(FRemoteObjectId());
+
+	TestValues.Add(FRemoteObjectId(FRemoteServerId(ERemoteServerIdConstants::FirstValid), 0));
+	TestValues.Add(FRemoteObjectId(FRemoteServerId(ERemoteServerIdConstants::FirstValid), 1));
+	TestValues.Add(FRemoteObjectId(FRemoteServerId(ERemoteServerIdConstants::FirstValid), MAX_REMOTE_OBJECT_SERIAL_NUMBER / 2));
+	TestValues.Add(FRemoteObjectId(FRemoteServerId(ERemoteServerIdConstants::FirstValid), MAX_REMOTE_OBJECT_SERIAL_NUMBER - 1));
+	TestValues.Add(FRemoteObjectId(FRemoteServerId(ERemoteServerIdConstants::FirstValid), MAX_REMOTE_OBJECT_SERIAL_NUMBER));
+
+	TestValues.Add(FRemoteObjectId(FRemoteServerId(ERemoteServerIdConstants::Max), 0));
+	TestValues.Add(FRemoteObjectId(FRemoteServerId(ERemoteServerIdConstants::Max), 1));
+	TestValues.Add(FRemoteObjectId(FRemoteServerId(ERemoteServerIdConstants::Max), MAX_REMOTE_OBJECT_SERIAL_NUMBER / 2));
+	TestValues.Add(FRemoteObjectId(FRemoteServerId(ERemoteServerIdConstants::Max), MAX_REMOTE_OBJECT_SERIAL_NUMBER - 1));
+	TestValues.Add(FRemoteObjectId(FRemoteServerId(ERemoteServerIdConstants::Max), MAX_REMOTE_OBJECT_SERIAL_NUMBER));
+}
+
+void FTestRemoteObjectIdNetSerializer::TestIsEqual()
+{
+	const int32 TestValueCount = TestValues.Num();
+
+	TArray<FRemoteObjectId> CompareValues[2];
+	TArray<bool> ExpectedResults[2];
+
+	CompareValues[0] = TestValues;
+	ExpectedResults[0].Init(true, TestValueCount);
+
+	for (int32 ValueIt = 0; ValueIt < TestValueCount; ++ValueIt)
+	{
+		const int32 NextValueIndex = (ValueIt + 1U) % TestValueCount;
+
+		CompareValues[1].Add(TestValues[NextValueIndex]);
+		ExpectedResults[1].Add(TestValues[ValueIt] == TestValues[NextValueIndex]);
+	}
+
+	// Do two rounds of testing per config, one where we compare each value with itself and one where we compare against another value.
+	for (int32 TestRoundIt : {0, 1})
+	{
+		// Do both quantized and regular compares
+		for (const bool bQuantizedCompare : {true, false})
+		{
+			const bool bSuccess = Super::TestIsEqual(TestValues.GetData(), CompareValues[TestRoundIt].GetData(), ExpectedResults[TestRoundIt].GetData(), TestValueCount, SerializerConfig, bQuantizedCompare);
+			if (!bSuccess)
+			{
+				return;
+			}
+		}
+	}
+}
+
+void FTestRemoteObjectIdNetSerializer::TestValidate()
+{
+	TArray<bool> ExpectedResults;
+	ExpectedResults.Init(true, TestValues.Num());
+
+	Super::TestValidate(TestValues.GetData(), ExpectedResults.GetData(), TestValues.Num(), SerializerConfig);
+}
+
+void FTestRemoteObjectIdNetSerializer::TestQuantize()
+{
+	Super::TestQuantize(TestValues.GetData(), TestValues.Num(), SerializerConfig);
+}
+
+void FTestRemoteObjectIdNetSerializer::TestSerialize()
+{
+	constexpr bool bQuantizedCompare = false;
+	Super::TestSerialize(TestValues.GetData(), TestValues.GetData(), TestValues.Num(), SerializerConfig, bQuantizedCompare);
+}
+
+}
diff --git a/Engine/Plugins/Runtime/ReplicationSystemTestPlugin/Source/Private/Tests/Serialization/TestRemoteServerIdNetSerializer.cpp b/Engine/Plugins/Runtime/ReplicationSystemTestPlugin/Source/Private/Tests/Serialization/TestRemoteServerIdNetSerializer.cpp
new file mode 100644
--- /dev/null
+++ b/Engine/Plugins/Runtime/ReplicationSystemTestPlugin/Source/Private/Tests/Serialization/TestRemoteServerIdNetSerializer.cpp
@@ -0,0 +1,148 @@
+// Copyright Epic Games, Inc. All Rights Reserved.
+
+#include "TestNetSerializerFixture.h"
+#include "Iris/Serialization/NetSerializers.h"
+#include "Iris/Serialization/RemoteServerIdNetSerializer.h"
+
+namespace UE::Net::Private
+{
+
+static FTestMessage& PrintRemoteServerIdNetSerializerConfig(FTestMessage& Message, const FNetSerializerConfig& InConfig)
+{
+	return Message;
+}
+
+class FTestRemoteServerIdNetSerializer : public TTestNetSerializerFixture<PrintRemoteServerIdNetSerializerConfig, FRemoteServerId>
+{
+public:
+	FTestRemoteServerIdNetSerializer() : TTestNetSerializerFixture<PrintRemoteServerIdNetSerializerConfig, FRemoteServerId>(UE_NET_GET_SERIALIZER(FRemoteServerIdNetSerializer)) {}
+
+	void TestIsEqual();
+	void TestValidate();
+	void TestQuantize();
+	void TestSerialize();
+
+protected:
+	virtual void SetUp() override;
+	virtual void TearDown() override;
+
+	void SetUpTestValues();
+
+	typedef TTestNetSerializerFixture<PrintRemoteServerIdNetSerializerConfig, FRemoteServerId> Super;
+
+	TArray<FRemoteServerId> TestValues;
+
+	static FRemoteServerIdNetSerializerConfig SerializerConfig;
+};
+
+FRemoteServerIdNetSerializerConfig FTestRemoteServerIdNetSerializer::SerializerConfig;
+
+UE_NET_TEST_FIXTURE(FTestRemoteServerIdNetSerializer, HasTestValues)
+{
+	UE_NET_ASSERT_GT_MSG(TestValues.Num(), SIZE_T(0), "No test values found");
+}
+
+UE_NET_TEST_FIXTURE(FTestRemoteServerIdNetSerializer, TestIsEqual)
+{
+	TestIsEqual();
+}
+
+UE_NET_TEST_FIXTURE(FTestRemoteServerIdNetSerializer, TestValidate)
+{
+	TestValidate();
+}
+
+UE_NET_TEST_FIXTURE(FTestRemoteServerIdNetSerializer, TestQuantize)
+{
+	TestQuantize();
+}
+
+UE_NET_TEST_FIXTURE(FTestRemoteServerIdNetSerializer, TestSerialize)
+{
+	TestSerialize();
+}
+
+void FTestRemoteServerIdNetSerializer::SetUp()
+{
+	Super::SetUp();
+
+	if (TestValues.Num() == 0)
+	{
+		SetUpTestValues();
+	}
+}
+
+void FTestRemoteServerIdNetSerializer::TearDown()
+{
+	TestValues.Reset();
+
+	Super::TearDown();
+}
+
+void FTestRemoteServerIdNetSerializer::SetUpTestValues()
+{
+	TestValues.Add(FRemoteServerId());
+
+	TestValues.Add(FRemoteServerId(ERemoteServerIdConstants::FirstValid));
+	TestValues.Add(FRemoteServerId::FromIdNumber(static_cast<uint32>(ERemoteServerIdConstants::FirstValid) + 1));
+	TestValues.Add(FRemoteServerId(ERemoteServerIdConstants::Database));
+	TestValues.Add(FRemoteServerId(ERemoteServerIdConstants::Asset));
+	TestValues.Add(FRemoteServerId(ERemoteServerIdConstants::Local));
+	TestValues.Add(FRemoteServerId::FromIdNumber(static_cast<uint32>(ERemoteServerIdConstants::Max) / 2));
+	TestValues.Add(FRemoteServerId::FromIdNumber(static_cast<uint32>(ERemoteServerIdConstants::FirstReserved) - 1));
+	TestValues.Add(FRemoteServerId(ERemoteServerIdConstants::FirstReserved));
+}
+
+void FTestRemoteServerIdNetSerializer::TestIsEqual()
+{
+	const int32 TestValueCount = TestValues.Num();
+
+	TArray<FRemoteServerId> CompareValues[2];
+	TArray<bool> ExpectedResults[2];
+
+	CompareValues[0] = TestValues;
+	ExpectedResults[0].Init(true, TestValueCount);
+
+	for (int32 ValueIt = 0; ValueIt < TestValueCount; ++ValueIt)
+	{
+		const int32 NextValueIndex = (ValueIt + 1U) % TestValueCount;
+
+		CompareValues[1].Add(TestValues[NextValueIndex]);
+		ExpectedResults[1].Add(TestValues[ValueIt] == TestValues[NextValueIndex]);
+	}
+
+	// Do two rounds of testing per config, one where we compare each value with itself and one where we compare against another value.
+	for (int32 TestRoundIt : {0, 1})
+	{
+		// Do both quantized and regular compares
+		for (const bool bQuantizedCompare : {true, false})
+		{
+			const bool bSuccess = Super::TestIsEqual(TestValues.GetData(), CompareValues[TestRoundIt].GetData(), ExpectedResults[TestRoundIt].GetData(), TestValueCount, SerializerConfig, bQuantizedCompare);
+			if (!bSuccess)
+			{
+				return;
+			}
+		}
+	}
+}
+
+void FTestRemoteServerIdNetSerializer::TestValidate()
+{
+	TArray<bool> ExpectedResults;
+	ExpectedResults.Init(true, TestValues.Num());
+
+	Super::TestValidate(TestValues.GetData(), ExpectedResults.GetData(), TestValues.Num(), SerializerConfig);
+}
+
+void FTestRemoteServerIdNetSerializer::TestQuantize()
+{
+	Super::TestQuantize(TestValues.GetData(), TestValues.Num(), SerializerConfig);
+}
+
+void FTestRemoteServerIdNetSerializer::TestSerialize()
+{
+	constexpr bool bQuantizedCompare = false;
+	Super::TestSerialize(TestValues.GetData(), TestValues.GetData(), TestValues.Num(), SerializerConfig, bQuantizedCompare);
+}
+
+}
diff --git a/Engine/Source/Runtime/CoreUObject/Public/UObject/RemoteObjectTypes.h b/Engine/Source/Runtime/CoreUObject/Public/UObject/RemoteObjectTypes.h
--- a/Engine/Source/Runtime/CoreUObject/Public/UObject/RemoteObjectTypes.h
+++ b/Engine/Source/Runtime/CoreUObject/Public/UObject/RemoteObjectTypes.h
@@ -15,6 +15,8 @@ class FArchive;
 namespace UE::Net
 {
 	struct FRemoteObjectReferenceNetSerializer;
+	struct FRemoteServerIdNetSerializer;
+	struct FRemoteObjectIdNetSerializer;
 }
 namespace UE::RemoteObject::Private
 {
@@ -48,6 +50,8 @@ struct FRemoteServerId
 	friend struct FRemoteObjectId;
 	friend UE::RemoteObject::Private::FRemoteIdLocalizationHelper;
 	friend struct UE::Net::FRemoteObjectReferenceNetSerializer;
+	friend struct UE::Net::FRemoteServerIdNetSerializer;
+	friend struct UE::Net::FRemoteObjectIdNetSerializer;
 
 	FRemoteServerId() = default;
 	FRemoteServerId(const FRemoteServerId&) = default;
@@ -319,6 +323,7 @@ public:
 
 	friend COREUOBJECT_API FArchive& operator<<(FArchive& Ar, FRemoteObjectId& Id);
 	friend struct UE::Net::FRemoteObjectReferenceNetSerializer;
+	friend struct UE::Net::FRemoteObjectIdNetSerializer;
 	friend UE::RemoteObject::Private::FRemoteIdLocalizationHelper;
 	friend struct FUObjectItem;
 	friend class FNetGUIDCache;	
diff --git a/Engine/Source/Runtime/Experimental/Iris/Core/Private/Iris/ReplicationState/DefaultPropertyNetSerializerInfos.cpp b/Engine/Source/Runtime/Experimental/Iris/Core/Private/Iris/ReplicationState/DefaultPropertyNetSerializerInfos.cpp
--- a/Engine/Source/Runtime/Experimental/Iris/Core/Private/Iris/ReplicationState/DefaultPropertyNetSerializerInfos.cpp
+++ b/Engine/Source/Runtime/Experimental/Iris/Core/Private/Iris/ReplicationState/DefaultPropertyNetSerializerInfos.cpp
@@ -7,6 +7,8 @@
 #include "Iris/Serialization/InternalNetSerializers.h"
 #include "Iris/Serialization/DateTimeNetSerializer.h"
 #include "Iris/Serialization/RemoteObjectReferenceNetSerializer.h"
+#include "Iris/Serialization/RemoteServerIdNetSerializer.h"
+#include "Iris/Serialization/RemoteObjectIdNetSerializer.h"
 #include "UObject/TextProperty.h"
 
 namespace UE::Net::Private
@@ -33,6 +35,8 @@ UE_NET_IMPLEMENT_SIMPLE_NETSERIALIZER_INFO(FDoubleProperty, FDoubleNetSerializer
 static const FName PropertyNetSerializerRegistry_NAME_SoftObjectPath("SoftObjectPath");
 static const FName PropertyNetSerializerRegistry_NAME_SoftClassPath("SoftClassPath");
 static const FName PropertyNetSerializerRegistry_NAME_RemoteObjectReference("RemoteObjectReference");
+static const FName PropertyNetSerializerRegistry_NAME_RemoteServerId("RemoteServerId");
+static const FName PropertyNetSerializerRegistry_NAME_RemoteObjectId("RemoteObjectId");
 
 // Info struct for raw object pointers, can be used in RPC parameters
 struct FObjectPropertyNetSerializerInfo : public FPropertyNetSerializerInfo
@@ -83,6 +87,8 @@ UE_NET_IMPLEMENT_SIMPLE_NETSERIALIZER_INFO(FFieldPathProperty, FFieldPathNetSeri
 UE_NET_IMPLEMENT_NAMED_STRUCT_NETSERIALIZER_INFO(PropertyNetSerializerRegistry_NAME_SoftObjectPath, FSoftObjectPathNetSerializer);
 UE_NET_IMPLEMENT_NAMED_STRUCT_NETSERIALIZER_INFO(PropertyNetSerializerRegistry_NAME_SoftClassPath, FSoftClassPathNetSerializer);
 UE_NET_IMPLEMENT_NAMED_STRUCT_NETSERIALIZER_INFO(PropertyNetSerializerRegistry_NAME_RemoteObjectReference, FRemoteObjectReferenceNetSerializer);
+UE_NET_IMPLEMENT_NAMED_STRUCT_NETSERIALIZER_INFO(PropertyNetSerializerRegistry_NAME_RemoteServerId, FRemoteServerIdNetSerializer);
+UE_NET_IMPLEMENT_NAMED_STRUCT_NETSERIALIZER_INFO(PropertyNetSerializerRegistry_NAME_RemoteObjectId, FRemoteObjectIdNetSerializer);
 
 // Strings
 UE_NET_IMPLEMENT_SIMPLE_NETSERIALIZER_INFO(FStrProperty, FStringNetSerializer);
@@ -239,6 +245,8 @@ void RegisterDefaultPropertyNetSerializerInfos()
 	UE_NET_REGISTER_NETSERIALIZER_INFO(PropertyNetSerializerRegistry_NAME_SoftObjectPath);
 	UE_NET_REGISTER_NETSERIALIZER_INFO(PropertyNetSerializerRegistry_NAME_SoftClassPath);
 	UE_NET_REGISTER_NETSERIALIZER_INFO(PropertyNetSerializerRegistry_NAME_RemoteObjectReference);
+	UE_NET_REGISTER_NETSERIALIZER_INFO(PropertyNetSerializerRegistry_NAME_RemoteServerId);
+	UE_NET_REGISTER_NETSERIALIZER_INFO(PropertyNetSerializerRegistry_NAME_RemoteObjectId);
 
 	// String types
 	UE_NET_REGISTER_NETSERIALIZER_INFO(FNameProperty);
diff --git a/Engine/Source/Runtime/Experimental/Iris/Core/Private/Iris/Serialization/RemoteObjectIdNetSerializer.cpp b/Engine/Source/Runtime/Experimental/Iris/Core/Private/Iris/Serialization/RemoteObjectIdNetSerializer.cpp
--- a/Engine/Source/Runtime/Experimental/Iris/Core/Private/Iris/Serialization/RemoteObjectIdNetSerializer.cpp
+++ b/Engine/Source/Runtime/Experimental/Iris/Core/Private/Iris/Serialization/RemoteObjectIdNetSerializer.cpp
@@ -1,6 +1,109 @@
 // Copyright Epic Games, Inc. All Rights Reserved.
 
-#include "Iris/ReplicationState/PropertyNetSerializerInfoRegistry.h"
+#include "RemoteObjectIdNetSerializer.h"
+#include "Iris/Serialization/NetBitStreamUtil.h"
+#include "Iris/Serialization/NetSerializationContext.h"
 #include "UObject/RemoteObjectTypes.h"
 
-UE_NET_IMPLEMENT_NAMED_STRUCT_LASTRESORT_NETSERIALIZER_AND_REGISTRY_DELEGATES(RemoteObjectId);
+namespace UE::Net
+{
+
+struct FRemoteObjectIdNetSerializer
+{
+	// Version
+	static const uint32 Version = 0;
+
+	// Types
+	typedef FRemoteObjectId SourceType;
+	typedef uint64 QuantizedType;
+	typedef FRemoteObjectIdNetSerializerConfig ConfigType;
+
+	static const ConfigType DefaultConfig;
+
+	static_assert(sizeof(QuantizedType) == sizeof(FRemoteObjectId::Id), "Quantized RemoteObjectId type is not the correct size.");
+
+	static void Serialize(FNetSerializationContext&, const FNetSerializeArgs& Args);
+	static void Deserialize(FNetSerializationContext&, const FNetDeserializeArgs& Args);
+
+	static void Quantize(FNetSerializationContext&, const FNetQuantizeArgs& Args);
+	static void Dequantize(FNetSerializationContext&, const FNetDequantizeArgs& Args);
+
+	static bool IsEqual(FNetSerializationContext&, const FNetIsEqualArgs& Args);
+	static bool Validate(FNetSerializationContext&, const FNetValidateArgs& Args);
+};
+UE_NET_IMPLEMENT_SERIALIZER_INTERNAL(FRemoteObjectIdNetSerializer);
+
+const FRemoteObjectIdNetSerializer::ConfigType FRemoteObjectIdNetSerializer::DefaultConfig;
+
+void FRemoteObjectIdNetSerializer::Serialize(FNetSerializationContext& Context, const FNetSerializeArgs& Args)
+{
+	const QuantizedType& Value = *reinterpret_cast<const QuantizedType*>(Args.Source);
+
+	FNetBitStreamWriter* Writer = Context.GetBitStreamWriter();
+
+	WriteUint64(Writer, Value);
+}
+
+void FRemoteObjectIdNetSerializer::Deserialize(FNetSerializationContext& Context, const FNetDeserializeArgs& Args)
+{
+	QuantizedType& Value = *reinterpret_cast<QuantizedType*>(Args.Target);
+
+	FNetBitStreamReader* Reader = Context.GetBitStreamReader();
+
+	Value = ReadUint64(Reader);
+}
+
+void FRemoteObjectIdNetSerializer::Quantize(FNetSerializationContext& Context, const FNetQuantizeArgs& Args)
+{
+	const SourceType& SourceValue = *reinterpret_cast<const SourceType*>(Args.Source);
+	QuantizedType& TargetValue = *reinterpret_cast<QuantizedType*>(Args.Target);
+
+	FRemoteObjectId GlobalizedId = SourceValue.GetGlobalized();
+	TargetValue = GlobalizedId.Id;
+
+	FRemoteObjectId LocalizedId;
+	LocalizedId.Id = TargetValue;
+	checkf(LocalizedId.GetServerId().Id != (uint32)ERemoteServerIdConstants::Local, TEXT("Local server id should never be serialized"));
+}
+
+void FRemoteObjectIdNetSerializer::Dequantize(FNetSerializationContext& Context, const FNetDequantizeArgs& Args)
+{
+	SourceType& TargetValue = *reinterpret_cast<SourceType*>(Args.Target);
+	QuantizedType& QuantizedValue = *reinterpret_cast<QuantizedType*>(Args.Source);
+
+	FRemoteObjectId LocalizedId;
+	LocalizedId.Id = QuantizedValue;
+	checkf(LocalizedId.GetServerId().Id != (uint32)ERemoteServerIdConstants::Local, TEXT("Local server id should never be serialized"));
+	TargetValue = LocalizedId.GetLocalized();
+}
+
+bool FRemoteObjectIdNetSerializer::IsEqual(FNetSerializationContext& Context, const FNetIsEqualArgs& Args)
+{
+	if (Args.bStateIsQuantized)
+	{
+		const QuantizedType& QuantizedValue0 = *reinterpret_cast<const QuantizedType*>(Args.Source0);
+		const QuantizedType& QuantizedValue1 = *reinterpret_cast<const QuantizedType*>(Args.Source1);
+
+		return QuantizedValue0 == QuantizedValue1;
+	}
+	else
+	{
+		const SourceType& SourceValue0 = *reinterpret_cast<const SourceType*>(Args.Source0);
+		const SourceType& SourceValue1 = *reinterpret_cast<const SourceType*>(Args.Source1);
+
+		return SourceValue0 == SourceValue1;
+	}
+}
+
+bool FRemoteObjectIdNetSerializer::Validate(FNetSerializationContext& Context, const FNetValidateArgs& Args)
+{
+	const SourceType& SourceValue = *reinterpret_cast<const SourceType*>(Args.Source);
+
+	// Server ID and serial number should be within the valid range
+	const bool bServerIdValid = SourceValue.GetServerId().GetIdNumber() <= MAX_REMOTE_OBJECT_SERVER_ID;
+	const bool bSerialNumberValid = SourceValue.SerialNumber <= MAX_REMOTE_OBJECT_SERIAL_NUMBER;
+
+	return bServerIdValid && bSerialNumberValid;
+}
+
+}
diff --git a/Engine/Source/Runtime/Experimental/Iris/Core/Private/Iris/Serialization/RemoteObjectIdNetSerializer.h b/Engine/Source/Runtime/Experimental/Iris/Core/Private/Iris/Serialization/RemoteObjectIdNetSerializer.h
new file mode 100644
--- /dev/null
+++ b/Engine/Source/Runtime/Experimental/Iris/Core/Private/Iris/Serialization/RemoteObjectIdNetSerializer.h
@@ -0,0 +1,17 @@
+// Copyright Epic Games, Inc. All Rights Reserved.
+
+#pragma once
+
+#include "Iris/Serialization/InternalNetSerializer.h"
+#include "RemoteObjectIdNetSerializer.generated.h"
+
+USTRUCT()
+struct FRemoteObjectIdNetSerializerConfig : public FNetSerializerConfig
+{
+	GENERATED_BODY()
+};
+
+namespace UE::Net
+{
+	UE_NET_DECLARE_SERIALIZER_INTERNAL(FRemoteObjectIdNetSerializer, IRISCORE_API);
+}
diff --git a/Engine/Source/Runtime/Experimental/Iris/Core/Private/Iris/Serialization/RemoteServerIdNetSerializer.cpp b/Engine/Source/Runtime/Experimental/Iris/Core/Private/Iris/Serialization/RemoteServerIdNetSerializer.cpp
--- a/Engine/Source/Runtime/Experimental/Iris/Core/Private/Iris/Serialization/RemoteServerIdNetSerializer.cpp
+++ b/Engine/Source/Runtime/Experimental/Iris/Core/Private/Iris/Serialization/RemoteServerIdNetSerializer.cpp
@@ -1,6 +1,104 @@
 // Copyright Epic Games, Inc. All Rights Reserved.
 
-#include "Iris/ReplicationState/PropertyNetSerializerInfoRegistry.h"
+#include "RemoteServerIdNetSerializer.h"
+#include "Iris/Serialization/NetSerializationContext.h"
 #include "UObject/RemoteObjectTypes.h"
 
-UE_NET_IMPLEMENT_NAMED_STRUCT_LASTRESORT_NETSERIALIZER_AND_REGISTRY_DELEGATES(RemoteServerId);
+namespace UE::Net
+{
+
+struct FRemoteServerIdNetSerializer
+{
+	// Version
+	static const uint32 Version = 0;
+
+	// Types
+	typedef FRemoteServerId SourceType;
+	typedef uint16 QuantizedType;
+	typedef FRemoteServerIdNetSerializerConfig ConfigType;
+
+	static const ConfigType DefaultConfig;
+
+	static_assert(sizeof(QuantizedType) * 8U >= REMOTE_OBJECT_SERVER_ID_BIT_SIZE, "Quantized ServerId is not large enough to store maximum server ID");
+
+	static void Serialize(FNetSerializationContext&, const FNetSerializeArgs& Args);
+	static void Deserialize(FNetSerializationContext&, const FNetDeserializeArgs& Args);
+
+	static void Quantize(FNetSerializationContext&, const FNetQuantizeArgs& Args);
+	static void Dequantize(FNetSerializationContext&, const FNetDequantizeArgs& Args);
+
+	static bool IsEqual(FNetSerializationContext&, const FNetIsEqualArgs& Args);
+	static bool Validate(FNetSerializationContext&, const FNetValidateArgs& Args);
+};
+UE_NET_IMPLEMENT_SERIALIZER_INTERNAL(FRemoteServerIdNetSerializer);
+
+const FRemoteServerIdNetSerializer::ConfigType FRemoteServerIdNetSerializer::DefaultConfig;
+
+void FRemoteServerIdNetSerializer::Serialize(FNetSerializationContext& Context, const FNetSerializeArgs& Args)
+{
+	const QuantizedType& Value = *reinterpret_cast<const QuantizedType*>(Args.Source);
+
+	FNetBitStreamWriter* Writer = Context.GetBitStreamWriter();
+
+	Writer->WriteBits(Value, REMOTE_OBJECT_SERVER_ID_BIT_SIZE);
+}
+
+void FRemoteServerIdNetSerializer::Deserialize(FNetSerializationContext& Context, const FNetDeserializeArgs& Args)
+{
+	QuantizedType& Value = *reinterpret_cast<QuantizedType*>(Args.Target);
+
+	FNetBitStreamReader* Reader = Context.GetBitStreamReader();
+
+	Value = static_cast<uint16>(Reader->ReadBits(REMOTE_OBJECT_SERVER_ID_BIT_SIZE));
+}
+
+void FRemoteServerIdNetSerializer::Quantize(FNetSerializationContext& Context, const FNetQuantizeArgs& Args)
+{
+	const SourceType& SourceValue = *reinterpret_cast<const SourceType*>(Args.Source);
+	QuantizedType& TargetValue = *reinterpret_cast<QuantizedType*>(Args.Target);
+
+	FRemoteServerId GlobalizedId = SourceValue.GetGlobalized();
+	TargetValue = static_cast<QuantizedType>(GlobalizedId.Id);
+
+	checkf(TargetValue != static_cast<QuantizedType>(ERemoteServerIdConstants::Local), TEXT("Local server id should never be serialized"));
+}
+
+void FRemoteServerIdNetSerializer::Dequantize(FNetSerializationContext& Context, const FNetDequantizeArgs& Args)
+{
+	SourceType& TargetValue = *reinterpret_cast<SourceType*>(Args.Target);
+	QuantizedType& QuantizedValue = *reinterpret_cast<QuantizedType*>(Args.Source);
+
+	checkf(QuantizedValue != static_cast<QuantizedType>(ERemoteServerIdConstants::Local), TEXT("Local server id should never be serialized"));
+
+	FRemoteServerId LocalizedId;
+	LocalizedId.Id = QuantizedValue;
+	TargetValue = LocalizedId.GetLocalized();
+}
+
+bool FRemoteServerIdNetSerializer::IsEqual(FNetSerializationContext& Context, const FNetIsEqualArgs& Args)
+{
+	if (Args.bStateIsQuantized)
+	{
+		const QuantizedType& QuantizedValue0 = *reinterpret_cast<const QuantizedType*>(Args.Source0);
+		const QuantizedType& QuantizedValue1 = *reinterpret_cast<const QuantizedType*>(Args.Source1);
+
+		return QuantizedValue0 == QuantizedValue1;
+	}
+	else
+	{
+		const SourceType& SourceValue0 = *reinterpret_cast<const SourceType*>(Args.Source0);
+		const SourceType& SourceValue1 = *reinterpret_cast<const SourceType*>(Args.Source1);
+
+		return SourceValue0 == SourceValue1;
+	}
+}
+
+bool FRemoteServerIdNetSerializer::Validate(FNetSerializationContext& Context, const FNetValidateArgs& Args)
+{
+	const SourceType& SourceValue = *reinterpret_cast<const SourceType*>(Args.Source);
+
+	// Server IDs should be within the valid range
+	return SourceValue.GetIdNumber() <= MAX_REMOTE_OBJECT_SERVER_ID;
+}
+
+}
\ No newline at end of file
diff --git a/Engine/Source/Runtime/Experimental/Iris/Core/Private/Iris/Serialization/RemoteServerIdNetSerializer.h b/Engine/Source/Runtime/Experimental/Iris/Core/Private/Iris/Serialization/RemoteServerIdNetSerializer.h
new file mode 100644
--- /dev/null
+++ b/Engine/Source/Runtime/Experimental/Iris/Core/Private/Iris/Serialization/RemoteServerIdNetSerializer.h
@@ -0,0 +1,17 @@
+// Copyright Epic Games, Inc. All Rights Reserved.
+
+#pragma once
+
+#include "Iris/Serialization/InternalNetSerializer.h"
+#include "RemoteServerIdNetSerializer.generated.h"
+
+USTRUCT()
+struct FRemoteServerIdNetSerializerConfig : public FNetSerializerConfig
+{
+	GENERATED_BODY()
+};
+
+namespace UE::Net
+{
+	UE_NET_DECLARE_SERIALIZER_INTERNAL(FRemoteServerIdNetSerializer, IRISCORE_API);
+}
