Iris
- Fixed issue where dirty status of new subobjects created from PreUpdate of objects was masked out.

diff --git a/Engine/Plugins/Runtime/ReplicationSystemTestPlugin/Source/Private/Tests/ReplicationSystem/ReplicationSystemServerClientTestFixture.h b/Engine/Plugins/Runtime/ReplicationSystemTestPlugin/Source/Private/Tests/ReplicationSystem/ReplicationSystemServerClientTestFixture.h
--- a/Engine/Plugins/Runtime/ReplicationSystemTestPlugin/Source/Private/Tests/ReplicationSystem/ReplicationSystemServerClientTestFixture.h
+++ b/Engine/Plugins/Runtime/ReplicationSystemTestPlugin/Source/Private/Tests/ReplicationSystem/ReplicationSystemServerClientTestFixture.h
@@ -130,6 +130,22 @@ public:
 		return nullptr;
 	}
 
+	template<typename T>
+	T* CreateObjectWithParams(const UObjectReplicationBridge::FRootObjectReplicationParams& Params)
+	{
+		T* CreatedObject = NewObject<T>();
+		if (Cast<UReplicatedTestObject>(CreatedObject))
+		{
+			CreatedObjects.Add(TStrongObjectPtr<UObject>(CreatedObject));
+
+			// Add it to the bridge for replication
+			ReplicationBridge->BeginReplication(CreatedObject, Params);
+	
+			return CreatedObject;
+		}
+		return nullptr;
+	}
+
 	template<typename T>
 	T* CreateSubObject(FNetRefHandle OwnerHandle, FNetRefHandle InsertRelativeToSubObjectHandle = FNetRefHandle::GetInvalid(), UReplicationBridge::ESubObjectInsertionOrder InsertionOrder = UReplicationBridge::ESubObjectInsertionOrder::None)
 	{
diff --git a/Engine/Plugins/Runtime/ReplicationSystemTestPlugin/Source/Private/Tests/ReplicationSystem/TestPreUpdateCallback.cpp b/Engine/Plugins/Runtime/ReplicationSystemTestPlugin/Source/Private/Tests/ReplicationSystem/TestPreUpdateCallback.cpp
--- a/Engine/Plugins/Runtime/ReplicationSystemTestPlugin/Source/Private/Tests/ReplicationSystem/TestPreUpdateCallback.cpp
+++ b/Engine/Plugins/Runtime/ReplicationSystemTestPlugin/Source/Private/Tests/ReplicationSystem/TestPreUpdateCallback.cpp
@@ -72,8 +72,68 @@ UE_NET_TEST_FIXTURE(FReplicationSystemServerClientTestFixture, CreateObjectsInsi
 	UE_NET_ASSERT_NE(ClientSubObject, nullptr);
 
 	// The dirty property of the subobject should have replicated too.
-	UE_NET_ASSERT_EQ(ClientSubObject->IntA, ClientSubObject->IntA);
+	UE_NET_ASSERT_EQ(ClientSubObject->IntA, ServerSubObject->IntA);
 }
 
+/** Test that validates behavior when creating replicated sub objects inside PreUpdate/PreReplication when using fully pushbased objects */
+UE_NET_TEST_FIXTURE(FReplicationSystemServerClientTestFixture, CreateObjectsInsidePreUpdateTestIsPolledIfFullyPushbased)
+{
+	// Add client
+	FReplicationSystemTestClient* Client = CreateClient();
+
+	// Spawn object on server that is polled late in order to test ForceNetUpdate
+	UObjectReplicationBridge::FRootObjectReplicationParams Params;
+	Params.bUseClassConfigDynamicFilter = true;
+	Params.bNeedsPreUpdate = true;
+
+	UTestReplicatedIrisPushModelObject* ServerObject = Server->CreateObjectWithParams<UTestReplicatedIrisPushModelObject>(Params);
+	UTestReplicatedIrisPushModelObject* ServerSubObject(nullptr);
+	
+	// Send and deliver packet
+	Server->UpdateAndSend({ Client });
+	
+	// Object should have been created on the client
+	UTestReplicatedIrisPushModelObject* ClientObject = Cast<UTestReplicatedIrisPushModelObject>(Client->GetReplicationBridge()->GetReplicatedObject(ServerObject->NetRefHandle));
+	UE_NET_ASSERT_NE(ClientObject, nullptr);
+
+	// Send and deliver packet
+	Server->UpdateAndSend({ Client });
+
+	// Dirty a property
+	ServerObject->SetIntA(0xBB);
+
+	// Add a PreUpdate where we create a new subobject
+	auto PreUpdateObject = [&](TArrayView<UObject*> Instances, const UReplicationBridge* Bridge)
+	{
+		for (UObject* ReplicatedObject : Instances)
+		{
+			if (ServerObject == ReplicatedObject)
+			{
+				if (ServerSubObject == nullptr)
+				{
+					// Create a subobject
+					ServerSubObject = Server->CreateSubObject<UTestReplicatedIrisPushModelObject>(ServerObject->NetRefHandle);
+
+					// Dirty this subobject
+					ServerSubObject->SetIntA(0xBB);
+				}
+			}
+		}
+	};
+	Server->GetReplicationBridge()->SetExternalPreUpdateFunctor(PreUpdateObject);
+	Server->UpdateAndSend({ Client });
+	
+	// Was the property changed in PreUpdate replicated ?
+	UE_NET_ASSERT_EQ(ClientObject->GetIntA(), ServerObject->GetIntA());
+
+	// And in the same Send was the subobject replicated ?
+	UTestReplicatedIrisPushModelObject* ClientSubObject = Cast<UTestReplicatedIrisPushModelObject>(Client->GetReplicationBridge()->GetReplicatedObject(ServerSubObject->NetRefHandle));
+	UE_NET_ASSERT_NE(ClientSubObject, nullptr);
+
+	// The dirty property of the subobject should have replicated too.
+	UE_NET_ASSERT_EQ(ClientSubObject->GetIntA(), ServerSubObject->GetIntA());
+}
+
+
 
 }
\ No newline at end of file
diff --git a/Engine/Source/Runtime/Experimental/Iris/Core/Private/Iris/ReplicationSystem/DirtyNetObjectTracker.cpp b/Engine/Source/Runtime/Experimental/Iris/Core/Private/Iris/ReplicationSystem/DirtyNetObjectTracker.cpp
--- a/Engine/Source/Runtime/Experimental/Iris/Core/Private/Iris/ReplicationSystem/DirtyNetObjectTracker.cpp
+++ b/Engine/Source/Runtime/Experimental/Iris/Core/Private/Iris/ReplicationSystem/DirtyNetObjectTracker.cpp
@@ -227,11 +227,11 @@ void FDirtyNetObjectTracker::UpdateDirtyNetObjects()
 
 	ApplyAndTryResetGlobalDirtyObjectList();
 
-	//$IRIS TODO:  We could look if any objects where actually in the global list and skip the array iteration if not needed.
-
 	const uint32 NumWords = AccumulatedDirtyNetObjects.GetNumWords();
 
-	const FNetBitArrayView GlobalScopeList = NetRefHandleManager->GetCurrentFrameScopableInternalIndices();
+	// Note: We use the actual GlobalScopableInternalIndices BitArray as we allow new (sub)objects to be added
+	// during the NetUpdate and we do not want them to be removed from the set of dirty objects.
+	const FNetBitArrayView GlobalScopeList = NetRefHandleManager->GetGlobalScopableInternalIndices();
 	const uint32* GlobalScopeListData = GlobalScopeList.GetDataChecked(NumWords);
 	
 	uint32* AccumulatedDirtyNetObjectsData = AccumulatedDirtyNetObjects.GetDataChecked(NumWords);
