Iris
* Renamed ReplicationSystem::IsValidHandle to IsNetRefHandleAssigned to more clearly explain it's intent

#rnx

diff --git a/Engine/Plugins/Runtime/ReplicationSystemTestPlugin/Source/Private/Tests/ReplicationSystem/ReplicationSystemServerClientTestFixture.cpp b/Engine/Plugins/Runtime/ReplicationSystemTestPlugin/Source/Private/Tests/ReplicationSystem/ReplicationSystemServerClientTestFixture.cpp
--- a/Engine/Plugins/Runtime/ReplicationSystemTestPlugin/Source/Private/Tests/ReplicationSystem/ReplicationSystemServerClientTestFixture.cpp
+++ b/Engine/Plugins/Runtime/ReplicationSystemTestPlugin/Source/Private/Tests/ReplicationSystem/ReplicationSystemServerClientTestFixture.cpp
@@ -289,7 +289,7 @@ bool FReplicationSystemTestNode::IsResolvableNetRefHandle(FNetRefHandle RefHandl
 
 bool FReplicationSystemTestNode::IsValidNetRefHandle(FNetRefHandle RefHandle)
 {
-	return ReplicationSystem->GetReplicationSystemInternal()->GetNetRefHandleManager().IsValidNetRefHandle(RefHandle);
+	return ReplicationSystem->GetReplicationSystemInternal()->GetNetRefHandleManager().IsNetRefHandleAssigned(RefHandle);
 }
 
 uint32 FReplicationSystemTestNode::AddConnection()
diff --git a/Engine/Source/Runtime/Experimental/Iris/Core/Private/Iris/ReplicationSystem/ReplicationBridge.cpp b/Engine/Source/Runtime/Experimental/Iris/Core/Private/Iris/ReplicationSystem/ReplicationBridge.cpp
--- a/Engine/Source/Runtime/Experimental/Iris/Core/Private/Iris/ReplicationSystem/ReplicationBridge.cpp
+++ b/Engine/Source/Runtime/Experimental/Iris/Core/Private/Iris/ReplicationSystem/ReplicationBridge.cpp
@@ -1101,7 +1101,7 @@ UE::Net::FNetRefHandle UReplicationBridge::StoreDestructionInfo(FNetRefHandle Ha
 
 bool UReplicationBridge::IsReplicatedHandle(FNetRefHandle Handle) const
 {
-	return Handle.IsValid() && ReplicationSystem->IsValidHandle(Handle);
+	return Handle.IsValid() && NetRefHandleManager->IsNetRefHandleAssigned(Handle);
 }
 
 void UReplicationBridge::SetNetPushIdOnFragments(const TArrayView<const UE::Net::FReplicationFragment*const>& Fragments, const UE::Net::Private::FNetPushObjectHandle& PushHandle)
diff --git a/Engine/Source/Runtime/Experimental/Iris/Core/Private/Iris/ReplicationSystem/ReplicationSystem.cpp b/Engine/Source/Runtime/Experimental/Iris/Core/Private/Iris/ReplicationSystem/ReplicationSystem.cpp
--- a/Engine/Source/Runtime/Experimental/Iris/Core/Private/Iris/ReplicationSystem/ReplicationSystem.cpp
+++ b/Engine/Source/Runtime/Experimental/Iris/Core/Private/Iris/ReplicationSystem/ReplicationSystem.cpp
@@ -1405,7 +1405,7 @@ UReplicationBridge* UReplicationSystem::GetReplicationBridge() const
 	return Impl->ReplicationSystemInternal.GetReplicationBridge();
 }
 
-bool UReplicationSystem::IsValidHandle(FNetRefHandle Handle) const
+bool UReplicationSystem::IsNetRefHandleAssigned(FNetRefHandle Handle) const
 {
 	return Handle.IsValid() && Impl->ReplicationSystemInternal.GetNetRefHandleManager().IsNetRefHandleAssigned(Handle);
 }
diff --git a/Engine/Source/Runtime/Experimental/Iris/Core/Public/Iris/ReplicationSystem/ReplicationSystem.h b/Engine/Source/Runtime/Experimental/Iris/Core/Public/Iris/ReplicationSystem/ReplicationSystem.h
--- a/Engine/Source/Runtime/Experimental/Iris/Core/Public/Iris/ReplicationSystem/ReplicationSystem.h
+++ b/Engine/Source/Runtime/Experimental/Iris/Core/Public/Iris/ReplicationSystem/ReplicationSystem.h
@@ -403,10 +403,11 @@ public:
 
 	/**
 	 * Check whether a FNetRefHandle is still associated with a replicated object.
+	 * An handle can be valid but reference an object that doesn't exist locally.
 	 * @param Handle The handle check.
-	 * @return true if the handle is still valid, false if not. 
+	 * @return true if the handle is assigned, false if not. 
 	 */
-	IRISCORE_API bool IsValidHandle(FNetRefHandle Handle) const;
+	IRISCORE_API bool IsNetRefHandleAssigned(FNetRefHandle Handle) const;
 
 	/**
 	 * Get the ReplicationProtocol for a handle.
