Iris - Fix edge case with prioritization causing a crash.

diff --git a/Engine/Plugins/Runtime/ReplicationSystemTestPlugin/Source/Private/Tests/ReplicationSystem/Prioritization/TestNetObjectPrioritizer.cpp b/Engine/Plugins/Runtime/ReplicationSystemTestPlugin/Source/Private/Tests/ReplicationSystem/Prioritization/TestNetObjectPrioritizer.cpp
--- a/Engine/Plugins/Runtime/ReplicationSystemTestPlugin/Source/Private/Tests/ReplicationSystem/Prioritization/TestNetObjectPrioritizer.cpp
+++ b/Engine/Plugins/Runtime/ReplicationSystemTestPlugin/Source/Private/Tests/ReplicationSystem/Prioritization/TestNetObjectPrioritizer.cpp
@@ -5,38 +5,18 @@
 #include "Iris/ReplicationSystem/NetRefHandleManager.h"
 #include "Iris/ReplicationSystem/ReplicationSystemInternal.h"
 #include "Iris/ReplicationSystem/Prioritization/NetObjectPrioritizerDefinitions.h"
+#include "Templates/AlignmentTemplates.h"
 #include "Tests/ReplicationSystem/Prioritization/MockNetObjectPrioritizer.h"
 #include "Tests/ReplicationSystem/Prioritization/TestPrioritizationObject.h"
+#include "Tests/ReplicationSystem/ReplicationSystemConfigOverrideTestFixture.h"
 #include "Tests/ReplicationSystem/ReplicationSystemTestFixture.h"
-#include "Iris/ReplicationSystem/RepTag.h"
 
 namespace UE::Net::Private
 {
 
-class FTestNetObjectPrioritizerAPIFixture : public FReplicationSystemTestFixture
+class FTestNetObjectPrioritizerAPIFixtureBase
 {
 public:
-	FTestNetObjectPrioritizerAPIFixture() : MockNetObjectPrioritizer(nullptr) {}
-
-protected:
-	virtual void SetUp() override
-	{
-		InitNetObjectPrioritizerDefinitions();
-		FReplicationSystemTestFixture::SetUp();
-		InitMockNetObjectPrioritizer();
-
-		NetRefHandleManager = &ReplicationSystem->GetReplicationSystemInternal()->GetNetRefHandleManager();
-	}
-
-	virtual void TearDown() override
-	{
-		FReplicationSystemTestFixture::TearDown();
-		RestoreNetObjectPrioritizerDefinitions();
-
-		NetRefHandleManager = nullptr;
-	}
-
-private:
 	void InitNetObjectPrioritizerDefinitions()
 	{
 		const UClass* NetObjectPrioritizerDefinitionsClass = UNetObjectPrioritizerDefinitions::StaticClass();
@@ -47,12 +27,12 @@ private:
 		UNetObjectPrioritizerDefinitions* PrioritizerDefinitions = GetMutableDefault<UNetObjectPrioritizerDefinitions>();
 		DefinitionsProperty->CopyCompleteValue(&OriginalPrioritizerDefinitions, (void*)(UPTRINT(PrioritizerDefinitions) + DefinitionsProperty->GetOffset_ForInternal()));
 
-		// Modify definitions to only include our mock prioritizer. Ugly... 
+		// Modify definitions to only include our mock prioritizers. Ugly... 
 		TArray<FNetObjectPrioritizerDefinition> NewPrioritizerDefinitions;
 		FNetObjectPrioritizerDefinition& MockDefinition = NewPrioritizerDefinitions.Emplace_GetRef();
-		MockDefinition.PrioritizerName = TEXT("MockPrioritizer");
-		MockDefinition.ClassName = TEXT("/Script/ReplicationSystemTestPlugin.MockNetObjectPrioritizer");
-		MockDefinition.ConfigClassName = TEXT("/Script/ReplicationSystemTestPlugin.MockNetObjectPrioritizerConfig");
+		MockDefinition.PrioritizerName = FName("MockPrioritizer");
+		MockDefinition.ClassName = FName("/Script/ReplicationSystemTestPlugin.MockNetObjectPrioritizer");
+		MockDefinition.ConfigClassName = FName("/Script/ReplicationSystemTestPlugin.MockNetObjectPrioritizerConfig");
 		DefinitionsProperty->CopyCompleteValue((void*)(UPTRINT(PrioritizerDefinitions) + DefinitionsProperty->GetOffset_ForInternal()), &NewPrioritizerDefinitions);
 	}
 
@@ -65,26 +45,78 @@ private:
 		DefinitionsProperty->CopyCompleteValue((void*)(UPTRINT(PrioritizerDefinitions) + DefinitionsProperty->GetOffset_ForInternal()), &OriginalPrioritizerDefinitions);
 		OriginalPrioritizerDefinitions.Empty();
 
-		MockPrioritizerHandle = InvalidNetObjectPrioritizerHandle;
 		MockNetObjectPrioritizer = nullptr;
+		MockPrioritizerHandle = InvalidNetObjectPrioritizerHandle;
 	}
 
-	void InitMockNetObjectPrioritizer()
+	void InitMockNetObjectPrioritizers(const UReplicationSystem* ReplicationSystem)
 	{
-		MockNetObjectPrioritizer = ExactCast<UMockNetObjectPrioritizer>(ReplicationSystem->GetPrioritizer(TEXT("MockPrioritizer")));
+		MockNetObjectPrioritizer = ExactCast<UMockNetObjectPrioritizer>(ReplicationSystem->GetPrioritizer(FName("MockPrioritizer")));
 		MockPrioritizerHandle = ReplicationSystem->GetPrioritizerHandle(FName("MockPrioritizer"));
 	}
 
 protected:
 	UMockNetObjectPrioritizer* MockNetObjectPrioritizer = nullptr;
 	FNetObjectPrioritizerHandle MockPrioritizerHandle = InvalidNetObjectPrioritizerHandle;
-	FNetRefHandleManager* NetRefHandleManager = nullptr;
-	static constexpr float FakeDeltaTime = 0.0334f;
 
 private:
 	TArray<FNetObjectPrioritizerDefinition> OriginalPrioritizerDefinitions;
 };
 
+class FTestNetObjectPrioritizerAPIFixture : public FReplicationSystemTestFixture, public FTestNetObjectPrioritizerAPIFixtureBase
+{
+protected:
+	virtual void SetUp() override
+	{
+		InitNetObjectPrioritizerDefinitions();
+		FReplicationSystemTestFixture::SetUp();
+		InitMockNetObjectPrioritizers(ReplicationSystem);
+
+		NetRefHandleManager = &ReplicationSystem->GetReplicationSystemInternal()->GetNetRefHandleManager();
+	}
+
+	virtual void TearDown() override
+	{
+		FReplicationSystemTestFixture::TearDown();
+		RestoreNetObjectPrioritizerDefinitions();
+
+		NetRefHandleManager = nullptr;
+	}
+
+protected:
+	FNetRefHandleManager* NetRefHandleManager = nullptr;
+	static constexpr float FakeDeltaTime = 0.0334f;
+};
+
+class FTestNetObjectPrioritizationEdgeCaseFixture : public FReplicationSystemConfigOverrideRPCTestFixture, public FTestNetObjectPrioritizerAPIFixtureBase
+{
+protected:
+	virtual void SetUp() override
+	{
+		InitNetObjectPrioritizerDefinitions();
+		FReplicationSystemConfigOverrideRPCTestFixture::SetUp();
+
+		CreateServer({.MaxReplicatedObjectCount = MaxObjectCount, .InitialNetObjectListCount = MaxObjectCount});
+		ReplicationSystem = Server->GetReplicationSystem();
+		InitMockNetObjectPrioritizers(ReplicationSystem);
+	}
+
+	virtual void TearDown() override
+	{
+		FReplicationSystemConfigOverrideRPCTestFixture::TearDown();
+		RestoreNetObjectPrioritizerDefinitions();
+
+		ReplicationSystem  = nullptr;
+	}
+
+protected:
+	// This constant needs to match FReplicationPrioritization::FPrioritizerBatchHelper::MaxObjectCountPerBatch for tests to be effective
+	static constexpr uint32 BatchObjectCount = 1024;
+	static constexpr uint32 MaxObjectCount = Align(BatchObjectCount + 1, 64);
+
+	UReplicationSystem* ReplicationSystem = nullptr;
+};
+
 UE_NET_TEST_FIXTURE(FTestNetObjectPrioritizerAPIFixture, CanGetPrioritizerHandle)
 {
 	UE_NET_ASSERT_NE(MockPrioritizerHandle, InvalidNetObjectPrioritizerHandle);
@@ -506,4 +538,63 @@ UE_NET_TEST_FIXTURE(FTestNetObjectPrioritizerAPIFixture, ObjectGetsUpdatedPriori
 	UE_NET_ASSERT_EQ(Priority, UpdatedPriority);
 }
 
+// Edge case tests
+UE_NET_TEST_FIXTURE(FTestNetObjectPrioritizationEdgeCaseFixture, PrioritizationSystemHandlesObjectCountEdgeCaseGracefully)
+{
+	// First let's create the edge case scenario where we have BatchObjectCount objects and the last object is using the highest index available.
+	{
+		MockNetObjectPrioritizer->ResetFunctionCallStatus();
+
+		// We want the add calls to succeed
+		{
+			UMockNetObjectPrioritizer::FFunctionCallSetup CallSetup;
+			CallSetup.AddObject.ReturnValue = true;
+			MockNetObjectPrioritizer->SetFunctionCallSetup(CallSetup);
+		}
+
+		// 1. First create objects that we'll later destroy so we get exactly 
+		TArray<UTestReplicatedIrisObject*> ObjectsToDestroy;
+		static_assert(MaxObjectCount > BatchObjectCount, "If MaxObjectCount <= BatchObjectCount we're unable to create the right amount of objects");
+		{
+			const uint32 ObjectCountToDestroy = MaxObjectCount - BatchObjectCount - 1;
+			ObjectsToDestroy.Reserve(ObjectCountToDestroy);
+			for (unsigned It = 0, EndIt = ObjectCountToDestroy; It < EndIt; ++It)
+			{
+				UTestReplicatedIrisObject* ServerObject = Server->CreateObject();
+				ObjectsToDestroy.Add(ServerObject);
+			}
+		}
+
+		// 2. Then create BatchObjectCount objects using the same prioritizer.
+		for (unsigned It = 0, EndIt = BatchObjectCount; It < EndIt; ++It)
+		{
+			UTestReplicatedIrisObject* ServerObject = Server->CreateObject();
+			ReplicationSystem->SetPrioritizer(ServerObject->NetRefHandle, MockPrioritizerHandle);
+		}
+
+		// 3. Destroy the first objects so we have exactly BatchObjectCount using the highest available internal indices
+		for (UTestReplicatedIrisObject* Object : ObjectsToDestroy)
+		{
+			Server->DestroyObject(Object);
+		}
+	}
+	
+	// Setup a connection with a valid view
+	{
+		ReplicationSystem->AddConnection(13);
+
+		FReplicationView View;
+		View.Views.AddDefaulted();
+		ReplicationSystem->SetReplicationView(13, View);
+	}
+
+	// There should be no ensures or checks when updating.
+	FTestEnsureScope EnsureScope;
+
+	ReplicationSystem->NetUpdate(0.1f);
+	ReplicationSystem->PostSendUpdate();
+
+	UE_NET_ASSERT_EQ(EnsureScope.GetCount(), 0);
+}
+
 }
diff --git a/Engine/Source/Runtime/Experimental/Iris/Core/Private/Iris/ReplicationSystem/Prioritization/ReplicationPrioritization.cpp b/Engine/Source/Runtime/Experimental/Iris/Core/Private/Iris/ReplicationSystem/Prioritization/ReplicationPrioritization.cpp
--- a/Engine/Source/Runtime/Experimental/Iris/Core/Private/Iris/ReplicationSystem/Prioritization/ReplicationPrioritization.cpp
+++ b/Engine/Source/Runtime/Experimental/Iris/Core/Private/Iris/ReplicationSystem/Prioritization/ReplicationPrioritization.cpp
@@ -173,7 +173,7 @@ public:
 				 * If any prioritizer has a full batch then return ProcessFull. This limits the memory footprint to a maximum of two chunks per prioritizer.
 				 * Else continue.
 				 */
-				if (ObjectCount < MaxObjectCountPerBatch)
+				if ((ObjectCount < MaxObjectCountPerBatch) || (BatchInfo.CurrentObjectIndex == Objects.GetNumBits()))
 				{
 					return EBatchProcessStatus::ProcessAllBatchesAndStop;
 				}
@@ -793,8 +793,7 @@ void FReplicationPrioritization::PrioritizeForConnection(uint32 ConnId, FPriorit
 			continue;
 		}
 
-		// Unhandled status
-		checkf(false, TEXT("Unexpected BatchProcessStatus %u"), ProcessStatus);
+		// We're done
 		break;
 	}
 
