Iris - Test ObjectNetSerializer IsEqual on invalid memory.

diff --git a/Engine/Plugins/Runtime/ReplicationSystemTestPlugin/Source/Private/Tests/Serialization/TestObjectNetSerializer.cpp b/Engine/Plugins/Runtime/ReplicationSystemTestPlugin/Source/Private/Tests/Serialization/TestObjectNetSerializer.cpp
--- a/Engine/Plugins/Runtime/ReplicationSystemTestPlugin/Source/Private/Tests/Serialization/TestObjectNetSerializer.cpp
+++ b/Engine/Plugins/Runtime/ReplicationSystemTestPlugin/Source/Private/Tests/Serialization/TestObjectNetSerializer.cpp
@@ -4,6 +4,7 @@
 #include "Iris/Serialization/ObjectNetSerializer.h"
 #include "Iris/Core/NetObjectReference.h"
 #include "Iris/Serialization/InternalNetSerializationContext.h"
+#include "Templates/UnrealTemplate.h"
 #include "Tests/ReplicationSystem/ReplicationSystemServerClientTestFixture.h"
 #include "Tests/ReplicationSystem/ReplicatedTestObject.h"
 
@@ -187,4 +188,36 @@ UE_NET_TEST_FIXTURE(FReplicationSystemServerClientTestFixture, TestSoftObjectRef
 	UE_NET_ASSERT_EQ((UObject*)ClientObjectA->SoftObjectPtrRef.Get(), (UObject*)ClientObjectB);
 }
 
+UE_NET_TEST_FIXTURE(FReplicationSystemServerClientTestFixture, TestGarbageObjectPtrReference)
+{
+	// Add a client
+	FReplicationSystemTestClient* Client = CreateClient();
+
+	// Spawn objects on server
+	UTestReplicatedIrisObjectWithObjectReference* ObjectA = Server->CreateObject<UTestReplicatedIrisObjectWithObjectReference>();
+	UTestReplicatedIrisObject* ObjectB = Server->CreateObject<UTestReplicatedIrisObject>();
+
+	void* FakeObject = FMemory::Malloc(sizeof(UTestReplicatedIrisObject), alignof(UTestReplicatedIrisObject));
+	FMemory::Memcpy(FakeObject, ObjectB, sizeof(UTestReplicatedIrisObject));
+	FObjectPtr TestInvalidObjectPtr_Untyped((UObject*)FakeObject);
+	TObjectPtr<UObject> TestInvalidObjectPtr(TestInvalidObjectPtr_Untyped);
+
+	ObjectA->RawObjectPtrRef = TestInvalidObjectPtr;
+
+	// Send and deliver packet
+	Server->UpdateAndSend({Client});
+
+	ObjectA->RawObjectPtrRef = ObjectB;
+
+	// Invalidate fake object
+	FMemory::Memzero(FakeObject, sizeof(UTestReplicatedIrisObject));
+	FMemory::Free(FakeObject);
+
+	// Force TObjectPtr creation to mark objects as reachable
+	TGuardValue<bool> IsIncrementalReachabilityPending(UE::GC::GIsIncrementalReachabilityPending, true);
+
+	// Send and deliver packet. If a TObjectPtr is created this may crash if for example running with -stompmalloc or assert if running with checks enabled.
+	Server->UpdateAndSend({Client});
+}
+
 }
