Fixed test associated with submitted cl 45023168

diff --git a/Engine/Plugins/Runtime/ReplicationSystemTestPlugin/Source/Private/Tests/Serialization/TestScriptInterfaceNetSerializer.cpp b/Engine/Plugins/Runtime/ReplicationSystemTestPlugin/Source/Private/Tests/Serialization/TestScriptInterfaceNetSerializer.cpp
--- a/Engine/Plugins/Runtime/ReplicationSystemTestPlugin/Source/Private/Tests/Serialization/TestScriptInterfaceNetSerializer.cpp
+++ b/Engine/Plugins/Runtime/ReplicationSystemTestPlugin/Source/Private/Tests/Serialization/TestScriptInterfaceNetSerializer.cpp
@@ -12,7 +12,7 @@ namespace UE::Net::Private
 static FTestMessage& PrintScriptInterfaceNetSerializerConfig(FTestMessage& Message, const FNetSerializerConfig& InConfig)
 {
 	const FScriptInterfaceNetSerializerConfig& Config = static_cast<const FScriptInterfaceNetSerializerConfig&>(InConfig);
-	return Message << TEXT("Interface: ") << Config.InterfaceClass->GetName();
+	return Message << TEXT("Interface: ") << Config.PropertyClass->GetName();
 }
 
 class FTestScriptInterfaceNetSerializer : public FReplicationSystemServerClientTestFixture
