Fixed ReplicationSystemTest compile error

diff --git a/Engine/Plugins/Runtime/ReplicationSystemTestPlugin/Source/Private/Tests/Serialization/TestRemoteObjectIdNetSerializer.cpp b/Engine/Plugins/Runtime/ReplicationSystemTestPlugin/Source/Private/Tests/Serialization/TestRemoteObjectIdNetSerializer.cpp
--- a/Engine/Plugins/Runtime/ReplicationSystemTestPlugin/Source/Private/Tests/Serialization/TestRemoteObjectIdNetSerializer.cpp
+++ b/Engine/Plugins/Runtime/ReplicationSystemTestPlugin/Source/Private/Tests/Serialization/TestRemoteObjectIdNetSerializer.cpp
@@ -39,7 +39,7 @@ FRemoteObjectIdNetSerializerConfig FTestRemoteObjectIdNetSerializer::SerializerC
 
 UE_NET_TEST_FIXTURE(FTestRemoteObjectIdNetSerializer, HasTestValues)
 {
-	UE_NET_ASSERT_GT_MSG(TestValues.Num(), SIZE_T(0), "No test values found");
+	UE_NET_ASSERT_GT_MSG(TestValues.Num(), 0, "No test values found");
 }
 
 UE_NET_TEST_FIXTURE(FTestRemoteObjectIdNetSerializer, TestIsEqual)
diff --git a/Engine/Plugins/Runtime/ReplicationSystemTestPlugin/Source/Private/Tests/Serialization/TestRemoteServerIdNetSerializer.cpp b/Engine/Plugins/Runtime/ReplicationSystemTestPlugin/Source/Private/Tests/Serialization/TestRemoteServerIdNetSerializer.cpp
--- a/Engine/Plugins/Runtime/ReplicationSystemTestPlugin/Source/Private/Tests/Serialization/TestRemoteServerIdNetSerializer.cpp
+++ b/Engine/Plugins/Runtime/ReplicationSystemTestPlugin/Source/Private/Tests/Serialization/TestRemoteServerIdNetSerializer.cpp
@@ -39,7 +39,7 @@ FRemoteServerIdNetSerializerConfig FTestRemoteServerIdNetSerializer::SerializerC
 
 UE_NET_TEST_FIXTURE(FTestRemoteServerIdNetSerializer, HasTestValues)
 {
-	UE_NET_ASSERT_GT_MSG(TestValues.Num(), SIZE_T(0), "No test values found");
+	UE_NET_ASSERT_GT_MSG(TestValues.Num(), 0, "No test values found");
 }
 
 UE_NET_TEST_FIXTURE(FTestRemoteServerIdNetSerializer, TestIsEqual)
