Iris - Work around bad code generation.

diff --git a/Engine/Plugins/Runtime/ReplicationSystemTestPlugin/Source/Private/Tests/Serialization/TestFloatNetSerializer.cpp b/Engine/Plugins/Runtime/ReplicationSystemTestPlugin/Source/Private/Tests/Serialization/TestFloatNetSerializer.cpp
--- a/Engine/Plugins/Runtime/ReplicationSystemTestPlugin/Source/Private/Tests/Serialization/TestFloatNetSerializer.cpp
+++ b/Engine/Plugins/Runtime/ReplicationSystemTestPlugin/Source/Private/Tests/Serialization/TestFloatNetSerializer.cpp
@@ -97,11 +97,14 @@ void FTestFloatNetSerializerBase<T>::TestIsEqual()
 	T CompareValues[2][sizeof(Values)/sizeof(Values[0])];
 	bool ExpectedResults[2][sizeof(Values)/sizeof(Values[0])];
 
+	// As a temporary workaround for some compiler issues we copy the memory for the entire array instead of one value at a time. Revert back to per value copy when possible to detect future compiler regressions.
+	FPlatformMemory::Memcpy(CompareValues[0], Values, ValueCount*sizeof(T));
 	for (SIZE_T ValueIt = 0; ValueIt < ValueCount; ++ValueIt)
 	{
-		CompareValues[0][ValueIt] = Values[ValueIt];
+		//CompareValues[0][ValueIt] = Values[ValueIt];
 		ExpectedResults[0][ValueIt] = true;
 	}
+
 	for (SIZE_T ValueIt = 0; ValueIt < ValueCount; ++ValueIt)
 	{
 		CompareValues[1][ValueIt] = Values[(ValueIt + 1) % ValueCount];
