Iris - Suppress warning logs in tests

diff --git a/Engine/Plugins/Runtime/ReplicationSystemTestPlugin/Source/Private/Tests/Serialization/TestObjectNetSerializer.cpp b/Engine/Plugins/Runtime/ReplicationSystemTestPlugin/Source/Private/Tests/Serialization/TestObjectNetSerializer.cpp
--- a/Engine/Plugins/Runtime/ReplicationSystemTestPlugin/Source/Private/Tests/Serialization/TestObjectNetSerializer.cpp
+++ b/Engine/Plugins/Runtime/ReplicationSystemTestPlugin/Source/Private/Tests/Serialization/TestObjectNetSerializer.cpp
@@ -2,8 +2,10 @@
 
 #include "TestNetSerializerFixture.h"
 #include "Iris/Serialization/ObjectNetSerializer.h"
+#include "Iris/Core/IrisLog.h"
 #include "Iris/Core/NetObjectReference.h"
 #include "Iris/Serialization/InternalNetSerializationContext.h"
+#include "Logging/LogScopedVerbosityOverride.h"
 #include "Templates/UnrealTemplate.h"
 #include "Tests/ReplicationSystem/ReplicationSystemServerClientTestFixture.h"
 #include "Tests/ReplicationSystem/ReplicatedTestObject.h"
@@ -244,7 +246,10 @@ UE_NET_TEST_FIXTURE(FReplicationSystemServerClientTestFixture, TestTypedObjectRe
 	*(TSoftObjectPtr<UTestReplicatedIrisObject>*)(&ServerObject->TypedSoftObjectPtrRef) = ReplicatedIrisObject;
 
 	// Send and deliver packet
-	Server->UpdateAndSend({Client});
+	{
+		LOG_SCOPE_VERBOSITY_OVERRIDE(LogIris, ELogVerbosity::Error);
+		Server->UpdateAndSend({Client});
+	}
 
 	UTestReplicatedIrisObjectWithObjectReference* ClientObject = Client->GetObjectAs<UTestReplicatedIrisObjectWithObjectReference>(ServerObject->NetRefHandle);
 	UE_NET_ASSERT_NE(ClientObject, nullptr);
