Iris - Prevent ensures from breaking the debugger when an ensure is hit within a TestEnsureScope.

diff --git a/Engine/Plugins/Runtime/ReplicationSystemTestPlugin/Source/Private/NetworkAutomationTestMacros.cpp b/Engine/Plugins/Runtime/ReplicationSystemTestPlugin/Source/Private/NetworkAutomationTestMacros.cpp
--- a/Engine/Plugins/Runtime/ReplicationSystemTestPlugin/Source/Private/NetworkAutomationTestMacros.cpp
+++ b/Engine/Plugins/Runtime/ReplicationSystemTestPlugin/Source/Private/NetworkAutomationTestMacros.cpp
@@ -7,6 +7,8 @@
 
 DEFINE_LOG_CATEGORY(LogNetworkAutomationTest);
 
+extern CORE_API bool GIgnoreDebugger;
+
 namespace UE::Net
 {
 
@@ -17,8 +19,15 @@ FTestEnsureScope::FTestEnsureScope()
 		UE_LOG(LogNetworkAutomationTest, Display, TEXT("Ensure condition failed: %hs\n%s\n"), Args.Expression, Args.Message);
 		return true;
 	}
-  ) 
+  )
+, bIgnoreDebugger(GIgnoreDebugger)
+{
+	GIgnoreDebugger = true;
+}
+
+FTestEnsureScope::~FTestEnsureScope()
 {
+	GIgnoreDebugger = bIgnoreDebugger;
 }
 
 FTestResult::FTestResult()
diff --git a/Engine/Plugins/Runtime/ReplicationSystemTestPlugin/Source/Private/NetworkAutomationTestMacros.h b/Engine/Plugins/Runtime/ReplicationSystemTestPlugin/Source/Private/NetworkAutomationTestMacros.h
--- a/Engine/Plugins/Runtime/ReplicationSystemTestPlugin/Source/Private/NetworkAutomationTestMacros.h
+++ b/Engine/Plugins/Runtime/ReplicationSystemTestPlugin/Source/Private/NetworkAutomationTestMacros.h
@@ -22,6 +22,10 @@ class FTestEnsureScope : public FEnsureScope
 {
 public:
 	FTestEnsureScope();
+	~FTestEnsureScope();
+
+private:
+	bool bIgnoreDebugger;
 };
 
 enum ETestResult
