Iris - Removed error prone SetConnectionFilter API from UReplicationSystem. Use exclusion groups instead.

diff --git a/Engine/Plugins/Runtime/ReplicationSystemTestPlugin/Source/Private/Tests/ReplicationSystem/Filtering/TestFiltering.cpp b/Engine/Plugins/Runtime/ReplicationSystemTestPlugin/Source/Private/Tests/ReplicationSystem/Filtering/TestFiltering.cpp
--- a/Engine/Plugins/Runtime/ReplicationSystemTestPlugin/Source/Private/Tests/ReplicationSystem/Filtering/TestFiltering.cpp
+++ b/Engine/Plugins/Runtime/ReplicationSystemTestPlugin/Source/Private/Tests/ReplicationSystem/Filtering/TestFiltering.cpp
@@ -464,348 +464,6 @@ UE_NET_TEST_FIXTURE(FTestFilteringFixture, LateAddedSubObjectGetsOwnerPropagated
 	UE_NET_ASSERT_EQ(Server->ReplicationSystem->GetOwningNetConnection(ServerSubObject->NetRefHandle), Client->ConnectionIdOnServer);
 }
 
-// Connection filtering tests
-UE_NET_TEST_FIXTURE(FTestFilteringFixture, ConnectionFilterPreventsObjectFromReplicatingToFilteredOutConnections)
-{
-	// Add client
-	FReplicationSystemTestClient* Client = CreateClient();
-
-	// Spawn object on server
-	UReplicatedTestObject* ServerObject = Server->CreateObject(0, 0);
-
-	// Apply connection filter
-	TBitArray<> NoConnections;
-	const bool bFilterWasApplied = Server->ReplicationSystem->SetConnectionFilter(ServerObject->NetRefHandle, NoConnections, ENetFilterStatus::Allow);
-	UE_NET_ASSERT_TRUE(bFilterWasApplied);
-
-	// Send and deliver packet
-	Server->NetUpdate();
-	Server->SendAndDeliverTo(Client, DeliverPacket);
-	Server->PostSendUpdate();
-
-	// Object should not have been created on the client
-	UE_NET_ASSERT_EQ(Client->GetReplicationBridge()->GetReplicatedObject(ServerObject->NetRefHandle), nullptr);
-}
-
-UE_NET_TEST_FIXTURE(FTestFilteringFixture, ConnectionFilterAllowsObjectToReplicateToAllowedConnections)
-{
-	// Add client
-	FReplicationSystemTestClient* Client = CreateClient();
-
-	// Spawn object on server
-	UReplicatedTestObject* ServerObject = Server->CreateObject(0, 0);
-
-	// Apply connection filter
-	TBitArray<> AllowedConnections;
-	AllowedConnections.Init(false, Client->ConnectionIdOnServer + 1);
-	AllowedConnections[Client->ConnectionIdOnServer] = true;
-	Server->ReplicationSystem->SetConnectionFilter(ServerObject->NetRefHandle, AllowedConnections, ENetFilterStatus::Allow);
-
-	// Send and deliver packet
-	Server->NetUpdate();
-	Server->SendAndDeliverTo(Client, DeliverPacket);
-	Server->PostSendUpdate();
-
-	// Object should have been created on the client
-	UE_NET_ASSERT_NE(Client->GetReplicationBridge()->GetReplicatedObject(ServerObject->NetRefHandle), nullptr);
-}
-
-UE_NET_TEST_FIXTURE(FTestFilteringFixture, ConnectionFilterAllowsObjectToReplicateToLateJoiningConnections)
-{
-	// Spawn object on server
-	UReplicatedTestObject* ServerObject = Server->CreateObject(0, 0);
-
-	// Apply filtering that allows all connections
-	TBitArray<> NoConnections;
-	constexpr ENetFilterStatus ReplicationStatus = ENetFilterStatus::Disallow;
-	Server->ReplicationSystem->SetConnectionFilter(ServerObject->NetRefHandle, NoConnections, ReplicationStatus);
-
-	// Send and deliver packets
-	Server->NetUpdate();
-	Server->PostSendUpdate();
-
-	// Late join clients
-	FReplicationSystemTestClient* ClientArray[] = {CreateClient(), CreateClient(), CreateClient()};
-
-	// Send and deliver packets
-	Server->NetUpdate();
-	for (FReplicationSystemTestClient* Client : ClientArray)
-	{
-		Server->SendAndDeliverTo(Client, DeliverPacket);
-	}
-	Server->PostSendUpdate();
-
-	// Object should have been created on all clients
-	for (FReplicationSystemTestClient*& Client : ClientArray)
-	{
-		UE_NET_ASSERT_NE(Client->GetReplicationBridge()->GetReplicatedObject(ServerObject->NetRefHandle), nullptr);
-	}
-}
-
-UE_NET_TEST_FIXTURE(FTestFilteringFixture, ConnectionFilterAllowsObjectToReplicateAndDoesNotAffectExistingGroupFilteredObjects)
-{
-	// Add client
-	FReplicationSystemTestClient* Client = CreateClient();
-
-	// Spawn object on server
-	UReplicatedTestObject* ServerObject0 = Server->CreateObject(0, 0);
-	
-	// Setup group filter
-	FNetObjectGroupHandle GroupHandle = Server->ReplicationSystem->CreateGroup(NAME_None);
-	Server->ReplicationSystem->AddToGroup(GroupHandle, ServerObject0->NetRefHandle);
-	Server->ReplicationSystem->AddExclusionFilterGroup(GroupHandle);
-
-	// Spawn object on server
-	UReplicatedTestObject* ServerObject = Server->CreateObject(0, 0);
-	UReplicatedTestObject* ServerSubObject = Server->CreateSubObject(ServerObject->NetRefHandle, 0, 0);
-
-	// Apply filtering that allows all connections
-	TBitArray<> NoConnections;
-	constexpr ENetFilterStatus ReplicationStatus = ENetFilterStatus::Disallow;
-	Server->ReplicationSystem->SetConnectionFilter(ServerObject->NetRefHandle, NoConnections, ReplicationStatus);
-
-	// Send and deliver packets
-	Server->NetUpdate();
-	Server->PostSendUpdate();
-
-	// Send and deliver packets
-	Server->NetUpdate();
-	Server->SendAndDeliverTo(Client, DeliverPacket);
-	Server->PostSendUpdate();
-
-	// Validate status on clients
-	// Object should not have been created on the clients 
-	UE_NET_ASSERT_EQ(Client->GetReplicationBridge()->GetReplicatedObject(ServerObject0->NetRefHandle), nullptr);
-		
-	// Object should have been created on all clients
-	UE_NET_ASSERT_NE(Client->GetReplicationBridge()->GetReplicatedObject(ServerObject->NetRefHandle), nullptr);
-	UE_NET_ASSERT_NE(Client->GetReplicationBridge()->GetReplicatedObject(ServerSubObject->NetRefHandle), nullptr);
-}
-
-UE_NET_TEST_FIXTURE(FTestFilteringFixture, ConnectionFilterPreventsObjectFromReplicatingToFilteredOutLateJoiningConnections)
-{
-	// Spawn object on server
-	UReplicatedTestObject* ServerObject = Server->CreateObject(0, 0);
-
-	// Apply filtering that allows all connections
-	TBitArray<> Connection1(false, 2);
-	Connection1[1] = true;
-	constexpr ENetFilterStatus ReplicationStatus = ENetFilterStatus::Allow;
-	Server->ReplicationSystem->SetConnectionFilter(ServerObject->NetRefHandle, Connection1, ReplicationStatus);
-
-	// Send and deliver packets
-	Server->NetUpdate();
-	Server->PostSendUpdate();
-
-	// Late join clients
-	FReplicationSystemTestClient* ClientArray[] = {CreateClient(), CreateClient(), CreateClient()};
-
-	// Since we must set up the filtering before the clients are created we need to make sure our assumptions are valid.
-	UE_NET_ASSERT_EQ(ClientArray[0]->ConnectionIdOnServer, 1U);
-
-	// Send and deliver packets
-	Server->UpdateAndSend(ClientArray, DeliverPacket);
-
-	// Object should have been created on the client with connection ID 1.
-	for (FReplicationSystemTestClient*& Client : ClientArray)
-	{
-		const uint32 ConnectionId = Client->ConnectionIdOnServer;
-		const bool bShouldHaveObject = ConnectionId == 1;
-		if (bShouldHaveObject)
-		{
-			UE_NET_ASSERT_NE(Client->GetReplicationBridge()->GetReplicatedObject(ServerObject->NetRefHandle), nullptr);
-		}
-		else
-		{
-			UE_NET_ASSERT_EQ(Client->GetReplicationBridge()->GetReplicatedObject(ServerObject->NetRefHandle), nullptr);
-		}
-	}
-}
-
-UE_NET_TEST_FIXTURE(FTestFilteringFixture, ConnectionFilterReplicatesOnlyToAllowedConnection)
-{
-	// Add clients
-	FReplicationSystemTestClient* ClientArray[] = {CreateClient(), CreateClient(), CreateClient()};
-	constexpr SIZE_T LastClientIndex = UE_ARRAY_COUNT(ClientArray) - 1;
-
-	// Spawn object on server
-	UReplicatedTestObject* ServerObject = Server->CreateObject(0, 0);
-
-	// Apply filtering that allows the last client to receive the object
-	const uint32 ConnectionIdForLastClient = ClientArray[LastClientIndex]->ConnectionIdOnServer;
-	TBitArray<> AllowedConnections;
-	AllowedConnections.Init(false, ConnectionIdForLastClient + 1);
-	AllowedConnections[ConnectionIdForLastClient] = true;
-	Server->ReplicationSystem->SetConnectionFilter(ServerObject->NetRefHandle, AllowedConnections, ENetFilterStatus::Allow);
-
-	// Send and deliver packets
-	Server->UpdateAndSend(ClientArray, DeliverPacket);
-
-	// Object should only have been created on the last client
-	for (FReplicationSystemTestClient*& Client : ClientArray)
-	{
-		const SIZE_T ClientIndex = &Client - &ClientArray[0];
-		if (ClientIndex == LastClientIndex)
-		{
-			UE_NET_ASSERT_NE(Client->GetReplicationBridge()->GetReplicatedObject(ServerObject->NetRefHandle), nullptr);
-		}
-		else
-		{
-			UE_NET_ASSERT_EQ(Client->GetReplicationBridge()->GetReplicatedObject(ServerObject->NetRefHandle), nullptr);
-		}
-	}
-}
-
-UE_NET_TEST_FIXTURE(FTestFilteringFixture, CanChangeConnectionFilter)
-{
-	// Add client
-	FReplicationSystemTestClient* Client = CreateClient();
-
-	// Spawn object on server
-	UReplicatedTestObject* ServerObject = Server->CreateObject(0, 0);
-
-	// Send and deliver packets
-	Server->NetUpdate();
-	Server->SendAndDeliverTo(Client, DeliverPacket);
-	Server->PostSendUpdate();
-
-	// Object should now exist on client
-	UE_NET_ASSERT_NE(Client->GetReplicationBridge()->GetReplicatedObject(ServerObject->NetRefHandle), nullptr);
-
-	// Apply connection filtering, not allowing any connection to receive the object
-	TBitArray<> NoConnections;
-	Server->ReplicationSystem->SetConnectionFilter(ServerObject->NetRefHandle, NoConnections, ENetFilterStatus::Allow);
-
-	// Send and deliver packets
-	Server->NetUpdate();
-	Server->SendAndDeliverTo(Client, DeliverPacket);
-	Server->PostSendUpdate();
-
-	// As the object is now filtered it should be deleted on the client
-	UE_NET_ASSERT_EQ(Client->GetReplicationBridge()->GetReplicatedObject(ServerObject->NetRefHandle), nullptr);
-
-	// Set the client as an allowed connection
-	TBitArray<> AllowedConnections;
-	AllowedConnections.Init(false, Client->ConnectionIdOnServer + 1);
-	AllowedConnections[Client->ConnectionIdOnServer] = true;
-	Server->ReplicationSystem->SetConnectionFilter(ServerObject->NetRefHandle, AllowedConnections, ENetFilterStatus::Allow);
-
-	// Send and deliver packets
-	Server->NetUpdate();
-	Server->SendAndDeliverTo(Client, DeliverPacket);
-	Server->PostSendUpdate();
-
-	// The object should have been created again
-	UE_NET_ASSERT_NE(Client->GetReplicationBridge()->GetReplicatedObject(ServerObject->NetRefHandle), nullptr);
-
-	// Finally, set the filter to not include any connections again.
-	Server->ReplicationSystem->SetConnectionFilter(ServerObject->NetRefHandle, NoConnections, ENetFilterStatus::Allow);
-
-	// Send and deliver packets
-	Server->NetUpdate();
-	Server->SendAndDeliverTo(Client, DeliverPacket);
-	Server->PostSendUpdate();
-
-	// The client is no longer owning the object so it should be deleted
-	UE_NET_ASSERT_EQ(Client->GetReplicationBridge()->GetReplicatedObject(ServerObject->NetRefHandle), nullptr);
-}
-
-UE_NET_TEST_FIXTURE(FTestFilteringFixture, CanToggleConnectionFilter)
-{
-	// Add client
-	FReplicationSystemTestClient* Client = CreateClient();
-
-	// Spawn object on server
-	UReplicatedTestObject* ServerObject = Server->CreateObject(0, 0);
-
-	// Send and deliver packets
-	Server->NetUpdate();
-	Server->SendAndDeliverTo(Client, DeliverPacket);
-	Server->PostSendUpdate();
-
-	// Object should now exist on client
-	UE_NET_ASSERT_NE(Client->GetReplicationBridge()->GetReplicatedObject(ServerObject->NetRefHandle), nullptr);
-
-	// Apply connection filtering, not allowing any connection to receive the object
-	TBitArray<> NoConnections;
-	Server->ReplicationSystem->SetConnectionFilter(ServerObject->NetRefHandle, NoConnections, ENetFilterStatus::Allow);
-
-	// Send and deliver packets
-	Server->NetUpdate();
-	Server->SendAndDeliverTo(Client, DeliverPacket);
-	Server->PostSendUpdate();
-
-	// As object is now filtered it should be deleted on the client
-	UE_NET_ASSERT_EQ(Client->GetReplicationBridge()->GetReplicatedObject(ServerObject->NetRefHandle), nullptr);
-
-	// Remove the connection filter
-	Server->ReplicationSystem->SetFilter(ServerObject->NetRefHandle, InvalidNetObjectFilterHandle);
-
-	// Send and deliver packets
-	Server->NetUpdate();
-	Server->SendAndDeliverTo(Client, DeliverPacket);
-	Server->PostSendUpdate();
-
-	// The object should have been created again
-	UE_NET_ASSERT_NE(Client->GetReplicationBridge()->GetReplicatedObject(ServerObject->NetRefHandle), nullptr);
-
-	// Destroy the object
-	Server->DestroyObject(ServerObject);
-
-	// Send and deliver packets
-	Server->NetUpdate();
-	Server->SendAndDeliverTo(Client, DeliverPacket);
-	Server->PostSendUpdate();
-}
-
-
-// Keep last! Toggle between different kinds of filters
-UE_NET_TEST_FIXTURE(FTestFilteringFixture, CanToggleBetweenAllFilters)
-{
-	// Add client
-	FReplicationSystemTestClient* Client = CreateClient();
-
-	// Spawn object on server
-	UReplicatedTestObject* ServerObject = Server->CreateObject(0, 0);
-
-	// Toggle between owner and connection filters
-	{
-		// Apply owner filtering
-		Server->ReplicationSystem->SetFilter(ServerObject->NetRefHandle, ToOwnerFilterHandle);
-
-		// Send and deliver packets
-		Server->NetUpdate();
-		Server->SendAndDeliverTo(Client, DeliverPacket);
-		Server->PostSendUpdate();
-
-		// Object should not exist on client
-		UE_NET_ASSERT_EQ(Client->GetReplicationBridge()->GetReplicatedObject(ServerObject->NetRefHandle), nullptr);
-
-		// Apply connection filtering, not allowing any connection to receive the object
-		TBitArray<> NoConnections;
-		Server->ReplicationSystem->SetConnectionFilter(ServerObject->NetRefHandle, NoConnections, ENetFilterStatus::Allow);
-
-		// Send and deliver packets
-		Server->NetUpdate();
-		Server->SendAndDeliverTo(Client, DeliverPacket);
-		Server->PostSendUpdate();
-
-		// Object should still not exist on client
-		UE_NET_ASSERT_EQ(Client->GetReplicationBridge()->GetReplicatedObject(ServerObject->NetRefHandle), nullptr);
-
-		// Finally test going from connection filtering to owner filtering
-		Server->ReplicationSystem->SetFilter(ServerObject->NetRefHandle, ToOwnerFilterHandle);
-
-		// Send and deliver packets
-		Server->NetUpdate();
-		Server->SendAndDeliverTo(Client, DeliverPacket);
-		Server->PostSendUpdate();
-
-		// Object should still not exist on client
-		UE_NET_ASSERT_EQ(Client->GetReplicationBridge()->GetReplicatedObject(ServerObject->NetRefHandle), nullptr);
-	}
-}
-
 // Owner filtering tests
 UE_NET_TEST_FIXTURE(FTestFilteringFixture, GroupFilterPreventsObjectFromReplicating)
 {
@@ -866,75 +524,6 @@ UE_NET_TEST_FIXTURE(FTestFilteringFixture, GroupFilterAllowsObjectToReplicate)
 	UE_NET_ASSERT_NE(Client->GetReplicationBridge()->GetReplicatedObject(ServerObject->NetRefHandle), nullptr);
 }
 
-UE_NET_TEST_FIXTURE(FTestFilteringFixture, GroupFilterRestoresConnectionConnectionFilter)
-{
-	// Add clients
-	FReplicationSystemTestClient* ClientArray[] = {CreateClient(), CreateClient(), CreateClient()};
-	constexpr SIZE_T LastClientIndex = UE_ARRAY_COUNT(ClientArray) - 1;
-
-	// Spawn object on server
-	UReplicatedTestObject* ServerObject = Server->CreateObject(0, 0);
-
-	// Apply filtering that allows the last client to receive the object
-	const uint32 ConnectionIdForLastClient = ClientArray[LastClientIndex]->ConnectionIdOnServer;
-	TBitArray<> AllowedConnections;
-	AllowedConnections.Init(false, ConnectionIdForLastClient + 1);
-	AllowedConnections[ConnectionIdForLastClient] = true;
-	Server->ReplicationSystem->SetConnectionFilter(ServerObject->NetRefHandle, AllowedConnections, ENetFilterStatus::Allow);
-
-	// Send and deliver packets
-	Server->UpdateAndSend(ClientArray, DeliverPacket);
-
-	// Object should only have been created on the last client
-	for (FReplicationSystemTestClient*& Client : ClientArray)
-	{
-		const SIZE_T ClientIndex = &Client - &ClientArray[0];
-		if (ClientIndex == LastClientIndex)
-		{
-			UE_NET_ASSERT_NE(Client->GetReplicationBridge()->GetReplicatedObject(ServerObject->NetRefHandle), nullptr);
-		}
-		else
-		{
-			UE_NET_ASSERT_EQ(Client->GetReplicationBridge()->GetReplicatedObject(ServerObject->NetRefHandle), nullptr);
-		}
-	}
-
-	// Create and set group filter for last client only
-	FNetObjectGroupHandle GroupHandle = Server->ReplicationSystem->CreateGroup(NAME_None);
-	Server->ReplicationSystem->AddToGroup(GroupHandle, ServerObject->NetRefHandle);
-	Server->ReplicationSystem->AddExclusionFilterGroup(GroupHandle);
-	Server->ReplicationSystem->SetGroupFilterStatus(GroupHandle, ConnectionIdForLastClient, ENetFilterStatus::Disallow);
-
-	// Send and deliver packets
-	Server->UpdateAndSend(ClientArray, DeliverPacket);
-
-	// Object Should now have been destroyed on the last client
-	for (FReplicationSystemTestClient*& Client : ClientArray)
-	{
-		UE_NET_ASSERT_EQ(Client->GetReplicationBridge()->GetReplicatedObject(ServerObject->NetRefHandle), nullptr);
-	}
-
-	// Clear Group filter
-	Server->ReplicationSystem->SetGroupFilterStatus(GroupHandle, ConnectionIdForLastClient, ENetFilterStatus::Allow);
-
-	// Send and deliver packets
-	Server->UpdateAndSend(ClientArray, DeliverPacket);
-
-	// Object should now be recreated on last client
-	for (FReplicationSystemTestClient*& Client : ClientArray)
-	{
-		const SIZE_T ClientIndex = &Client - &ClientArray[0];
-		if (ClientIndex == LastClientIndex)
-		{
-			UE_NET_ASSERT_NE(Client->GetReplicationBridge()->GetReplicatedObject(ServerObject->NetRefHandle), nullptr);
-		}
-		else
-		{
-			UE_NET_ASSERT_EQ(Client->GetReplicationBridge()->GetReplicatedObject(ServerObject->NetRefHandle), nullptr);
-		}
-	}
-}
-
 UE_NET_TEST_FIXTURE(FTestFilteringFixture, CanGetDynamicFilter)
 {
 	UE_NET_ASSERT_NE(MockNetObjectFilter, nullptr);
@@ -1050,8 +639,7 @@ UE_NET_TEST_FIXTURE(FTestFilteringFixture, SwitchingFiltersCallsRemoveObjectOnPr
 	{
 		MockNetObjectFilter->ResetFunctionCallStatus();
 
-		const TBitArray<> NoConnections;
-		Server->ReplicationSystem->SetConnectionFilter(ServerObject->NetRefHandle, NoConnections, ENetFilterStatus::Disallow);
+		Server->ReplicationSystem->SetFilter(ServerObject->NetRefHandle, ToOwnerFilterHandle);
 
 		Server->NetUpdate();
 		Server->PostSendUpdate();
@@ -1447,52 +1035,6 @@ UE_NET_TEST_FIXTURE(FTestFilteringFixture, InclusionGroupDoesNotOverrideOwnerFil
 	UE_NET_ASSERT_EQ(Clients[NonOwningClientIndex]->GetReplicationBridge()->GetReplicatedObject(ServerObject->NetRefHandle), nullptr);
 }
 
-UE_NET_TEST_FIXTURE(FTestFilteringFixture, InclusionGroupDoesNotOverrideConnectionFilter)
-{
-	// Add clients
-	constexpr uint32 AllowedClientIndex = 0;
-	constexpr uint32 DisallowedOwningClientIndex = 1;
-
-	FReplicationSystemTestClient* ClientArray[2];
-	for (FReplicationSystemTestClient*& Client : ClientArray)
-	{
-		Client = CreateClient();
-	}
-
-	// Spawn object on server
-	UReplicatedTestObject* ServerObject = Server->CreateObject(UTestReplicatedIrisObject::FComponents{});
-
-	// Apply connection filter
-	TBitArray<> AllowedConnections;
-	AllowedConnections.Init(false, ClientArray[AllowedClientIndex]->ConnectionIdOnServer + 1);
-	AllowedConnections[ClientArray[AllowedClientIndex]->ConnectionIdOnServer] = true;
-	Server->ReplicationSystem->SetConnectionFilter(ServerObject->NetRefHandle, AllowedConnections, ENetFilterStatus::Allow);
-
-	// Setup inclusion group filter
-	FNetObjectGroupHandle GroupHandle = Server->ReplicationSystem->CreateGroup(NAME_None);
-	Server->ReplicationSystem->AddToGroup(GroupHandle, ServerObject->NetRefHandle);
-	Server->ReplicationSystem->AddInclusionFilterGroup(GroupHandle);
-	Server->ReplicationSystem->SetGroupFilterStatus(GroupHandle, ENetFilterStatus::Allow);
-
-	// Send and deliver packet
-	Server->UpdateAndSend(ClientArray, DeliverPacket);
-
-	// Object should have been created on the owning client
-	UE_NET_ASSERT_NE(Clients[AllowedClientIndex]->GetReplicationBridge()->GetReplicatedObject(ServerObject->NetRefHandle), nullptr);
-	// Object should not have been created on the non-owning client
-	UE_NET_ASSERT_EQ(Clients[DisallowedOwningClientIndex]->GetReplicationBridge()->GetReplicatedObject(ServerObject->NetRefHandle), nullptr);
-
-	// Change filter to not allow replication. As it's an inclusion filter this should not change things.
-	Server->ReplicationSystem->SetGroupFilterStatus(GroupHandle, ENetFilterStatus::Disallow);
-
-	// Send and deliver packets
-	Server->UpdateAndSend(ClientArray, DeliverPacket);
-
-	// Object creation status should remain the same as before
-	UE_NET_ASSERT_NE(Clients[AllowedClientIndex]->GetReplicationBridge()->GetReplicatedObject(ServerObject->NetRefHandle), nullptr);
-	UE_NET_ASSERT_EQ(Clients[DisallowedOwningClientIndex]->GetReplicationBridge()->GetReplicatedObject(ServerObject->NetRefHandle), nullptr);
-}
-
 UE_NET_TEST_FIXTURE(FTestFilteringFixture, InclusionGroupDoesNotOverrideExclusionGroupFilter)
 {
 	// Add clients
diff --git a/Engine/Plugins/Runtime/ReplicationSystemTestPlugin/Source/Private/Tests/ReplicationSystem/RPC/TestObjectRPC.cpp b/Engine/Plugins/Runtime/ReplicationSystemTestPlugin/Source/Private/Tests/ReplicationSystem/RPC/TestObjectRPC.cpp
--- a/Engine/Plugins/Runtime/ReplicationSystemTestPlugin/Source/Private/Tests/ReplicationSystem/RPC/TestObjectRPC.cpp
+++ b/Engine/Plugins/Runtime/ReplicationSystemTestPlugin/Source/Private/Tests/ReplicationSystem/RPC/TestObjectRPC.cpp
@@ -426,56 +426,6 @@ namespace UE::Net::Private
 		}
 	}
 
-	// This test is written to verify a bug with connection filtering and late join. There will be an ensure if this test isn't working.
-	UE_NET_TEST_FIXTURE(FRPCTestFixture, TestSubObjectMulticastRPCIsNotReplicatedToFilteredOutConnection)
-	{
-		constexpr SIZE_T OwningClientIndex = 0;
-		constexpr SIZE_T NonOwningClientIndex = 1;
-
-		// Add a client
-		FReplicationSystemTestClient* ClientArray[2] = {};
-		ClientArray[OwningClientIndex] = CreateClient();
-
-		// Spawn connection filtered object with on server
-		UTestReplicatedObjectWithRPC* ServerRootObject = Server->CreateObject<UTestReplicatedObjectWithRPC>();
-		ServerRootObject->Init(Server->GetReplicationSystem());
-
-		const FNetRefHandle ServerRootObjectHandle = ServerRootObject->NetRefHandle;
-		
-		// Set connection filter
-		{
-			TBitArray<> AllowedConnections;
-			AllowedConnections.Add(false, ClientArray[OwningClientIndex]->ConnectionIdOnServer + 1);
-			AllowedConnections[ClientArray[OwningClientIndex]->ConnectionIdOnServer] = true;
-			Server->ReplicationSystem->SetConnectionFilter(ServerRootObjectHandle, AllowedConnections, ENetFilterStatus::Allow);
-		}
-	
-		UTestReplicatedObjectWithRPC* ServerSubObject = Server->CreateSubObject<UTestReplicatedObjectWithRPC>(ServerRootObjectHandle);
-		ServerSubObject->Init(Server->GetReplicationSystem());
-		ServerSubObject->SetRootObject(ServerRootObject);
-
-		// Send and deliver packet
-		Server->UpdateAndSend({ClientArray[OwningClientIndex]});
-
-		// Late join with second client
-		ClientArray[NonOwningClientIndex] = CreateClient();
-
-		// Call a multicast RPC server->client on the subobject
-		ServerSubObject->NetMulticast_MultiCastRPC();
-
-		// Send and deliver packet
-		{
-			FEnsureScope EnsureScope;
-			FCheckScope CheckScope;
-
-			Server->UpdateAndSend(ClientArray);
-
-			// No ensures or checks should trigger
-			UE_NET_ASSERT_EQ(EnsureScope.GetCount(), 0);
-			UE_NET_ASSERT_EQ(CheckScope.GetCount(), 0);
-		}
-	}
-
 	// This test is written to test that we can control if ordered unreliable attachments are pruned or not after each tick.
 	UE_NET_TEST_FIXTURE(FRPCTestFixture, TestOvercommitOrderedUnreliableRPC)
 	{
diff --git a/Engine/Plugins/Runtime/ReplicationSystemTestPlugin/Source/Private/Tests/ReplicationSystem/TestFlushBeforeDestroy.cpp b/Engine/Plugins/Runtime/ReplicationSystemTestPlugin/Source/Private/Tests/ReplicationSystem/TestFlushBeforeDestroy.cpp
--- a/Engine/Plugins/Runtime/ReplicationSystemTestPlugin/Source/Private/Tests/ReplicationSystem/TestFlushBeforeDestroy.cpp
+++ b/Engine/Plugins/Runtime/ReplicationSystemTestPlugin/Source/Private/Tests/ReplicationSystem/TestFlushBeforeDestroy.cpp
@@ -1679,57 +1679,6 @@ UE_NET_TEST_FIXTURE(FTestFlushBeforeDestroyFixture, TestSubobjectConditionOnSubo
 	UE_NET_ASSERT_EQ(ClientSubObject1->ConnectionFilteredComponents[0]->SkipOwnerA, ServerSubObject1->ConnectionFilteredComponents[0]->SkipOwnerA);
 }
 
-// Test sending attachment to connection filtered object and tear it off
-UE_NET_TEST_FIXTURE(FTestFlushBeforeDestroyFixture, TestTearOffOfConnectionFilteredObject)
-{
-	// Add clients
-	FReplicationSystemTestClient* ClientArray[] = {CreateClient(), CreateClient()};
-	TStrongObjectPtr<UMockNetObjectAttachmentHandler> ClientMockNetObjectAttachmentHandlers[UE_ARRAY_COUNT(ClientArray)];
-	for (FReplicationSystemTestClient*& Client : ClientArray)
-	{
-		RegisterNetBlobHandlers(Client);
-		ClientMockNetObjectAttachmentHandlers[&Client - ClientArray] = ClientMockNetObjectAttachmentHandler;
-	}
-
-	// Setup case where we have a new object for which we have an attachment which should execute a tearoff after we have confirmed creation
-	UReplicatedTestObject* ServerObject = Server->CreateObject();
-
-	// Apply connection filter
-	TBitArray<> AllowedConnections;
-	AllowedConnections.Init(false, Clients[0]->ConnectionIdOnServer + 1);
-	AllowedConnections[Clients[0]->ConnectionIdOnServer] = true;
-	Server->ReplicationSystem->SetConnectionFilter(ServerObject->NetRefHandle, AllowedConnections, ENetFilterStatus::Allow);
-
-	// Send update
-	Server->UpdateAndSend({ClientArray}, DeliverPacket);
-
-	// Create attachment
-	{
-		constexpr uint32 PayloadBitCount = 24;
-		const TRefCountPtr<FNetObjectAttachment>& Attachment = MockNetObjectAttachmentHandler->CreateReliableNetObjectAttachment(PayloadBitCount);
-		FNetObjectReference AttachmentTarget = FObjectReferenceCache::MakeNetObjectReference(ServerObject->NetRefHandle);
-		for (FReplicationSystemTestClient* Client : ClientArray)
-		{
-			Server->GetReplicationSystem()->QueueNetObjectAttachment(Client->ConnectionIdOnServer, AttachmentTarget, Attachment);
-		}
-	}
-
-	// Request tearoff
-	Server->ReplicationBridge->EndReplication(ServerObject, EEndReplicationFlags::TearOff);
-
-	// Update once and fail sending to cause torn off object to no longer be in scope
-	Server->UpdateAndSend({ClientArray}, DoNotDeliverPacket);
-
-	// Deliver a packet, this should flush the object and deliver the attachment to the allowed connection
-	Server->UpdateAndSend({ClientArray}, DeliverPacket);
-
-	// Verify that the attachment has been received on the allowed connection
-	UE_NET_ASSERT_EQ(ClientMockNetObjectAttachmentHandlers[0]->GetFunctionCallCounts().OnNetBlobReceived, 1U);
-
-	// Verify that the attachment was not received on the disallowed connection
-	UE_NET_ASSERT_EQ(ClientMockNetObjectAttachmentHandlers[1]->GetFunctionCallCounts().OnNetBlobReceived, 0U);
-}
-
 // Test sending attachment to owner filtered object and tear it off
 UE_NET_TEST_FIXTURE(FTestFlushBeforeDestroyFixture, TestTearOffOfOwnerFilteredObject)
 {
diff --git a/Engine/Source/Runtime/Net/Iris/Private/Iris/ReplicationSystem/Filtering/ReplicationFiltering.cpp b/Engine/Source/Runtime/Net/Iris/Private/Iris/ReplicationSystem/Filtering/ReplicationFiltering.cpp
--- a/Engine/Source/Runtime/Net/Iris/Private/Iris/ReplicationSystem/Filtering/ReplicationFiltering.cpp
+++ b/Engine/Source/Runtime/Net/Iris/Private/Iris/ReplicationSystem/Filtering/ReplicationFiltering.cpp
@@ -173,7 +173,7 @@ static FIrisLogOnceTracker ReplicationFiltering_MootDependentObjectTracker;
 FReplicationFiltering::FReplicationFiltering()
 : bHasNewConnection(0)
 , bHasRemovedConnection(0)
-, bHasDirtyConnectionFilter(0)
+, bHasDirtyOwnerFilter(0)
 , bHasDirtyOwner(0)
 , bHasDynamicFilters(0)
 , bHasDirtyExclusionFilterGroup(0)
@@ -251,12 +251,11 @@ void FReplicationFiltering::SetNetObjectListsSize(FInternalNetRefIndex MaxIntern
 
 	// Increase NetBitArrays
 	{
-		ObjectsWithDirtyConnectionFilter.SetNumBits(MaxInternalIndex);
+		ObjectsWithDirtyOwnerFilter.SetNumBits(MaxInternalIndex);
 		ObjectsWithDirtyOwner.SetNumBits(MaxInternalIndex);
 		ObjectsWithOwnerFilter.SetNumBits(MaxInternalIndex);
 		ObjectsWithPerObjectInfo.SetNumBits(MaxInternalIndex);
 
-		AllConnectionFilteredObjects.SetNumBits(MaxInternalIndex);
 		DynamicFilterEnabledObjects.SetNumBits(MaxInternalIndex);
 		
 		ObjectsRequiringDynamicFilterUpdate.SetNumBits(MaxInternalIndex);
@@ -346,7 +345,7 @@ void FReplicationFiltering::Filter()
 
 	UpdateGroupInclusionFiltering();
 
-	UpdateOwnerAndConnectionFiltering();
+	UpdateOwnerFiltering();
 
 	UpdateSubObjectFilters();
 
@@ -416,7 +415,6 @@ void FReplicationFiltering::BuildAlwaysRelevantList(FNetBitArrayView OutAlwaysRe
 
 	// The different list of filtered objects
 	const uint32* const WithOwnerData = ObjectsWithOwnerFilter.GetDataChecked(MaxWords);
-	const uint32* const ConnectionFiltersData = AllConnectionFilteredObjects.GetDataChecked(MaxWords);
 	const uint32* const DynamicFilteredData = DynamicFilterEnabledObjects.GetDataChecked(MaxWords);
 
 	const FNetBitArrayView GroupFilteredOutView = Groups->GetGroupFilteredOutObjects();
@@ -427,7 +425,7 @@ void FReplicationFiltering::BuildAlwaysRelevantList(FNetBitArrayView OutAlwaysRe
 	for (uint32 WordIndex = 0; WordIndex < MaxWords; ++WordIndex)
 	{
 		// Build the list of always relevant objects, e.g. objects that have no filters
-		OutAlwaysRelevantListData[WordIndex] = ScopeListData[WordIndex] & ~(WithOwnerData[WordIndex] | ConnectionFiltersData[WordIndex] | DynamicFilteredData[WordIndex] | GroupFilteredOutData[WordIndex]);
+		OutAlwaysRelevantListData[WordIndex] = ScopeListData[WordIndex] & ~(WithOwnerData[WordIndex] | DynamicFilteredData[WordIndex] | GroupFilteredOutData[WordIndex]);
 	}
 }
 
@@ -458,8 +456,8 @@ void FReplicationFiltering::SetOwningConnection(FInternalNetRefIndex ObjectIndex
 		ObjectsWithDirtyOwner.SetBit(ObjectIndex);
 		if (HasOwnerFilter(ObjectIndex))
 		{
-			bHasDirtyConnectionFilter = 1;
-			ObjectsWithDirtyConnectionFilter.SetBit(ObjectIndex);
+			bHasDirtyOwnerFilter = 1;
+			ObjectsWithDirtyOwnerFilter.SetBit(ObjectIndex);
 		}
 	}
 }
@@ -502,10 +500,10 @@ bool FReplicationFiltering::SetFilter(FInternalNetRefIndex ObjectIndex, FNetObje
 
 	auto ClearAndDirtyStaticFilter = [this](uint32 ObjIndex)
 	{
-		this->bHasDirtyConnectionFilter = 1;
+		this->bHasDirtyOwnerFilter = 1;
 
 		this->ObjectsWithOwnerFilter.ClearBit(ObjIndex);
-		this->ObjectsWithDirtyConnectionFilter.SetBit(ObjIndex);
+		this->ObjectsWithDirtyOwnerFilter.SetBit(ObjIndex);
 		this->FreePerObjectInfoForObject(ObjIndex);
 	};
 
@@ -587,54 +585,6 @@ bool FReplicationFiltering::IsUsingSpatialFilter(FInternalNetRefIndex ObjectInde
 	return bIsUsingSpatialFilter;
 }
 
-bool FReplicationFiltering::SetConnectionFilter(FInternalNetRefIndex ObjectIndex, const FNetBitArrayView& ConnectionIndices, ENetFilterStatus ReplicationStatus)
-{
-	// If we have some sort of other filtering set then remove that. From now on we want to use connection filtering.
-	if (HasOwnerFilter(ObjectIndex))
-	{
-		ObjectsWithOwnerFilter.ClearBit(ObjectIndex);
-	}
-
-	const uint8 DynamicFilterIndex = ObjectIndexToDynamicFilterIndex[ObjectIndex];
-	if (DynamicFilterIndex != InvalidDynamicFilterIndex)
-	{
-		RemoveFromDynamicFilter(ObjectIndex, DynamicFilterIndex);
-	}
-
-	bHasDirtyConnectionFilter = 1;
-	ObjectsWithDirtyConnectionFilter.SetBit(ObjectIndex);
-
-	PerObjectInfoIndexType ObjectInfoIndex = ObjectIndexToPerObjectInfoIndex[ObjectIndex];
-	FPerObjectInfo* ObjectInfo = nullptr;
-	if (ObjectInfoIndex == 0)
-	{
-		ObjectInfo = AllocPerObjectInfoForObject(ObjectIndex);
-	}
-	else
-	{
-		ObjectInfo = GetPerObjectInfo(ObjectInfoIndex);
-	}
-
-	{
-		// The ConnectionIndices bit array does not necessarily have to have storage for the max amount of connections
-		const uint32 InWordCount = ConnectionIndices.GetNumWords();
-		const uint32 MaxWordCount = PerObjectInfoStorageCountForConnections;
-		const uint32 WordMask = (ReplicationStatus == ENetFilterStatus::Allow ? 0U : ~0U);
-		for (uint32 WordIt = 0, WordEndIt = FPlatformMath::Min(InWordCount, MaxWordCount); WordIt < WordEndIt; ++WordIt)
-		{
-			ObjectInfo->ConnectionIds[WordIt] = ConnectionIndices.GetWord(WordIt) ^ WordMask;
-		}
-
-		// Fill in connections that weren't passed in ConnectionIndices
-		for (uint32 WordIt = InWordCount; WordIt < MaxWordCount; ++WordIt)
-		{
-			ObjectInfo->ConnectionIds[WordIt] = WordMask;
-		}
-	}
-
-	return true;
-}
-
 FNetObjectFilterHandle FReplicationFiltering::GetFilterHandle(const FName FilterName) const
 {
 	for (const FFilterInfo& Info : DynamicFilterInfos)
@@ -930,7 +880,7 @@ void FReplicationFiltering::UpdateObjectsInScope()
 	ModifiedWords.SetNumUninitialized(WordCountForObjectBitArrays);
 
 	uint32 ModifiedWordIndex = 0;
-	uint32* ObjectsWithDirtyConnectionFilterStorage = ObjectsWithDirtyConnectionFilter.GetDataChecked(WordCountForObjectBitArrays);
+	uint32* ObjectsWithDirtyOwnerFilterStorage = ObjectsWithDirtyOwnerFilter.GetDataChecked(WordCountForObjectBitArrays);
 	uint32* ObjectsWithDirtyOwnerStorage = ObjectsWithDirtyOwner.GetDataChecked(WordCountForObjectBitArrays);
 	const uint32* ObjectsInScopeStorage = ObjectsInScope.GetDataChecked(WordCountForObjectBitArrays);
 	uint32* FakePrevObjectsInScopeStorage = FakePrevObjectsInScope.GetDataChecked(WordCountForObjectBitArrays);
@@ -942,9 +892,9 @@ void FReplicationFiltering::UpdateObjectsInScope()
 		{
 			const uint32 ObjectsInScopeWord = ObjectsInScopeStorage[WordIt];
 			const uint32 PrevObjectsInScopeWord = PrevObjectsInScopeStorage[WordIt];
-			const uint32 ObjectsWithDirtyConnectionFilterWord = ObjectsWithDirtyConnectionFilterStorage[WordIt];
+			const uint32 ObjectsWithDirtyOwnerFilterWord = ObjectsWithDirtyOwnerFilterStorage[WordIt];
 			const uint32 ObjectsWithDirtyOwnerWord = ObjectsWithDirtyOwnerStorage[WordIt];
-			const uint32 SameFrameRemovedWord = ~(ObjectsInScopeWord | PrevObjectsInScopeWord) & (ObjectsWithDirtyConnectionFilterWord | ObjectsWithDirtyOwnerWord);
+			const uint32 SameFrameRemovedWord = ~(ObjectsInScopeWord | PrevObjectsInScopeWord) & (ObjectsWithDirtyOwnerFilterWord | ObjectsWithDirtyOwnerWord);
 
 			// Pretend that same frame removed objects existed in the previous frame.
 			FakePrevObjectsInScopeStorage[WordIt] = PrevObjectsInScopeWord | SameFrameRemovedWord;
@@ -975,7 +925,7 @@ void FReplicationFiltering::UpdateObjectsInScope()
 			const uint32 ExistingObjects = ObjectsInScopeStorage[WordIndex];
 
 			// Deleted objects can't be dirty and can't have filtering.
-			ObjectsWithDirtyConnectionFilterStorage[WordIndex] &= ExistingObjects;
+			ObjectsWithDirtyOwnerFilterStorage[WordIndex] &= ExistingObjects;
 			ObjectsWithOwnerFilterStorage[WordIndex] &= ExistingObjects;
 			ObjectsWithDirtyOwnerStorage[WordIndex] &= ExistingObjects;
 
@@ -1056,11 +1006,11 @@ void FReplicationFiltering::UpdateObjectsInScope()
 						ObjectsWithDirtyOwner.SetBit(ParentIndex);
 					}
 
-					if (HasOwnerFilter(ParentIndex) || HasConnectionFilter(ParentIndex))
+					if (HasOwnerFilter(ParentIndex))
 					{
-						bHasDirtyConnectionFilter = 1;
+						bHasDirtyOwnerFilter = 1;
 						// Updating the parent will update all subobjects.
-						ObjectsWithDirtyConnectionFilter.SetBit(ParentIndex);
+						ObjectsWithDirtyOwnerFilter.SetBit(ParentIndex);
 					}
 				}
 			}
@@ -1126,9 +1076,9 @@ uint32 FReplicationFiltering::GetOwningConnectionIfDirty(uint32 ObjectIndex) con
 	return ObjectIndexToOwningConnection[ObjectIndex];
 }
 
-void FReplicationFiltering::UpdateOwnerAndConnectionFiltering()
+void FReplicationFiltering::UpdateOwnerFiltering()
 {
-	if (!(bHasDirtyOwner || bHasDirtyConnectionFilter))
+	if (!(bHasDirtyOwner || bHasDirtyOwnerFilter))
 	{
 		return;
 	}
@@ -1158,14 +1108,9 @@ void FReplicationFiltering::UpdateOwnerAndConnectionFiltering()
 	const FNetBitArrayView CurrentFrameObjectsInScope = NetRefHandleManager->GetCurrentFrameScopableInternalIndices();
 
 	// Update filtering
-	if (bHasDirtyConnectionFilter)
+	if (bHasDirtyOwnerFilter)
 	{
-		IRIS_PROFILER_SCOPE(FReplicationFiltering_UpdateDirtyConnectionFilter);
-
-		ObjectsWithDirtyConnectionFilter.ForAllSetBits([this](uint32 DirtyObjectIndex)
-		{
-			AllConnectionFilteredObjects.SetBitValue(DirtyObjectIndex, HasConnectionFilter(DirtyObjectIndex));
-		});
+		IRIS_PROFILER_SCOPE(FReplicationFiltering_UpdateDirtyOwnerFilter);
 
 		auto UpdateConnectionScope = [this, &CurrentFrameObjectsInScope](uint32 ConnectionId)
 		{
@@ -1185,13 +1130,6 @@ void FReplicationFiltering::UpdateOwnerAndConnectionFiltering()
 						const bool bIsOwner = (ConnectionId == OwningConnection);
 						bObjectIsInScope = bIsOwner;
 					}
-					else if (HasConnectionFilter(ObjectIndex))
-					{
-						const PerObjectInfoIndexType ObjectInfoIndex = ObjectIndexToPerObjectInfoIndex[ObjectIndex];
-						const FPerObjectInfo* ObjectInfo = GetPerObjectInfo(ObjectInfoIndex);
-						const ENetFilterStatus ReplicationStatus = GetConnectionFilterStatus(*ObjectInfo, ConnectionId);
-						bObjectIsInScope = ReplicationStatus == ENetFilterStatus::Allow;
-					}
 
 					// Update scope for parent object.
 					{
@@ -1211,8 +1149,7 @@ void FReplicationFiltering::UpdateOwnerAndConnectionFiltering()
 					}
 				};
 
-				// Update the connection scope for dirty objects
-				ObjectsWithDirtyConnectionFilter.ForAllSetBits(MaskObject);
+				ObjectsWithDirtyOwnerFilter.ForAllSetBits(MaskObject);
 			}
 		};
 
@@ -1220,9 +1157,9 @@ void FReplicationFiltering::UpdateOwnerAndConnectionFiltering()
 	}
 	
 	// Clear out dirtiness
-	bHasDirtyConnectionFilter = 0;
+	bHasDirtyOwnerFilter = 0;
 	bHasDirtyOwner = 0;
-	ObjectsWithDirtyConnectionFilter.ClearAllBits();
+	ObjectsWithDirtyOwnerFilter.ClearAllBits();
 	ObjectsWithDirtyOwner.ClearAllBits();
 }
 
@@ -1848,12 +1785,6 @@ bool FReplicationFiltering::HasOwnerFilter(uint32 ObjectIndex) const
 	return bHasOwnerFilter;
 }
 
-bool FReplicationFiltering::HasConnectionFilter(uint32 ObjectIndex) const
-{
-	const bool bHasConnectionFilter = ObjectsWithPerObjectInfo.GetBit(ObjectIndex);
-	return bHasConnectionFilter;
-}
-
 FReplicationFiltering::PerObjectInfoIndexType FReplicationFiltering::AllocPerObjectInfo()
 {
 	const uint32 UsedPerObjectInfoBitsPerWord = sizeof(decltype(UsedPerObjectInfoStorage)::ElementType)*8;
@@ -2789,15 +2720,6 @@ FString FReplicationFiltering::PrintFilterObjectInfo(FInternalNetRefIndex Object
 			
 			PrintBuilder.Appendf(TEXT("\r\t OwnerFilter to connectionId:%u | %s"), ConnectionId, bIsOwner?TEXT("relevant"):TEXT("not relevant"));
 		}
-		else if (HasConnectionFilter(ObjectIndex))
-		{
-			const PerObjectInfoIndexType ObjectInfoIndex = ObjectIndexToPerObjectInfoIndex[ObjectIndex];
-			const FPerObjectInfo* ObjectInfo = GetPerObjectInfo(ObjectInfoIndex);
-
-			ENetFilterStatus FilterStatus = GetConnectionFilterStatus(*ObjectInfo, ConnectionId);
-
-			PrintBuilder.Appendf(TEXT("\r\t ConnectionFilter to connectionId:%u | %s"), ConnectionId, LexToString(FilterStatus));
-		}
 
 		const FPerConnectionInfo& ConnectionInfo = ConnectionInfos[ConnectionId];
 
diff --git a/Engine/Source/Runtime/Net/Iris/Private/Iris/ReplicationSystem/Filtering/ReplicationFiltering.h b/Engine/Source/Runtime/Net/Iris/Private/Iris/ReplicationSystem/Filtering/ReplicationFiltering.h
--- a/Engine/Source/Runtime/Net/Iris/Private/Iris/ReplicationSystem/Filtering/ReplicationFiltering.h
+++ b/Engine/Source/Runtime/Net/Iris/Private/Iris/ReplicationSystem/Filtering/ReplicationFiltering.h
@@ -102,9 +102,6 @@ public:
 	 */
 	bool IsUsingSpatialFilter(FInternalNetRefIndex ObjectIndex) const;
 
-	// Set whether an object is allowed to be replicated to certain connections or not.
-	bool SetConnectionFilter(FInternalNetRefIndex ObjectIndex, const FNetBitArrayView& ConnectionIndices, ENetFilterStatus ReplicationStatus);
-
 	FNetObjectFilterHandle GetFilterHandle(const FName FilterName) const;
 	UNetObjectFilter* GetFilter(const FName FilterName) const;
 
@@ -222,7 +219,7 @@ private:
 	void InitNewConnections();
 	void ResetRemovedConnections();
 	void UpdateObjectsInScope();
-	void UpdateOwnerAndConnectionFiltering();
+	void UpdateOwnerFiltering();
 	void UpdateGroupExclusionFiltering();
 	void UpdateGroupInclusionFiltering();
 	void UpdateSubObjectFilters();
@@ -248,7 +245,6 @@ private:
 	// Helper to update and reset group inclusion filter effects if objects are removed from a filter or after a filter status change, returns true if the group filter was changed
 	bool ClearGroupInclusionFilterEffectsForObject(uint32 ObjectIndex, uint32 ConnectionId);
 	bool HasOwnerFilter(uint32 ObjectIndex) const;
-	bool HasConnectionFilter(uint32 ObjectIndex) const;
 
 	PerObjectInfoIndexType AllocPerObjectInfo();
 	void FreePerObjectInfo(PerObjectInfoIndexType Index);
@@ -349,7 +345,7 @@ private:
 	uint32 WordCountForObjectBitArrays = 0;
 
 	// Filter specifics
-	FNetBitArray ObjectsWithDirtyConnectionFilter;
+	FNetBitArray ObjectsWithDirtyOwnerFilter;
 	FNetBitArray ObjectsWithDirtyOwner;
 
 	FNetBitArray ObjectsWithOwnerFilter;
@@ -412,7 +408,7 @@ private:
 
 	uint32 bHasNewConnection : 1;
 	uint32 bHasRemovedConnection : 1;
-	uint32 bHasDirtyConnectionFilter: 1;
+	uint32 bHasDirtyOwnerFilter: 1;
 	uint32 bHasDirtyOwner : 1;
 	uint32 bHasDynamicFilters : 1;
 	uint32 bHasDirtyExclusionFilterGroup : 1;
diff --git a/Engine/Source/Runtime/Net/Iris/Private/Iris/ReplicationSystem/ReplicationSystem.cpp b/Engine/Source/Runtime/Net/Iris/Private/Iris/ReplicationSystem/ReplicationSystem.cpp
--- a/Engine/Source/Runtime/Net/Iris/Private/Iris/ReplicationSystem/ReplicationSystem.cpp
+++ b/Engine/Source/Runtime/Net/Iris/Private/Iris/ReplicationSystem/ReplicationSystem.cpp
@@ -1509,29 +1509,6 @@ FName UReplicationSystem::GetFilterName(UE::Net::FNetObjectFilterHandle Filter)
 	return Impl->ReplicationSystemInternal.GetFiltering().GetFilterName(Filter);
 }
 
-bool UReplicationSystem::SetConnectionFilter(FNetRefHandle Handle, const TBitArray<>& Connections, UE::Net::ENetFilterStatus ReplicationStatus)
-{
-	using namespace UE::Net;
-	using namespace UE::Net::Private;
-
-	FNetRefHandleManager& NetRefHandleManager = Impl->ReplicationSystemInternal.GetNetRefHandleManager();
-	const FInternalNetRefIndex ObjectInternalIndex = NetRefHandleManager.GetInternalIndex(Handle);
-
-	if (Impl->ReplicationSystemInternal.AreFilterChangesBlocked())
-	{
-		ensureMsgf(false, TEXT("Setting filter conditions is not yet supported during this operation. Filter condition on %s (%s) failed."), *GetNameSafe(NetRefHandleManager.GetReplicatedObjectInstance(ObjectInternalIndex)), *Handle.ToString());
-		return false;
-	}
-
-	if (ObjectInternalIndex == FNetRefHandleManager::InvalidInternalIndex)
-	{
-		return false;
-	}
-
-	FReplicationFiltering& Filtering = Impl->ReplicationSystemInternal.GetFiltering();
-	return Filtering.SetConnectionFilter(ObjectInternalIndex, MakeNetBitArrayView(Connections.GetData(), Connections.Num()), ReplicationStatus);
-}
-
 UE::Net::FNetObjectGroupHandle UReplicationSystem::GetOrCreateSubObjectFilter(FName GroupName)
 {
 	using namespace UE::Net;
diff --git a/Engine/Source/Runtime/Net/Iris/Public/Iris/ReplicationSystem/ReplicationSystem.h b/Engine/Source/Runtime/Net/Iris/Public/Iris/ReplicationSystem/ReplicationSystem.h
--- a/Engine/Source/Runtime/Net/Iris/Public/Iris/ReplicationSystem/ReplicationSystem.h
+++ b/Engine/Source/Runtime/Net/Iris/Public/Iris/ReplicationSystem/ReplicationSystem.h
@@ -508,7 +508,7 @@ public:
 	IRISCORE_API uint32 GetOwningNetConnection(FNetRefHandle Handle) const;
 
 	/**
-	 * Sets a filter for a replicated object which will be used until the next call to SetFilter or SetConnectionFilter.
+	 * Sets a filter for a replicated object which will be used until the next call to SetFilter.
 	 * Filters are used prevent objects from being replicated to certain connections. An object that is filtered
 	 * out will cause the object to be destroyed on the remote side.
 	 * 
@@ -522,7 +522,6 @@ public:
 	 * the filter of the object is unspecified, it could be using a previous filter or use no filtering.
 	 * 
 	 * @see GetFilterHandle
-	 * @see SetConnectionFilter
 	 */
 	IRISCORE_API bool SetFilter(FNetRefHandle Handle, UE::Net::FNetObjectFilterHandle FilterHandle, FName FilterConfigProfile=NAME_None);
 
@@ -610,18 +609,6 @@ public:
 	/** Remove group from filtering system, will cancel effects of the group. */
 	IRISCORE_API void RemoveSubObjectFilter(FName GroupName);
 
-	/**
-	 * Set which connections the object is allowed to be replicated to. This will cancel the effect
-	 * of any previous SetFilter or SetConnectionFilter calls on the object.
-	 * It will also remove the object from any dynamic filters
-	 * @param Handle A valid handle to a replicated object.
-	 * @param Connections Set bits indicates the connection IDs that the object is allowed or not allowed to be replicated to depending on ReplicationStatus.
-	 * @param ReplicationStatus Whether the set bits in Connections indicate if the object is allowed or not allowed to be replicated to those connections.
-	 * @return true if the connection filter was properly set, false if not.
-	 * @see SetFilter
-	 */
-	IRISCORE_API bool SetConnectionFilter(FNetRefHandle Handle, const TBitArray<>& Connections, UE::Net::ENetFilterStatus ReplicationStatus);
-
 	/**
 	 * Enable or disable a replication condition for a single connection and do the inverse for all other connections. This will affect
 	 * the replication of properties with conditions that are dependent on this condition.
