Iris
* Removed more dependencies to UReplicationBridge

#rnx

diff --git a/Engine/Plugins/Runtime/ReplicationSystemTestPlugin/Source/Private/Tests/ReplicationSystem/ReplicatedTestObject.cpp b/Engine/Plugins/Runtime/ReplicationSystemTestPlugin/Source/Private/Tests/ReplicationSystem/ReplicatedTestObject.cpp
--- a/Engine/Plugins/Runtime/ReplicationSystemTestPlugin/Source/Private/Tests/ReplicationSystem/ReplicatedTestObject.cpp
+++ b/Engine/Plugins/Runtime/ReplicationSystemTestPlugin/Source/Private/Tests/ReplicationSystem/ReplicatedTestObject.cpp
@@ -119,7 +119,7 @@ void UReplicatedTestObjectBridge::SetExternalWorldLocationUpdateFunctor(TFunctio
 	TestFactory->SetWorldUpdateFunctor(LocUpdateFunctor);
 }
 
-void UReplicatedTestObjectBridge::SetExternalPreUpdateFunctor(TFunction<void(TArrayView<UObject*>, const UReplicationBridge*)> PreUpdateFunctor)
+void UReplicatedTestObjectBridge::SetExternalPreUpdateFunctor(TFunction<void(TArrayView<UObject*>, const UObjectReplicationBridge*)> PreUpdateFunctor)
 {
 	SetInstancePreUpdateFunction(PreUpdateFunctor);
 }
diff --git a/Engine/Plugins/Runtime/ReplicationSystemTestPlugin/Source/Private/Tests/ReplicationSystem/ReplicatedTestObject.h b/Engine/Plugins/Runtime/ReplicationSystemTestPlugin/Source/Private/Tests/ReplicationSystem/ReplicatedTestObject.h
--- a/Engine/Plugins/Runtime/ReplicationSystemTestPlugin/Source/Private/Tests/ReplicationSystem/ReplicatedTestObject.h
+++ b/Engine/Plugins/Runtime/ReplicationSystemTestPlugin/Source/Private/Tests/ReplicationSystem/ReplicatedTestObject.h
@@ -632,7 +632,7 @@ public:
 
 	void SetExternalWorldLocationUpdateFunctor(TFunction<void(FNetRefHandle NetHandle, const UObject* ReplicatedObject, FVector& OutLocation, float& OutCullDistance)> LocUpdateFunctor);
 
-	void SetExternalPreUpdateFunctor(TFunction<void(TArrayView<UObject*>, const UReplicationBridge*)> PreUpdateFunctor);
+	void SetExternalPreUpdateFunctor(TFunction<void(TArrayView<UObject*>, const UObjectReplicationBridge*)> PreUpdateFunctor);
 
 	float GetMaxTickRate() const { return Super::GetMaxTickRate(); }
 
diff --git a/Engine/Plugins/Runtime/ReplicationSystemTestPlugin/Source/Private/Tests/ReplicationSystem/TestDirtyNetObjectTracker.cpp b/Engine/Plugins/Runtime/ReplicationSystemTestPlugin/Source/Private/Tests/ReplicationSystem/TestDirtyNetObjectTracker.cpp
--- a/Engine/Plugins/Runtime/ReplicationSystemTestPlugin/Source/Private/Tests/ReplicationSystem/TestDirtyNetObjectTracker.cpp
+++ b/Engine/Plugins/Runtime/ReplicationSystemTestPlugin/Source/Private/Tests/ReplicationSystem/TestDirtyNetObjectTracker.cpp
@@ -233,7 +233,7 @@ UE_NET_TEST_FIXTURE(FReplicationSystemServerClientTestFixture, NetForceUpdateOth
 	// Force ObjectA to replicate
 	Server->ReplicationSystem->ForceNetUpdate(ServerObjectA->NetRefHandle);
 
-	auto PreUpdateObjectA = [&](TArrayView<UObject*> Instances, const UReplicationBridge* Bridge)
+	auto PreUpdateObjectA = [&](TArrayView<UObject*> Instances, const UObjectReplicationBridge* Bridge)
 	{
 		for (UObject* ReplicatedObject : Instances)
 		{
@@ -263,7 +263,7 @@ UE_NET_TEST_FIXTURE(FReplicationSystemServerClientTestFixture, NetForceUpdateOth
 	// Now the ForceNetUpdate on ObjectB is applied and it is replicated
 	UE_NET_ASSERT_EQ(ClientObjectB->IntA, ServerObjectB->IntA);
 
-	auto PreUpdateObjectB = [&](TArrayView<UObject*> Instances, const UReplicationBridge* Bridge)
+	auto PreUpdateObjectB = [&](TArrayView<UObject*> Instances, const UObjectReplicationBridge* Bridge)
 	{
 		for (UObject* ReplicatedObject : Instances)
 		{
@@ -328,7 +328,7 @@ UE_NET_TEST_FIXTURE(FReplicationSystemServerClientTestFixture, DirtyOtherObjectI
 	UE_NET_ASSERT_NE(ClientObjectB, nullptr);
 
 	UTestReplicatedIrisObject* ObjectToDirty = nullptr;
-	auto PreUpdateObject = [&](TArrayView<UObject*> Instances, const UReplicationBridge* Bridge)
+	auto PreUpdateObject = [&](TArrayView<UObject*> Instances, const UObjectReplicationBridge* Bridge)
 	{
 		for (UObject* ReplicatedObject : Instances)
 		{
@@ -382,7 +382,7 @@ UE_NET_TEST_FIXTURE(FReplicationSystemServerClientTestFixture, PushModelMarkSelf
 	UTestReplicatedIrisObject* ClientObject = Cast<UTestReplicatedIrisObject>(Client->GetReplicationBridge()->GetReplicatedObject(ServerObject->NetRefHandle));
 	UE_NET_ASSERT_NE(ClientObject, nullptr);
 
-	auto PreUpdateObject = [&](TArrayView<UObject*> Instances, const UReplicationBridge* Bridge)
+	auto PreUpdateObject = [&](TArrayView<UObject*> Instances, const UObjectReplicationBridge* Bridge)
 	{
 		for (UObject* InObject : Instances)
 		{
diff --git a/Engine/Plugins/Runtime/ReplicationSystemTestPlugin/Source/Private/Tests/ReplicationSystem/TestMultiReplicationSystems.cpp b/Engine/Plugins/Runtime/ReplicationSystemTestPlugin/Source/Private/Tests/ReplicationSystem/TestMultiReplicationSystems.cpp
--- a/Engine/Plugins/Runtime/ReplicationSystemTestPlugin/Source/Private/Tests/ReplicationSystem/TestMultiReplicationSystems.cpp
+++ b/Engine/Plugins/Runtime/ReplicationSystemTestPlugin/Source/Private/Tests/ReplicationSystem/TestMultiReplicationSystems.cpp
@@ -716,7 +716,7 @@ UE_NET_TEST_FIXTURE(FMultiReplicationSystemsTestFixture, GlobalDirtyTrackerInsid
 
 	auto PreUpdateObjectForServer = [&](int32 ServerIndex)
 		{
-			return [&ServerObjects, ServerIndex](TArrayView<UObject*> Instances, const UReplicationBridge* Bridge)
+			return [&ServerObjects, ServerIndex](TArrayView<UObject*> Instances, const UObjectReplicationBridge* Bridge)
 				{
 					for (UObject* InObject : Instances)
 					{
@@ -746,7 +746,7 @@ UE_NET_TEST_FIXTURE(FMultiReplicationSystemsTestFixture, GlobalDirtyTrackerInsid
 	}
 
 	// Now use the other server/repsystem to mark dirty
-	AllServers[0]->GetReplicationBridge()->SetExternalPreUpdateFunctor([](TArrayView<UObject*> Instances, const UReplicationBridge* Bridge) {});
+	AllServers[0]->GetReplicationBridge()->SetExternalPreUpdateFunctor([](TArrayView<UObject*> Instances, const UObjectReplicationBridge* Bridge) {});
 	AllServers[1]->GetReplicationBridge()->SetExternalPreUpdateFunctor(PreUpdateObjectForServer(1));
 
 	// Send and deliver packet
@@ -802,7 +802,7 @@ UE_NET_TEST_FIXTURE(FMultiReplicationSystemsTestFixture, GlobalDirtyTrackerInsid
 
 	auto PreUpdateObjectForServer = [&](int32 ServerIndex)
 		{
-			return [&ServerObjects, ServerIndex](TArrayView<UObject*> Instances, const UReplicationBridge* Bridge)
+			return [&ServerObjects, ServerIndex](TArrayView<UObject*> Instances, const UObjectReplicationBridge* Bridge)
 				{
 					// Dirty the object in the global dirty tracker
 					ServerObjects[ServerIndex]->ObjectReferenceComponents[0]->ModifyIntA();
@@ -826,7 +826,7 @@ UE_NET_TEST_FIXTURE(FMultiReplicationSystemsTestFixture, GlobalDirtyTrackerInsid
 	}
 
 	// This time on server 1's PreUpdate, mark server 0's object dirty.
-	AllServers[0]->GetReplicationBridge()->SetExternalPreUpdateFunctor([](TArrayView<UObject*> Instances, const UReplicationBridge* Bridge) {});
+	AllServers[0]->GetReplicationBridge()->SetExternalPreUpdateFunctor([](TArrayView<UObject*> Instances, const UObjectReplicationBridge* Bridge) {});
 	AllServers[1]->GetReplicationBridge()->SetExternalPreUpdateFunctor(PreUpdateObjectForServer(0));
 
 	// Send and deliver packet
@@ -843,8 +843,8 @@ UE_NET_TEST_FIXTURE(FMultiReplicationSystemsTestFixture, GlobalDirtyTrackerInsid
 	UE_NET_ASSERT_EQ(ClientObject1->ObjectReferenceComponents[0]->IntA, ServerObjects[1]->ObjectReferenceComponents[0]->IntA);
 
 	// Clear the PreUpdate functors
-	AllServers[0]->GetReplicationBridge()->SetExternalPreUpdateFunctor([](TArrayView<UObject*> Instances, const UReplicationBridge* Bridge) {});
-	AllServers[1]->GetReplicationBridge()->SetExternalPreUpdateFunctor([](TArrayView<UObject*> Instances, const UReplicationBridge* Bridge) {});
+	AllServers[0]->GetReplicationBridge()->SetExternalPreUpdateFunctor([](TArrayView<UObject*> Instances, const UObjectReplicationBridge* Bridge) {});
+	AllServers[1]->GetReplicationBridge()->SetExternalPreUpdateFunctor([](TArrayView<UObject*> Instances, const UObjectReplicationBridge* Bridge) {});
 
 	// Send and deliver packet. This should update client 0 to have the latest server value.
 	FullSendAndDeliverUpdateTwoPass();
@@ -868,8 +868,8 @@ UE_NET_TEST_FIXTURE(FMultiReplicationSystemsTestFixture, GlobalDirtyTrackerInsid
 	UE_NET_ASSERT_EQ(ClientObject1->ObjectReferenceComponents[0]->IntA, ServerObjects[1]->ObjectReferenceComponents[0]->IntA);
 
 	// Clear the PreUpdate functors
-	AllServers[0]->GetReplicationBridge()->SetExternalPreUpdateFunctor([](TArrayView<UObject*> Instances, const UReplicationBridge* Bridge) {});
-	AllServers[1]->GetReplicationBridge()->SetExternalPreUpdateFunctor([](TArrayView<UObject*> Instances, const UReplicationBridge* Bridge) {});
+	AllServers[0]->GetReplicationBridge()->SetExternalPreUpdateFunctor([](TArrayView<UObject*> Instances, const UObjectReplicationBridge* Bridge) {});
+	AllServers[1]->GetReplicationBridge()->SetExternalPreUpdateFunctor([](TArrayView<UObject*> Instances, const UObjectReplicationBridge* Bridge) {});
 
 	// Send and deliver packet. This should update client 0 to have the latest server value.
 	FullSendAndDeliverUpdateTwoPass();
@@ -992,7 +992,7 @@ UE_NET_TEST_FIXTURE(FMultiReplicationSystemsTestFixture, GlobalDirtyTrackerInsid
 		}
 	}
 
-	auto PreUpdateObject = [&](TArrayView<UObject*> Instances, const UReplicationBridge* Bridge)
+	auto PreUpdateObject = [&](TArrayView<UObject*> Instances, const UObjectReplicationBridge* Bridge)
 		{
 			for (UObject* InObject : Instances)
 			{
@@ -1021,7 +1021,7 @@ UE_NET_TEST_FIXTURE(FMultiReplicationSystemsTestFixture, GlobalDirtyTrackerInsid
 
 
 	// Now use the other server/repsystem to mark dirty
-	AllServers[0]->GetReplicationBridge()->SetExternalPreUpdateFunctor([](TArrayView<UObject*> Instances, const UReplicationBridge* Bridge) {});
+	AllServers[0]->GetReplicationBridge()->SetExternalPreUpdateFunctor([](TArrayView<UObject*> Instances, const UObjectReplicationBridge* Bridge) {});
 	AllServers[1]->GetReplicationBridge()->SetExternalPreUpdateFunctor(PreUpdateObject);
 
 	// Send and deliver packet
@@ -1034,8 +1034,8 @@ UE_NET_TEST_FIXTURE(FMultiReplicationSystemsTestFixture, GlobalDirtyTrackerInsid
 	UE_NET_ASSERT_EQ(ClientObject1->ObjectReferenceComponents[0]->IntA, ServerObject->ObjectReferenceComponents[0]->IntA);
 
 	// Clear the PreUpdate functors
-	AllServers[0]->GetReplicationBridge()->SetExternalPreUpdateFunctor([](TArrayView<UObject*> Instances, const UReplicationBridge* Bridge) {});
-	AllServers[1]->GetReplicationBridge()->SetExternalPreUpdateFunctor([](TArrayView<UObject*> Instances, const UReplicationBridge* Bridge) {});
+	AllServers[0]->GetReplicationBridge()->SetExternalPreUpdateFunctor([](TArrayView<UObject*> Instances, const UObjectReplicationBridge* Bridge) {});
+	AllServers[1]->GetReplicationBridge()->SetExternalPreUpdateFunctor([](TArrayView<UObject*> Instances, const UObjectReplicationBridge* Bridge) {});
 
 	// Send and deliver packet. This should update client 0 to have the server value.
 	FullSendAndDeliverUpdateTwoPass();
@@ -1060,8 +1060,8 @@ UE_NET_TEST_FIXTURE(FMultiReplicationSystemsTestFixture, GlobalDirtyTrackerInsid
 	UE_NET_ASSERT_EQ(ClientObject1->ObjectReferenceComponents[0]->IntA, ServerObject->ObjectReferenceComponents[0]->IntA);
 
 	// Clear the PreUpdate functors
-	AllServers[0]->GetReplicationBridge()->SetExternalPreUpdateFunctor([](TArrayView<UObject*> Instances, const UReplicationBridge* Bridge) {});
-	AllServers[1]->GetReplicationBridge()->SetExternalPreUpdateFunctor([](TArrayView<UObject*> Instances, const UReplicationBridge* Bridge) {});
+	AllServers[0]->GetReplicationBridge()->SetExternalPreUpdateFunctor([](TArrayView<UObject*> Instances, const UObjectReplicationBridge* Bridge) {});
+	AllServers[1]->GetReplicationBridge()->SetExternalPreUpdateFunctor([](TArrayView<UObject*> Instances, const UObjectReplicationBridge* Bridge) {});
 
 	// Send and deliver packet. This should update client 0 to have the latest server value.
 	FullSendAndDeliverUpdateTwoPass();
diff --git a/Engine/Plugins/Runtime/ReplicationSystemTestPlugin/Source/Private/Tests/ReplicationSystem/TestPreUpdateCallback.cpp b/Engine/Plugins/Runtime/ReplicationSystemTestPlugin/Source/Private/Tests/ReplicationSystem/TestPreUpdateCallback.cpp
--- a/Engine/Plugins/Runtime/ReplicationSystemTestPlugin/Source/Private/Tests/ReplicationSystem/TestPreUpdateCallback.cpp
+++ b/Engine/Plugins/Runtime/ReplicationSystemTestPlugin/Source/Private/Tests/ReplicationSystem/TestPreUpdateCallback.cpp
@@ -38,7 +38,7 @@ UE_NET_TEST_FIXTURE(FReplicationSystemServerClientTestFixture, CreateObjectsInsi
 	UE_NET_ASSERT_NE(ClientObject, nullptr);
 
 	// Add a PreUpdate where we create a new subobject
-	auto PreUpdateObject = [&](TArrayView<UObject*> Instances, const UReplicationBridge* Bridge)
+	auto PreUpdateObject = [&](TArrayView<UObject*> Instances, const UObjectReplicationBridge* Bridge)
 	{
 		for (UObject* ReplicatedObject : Instances)
 		{
@@ -103,7 +103,7 @@ UE_NET_TEST_FIXTURE(FReplicationSystemServerClientTestFixture, CreateObjectsInsi
 	ServerObject->SetIntA(0xBB);
 
 	// Add a PreUpdate where we create a new subobject
-	auto PreUpdateObject = [&](TArrayView<UObject*> Instances, const UReplicationBridge* Bridge)
+	auto PreUpdateObject = [&](TArrayView<UObject*> Instances, const UObjectReplicationBridge* Bridge)
 	{
 		for (UObject* ReplicatedObject : Instances)
 		{
diff --git a/Engine/Plugins/Runtime/ReplicationSystemTestPlugin/Source/Private/Tests/ReplicationSystem/TestReplicationConfig.cpp b/Engine/Plugins/Runtime/ReplicationSystemTestPlugin/Source/Private/Tests/ReplicationSystem/TestReplicationConfig.cpp
--- a/Engine/Plugins/Runtime/ReplicationSystemTestPlugin/Source/Private/Tests/ReplicationSystem/TestReplicationConfig.cpp
+++ b/Engine/Plugins/Runtime/ReplicationSystemTestPlugin/Source/Private/Tests/ReplicationSystem/TestReplicationConfig.cpp
@@ -235,7 +235,7 @@ UE_NET_TEST_FIXTURE(FReplicationConfigTestFixture, TestNetObjectListGrow_DuringP
 	UTestReplicatedIrisObject* ServerSubObject(nullptr);
 
 	// Add a PreUpdate callback where we create a new subobject
-	auto PreUpdateObject = [&](TArrayView<UObject*> Instances, const UReplicationBridge* Bridge)
+	auto PreUpdateObject = [&](TArrayView<UObject*> Instances, const UObjectReplicationBridge* Bridge)
 	{
 		for (UObject* ReplicatedObject : Instances)
 		{
diff --git a/Engine/Source/Runtime/Net/Iris/Public/Iris/ReplicationSystem/ObjectReplicationBridge.h b/Engine/Source/Runtime/Net/Iris/Public/Iris/ReplicationSystem/ObjectReplicationBridge.h
--- a/Engine/Source/Runtime/Net/Iris/Public/Iris/ReplicationSystem/ObjectReplicationBridge.h
+++ b/Engine/Source/Runtime/Net/Iris/Public/Iris/ReplicationSystem/ObjectReplicationBridge.h
@@ -325,7 +325,7 @@ protected:
 	IRISCORE_API UNetObjectFactory* GetNetFactory(UE::Net::FNetObjectFactoryId FactoryId) const;
 
 	/** Helper method that calls provided PreUpdateFunction and polls state data for all replicated instances with the NeedsPoll trait. */
-	using FInstancePreUpdateFunction = TFunction<void(TArrayView<UObject*>, const UReplicationBridge*)>;
+	using FInstancePreUpdateFunction = TFunction<void(TArrayView<UObject*>, const UObjectReplicationBridge*)>;
 
 	/** Set the function that we should call before copying state data. */
 	IRISCORE_API void SetInstancePreUpdateFunction(FInstancePreUpdateFunction InPreUpdateFunction);
