Iris - Minor refactor of WriteBool test to work with more compilers.

diff --git a/Engine/Plugins/Runtime/ReplicationSystemTestPlugin/Source/Private/Tests/Serialization/TestNetBitStreams.cpp b/Engine/Plugins/Runtime/ReplicationSystemTestPlugin/Source/Private/Tests/Serialization/TestNetBitStreams.cpp
--- a/Engine/Plugins/Runtime/ReplicationSystemTestPlugin/Source/Private/Tests/Serialization/TestNetBitStreams.cpp
+++ b/Engine/Plugins/Runtime/ReplicationSystemTestPlugin/Source/Private/Tests/Serialization/TestNetBitStreams.cpp
@@ -5,6 +5,7 @@
 #include "Iris/Serialization/NetBitStreamReader.h"
 #include "Iris/Serialization/NetBitStreamWriter.h"
 #include "Iris/Serialization/NetBitStreamUtil.h"
+#include "Templates/TypeCompatibleBytes.h"
 
 namespace UE::Net::Private
 {
@@ -903,52 +904,41 @@ UE_NET_TEST_FIXTURE(FNetBitStreamReaderWriterTest, WriteBool)
 	// We can't trust the compiler not to fiddle with bools and force them to be 0 or 1.
 	const uint8 TestValues[] = {0, 1, 128, 255};
 	constexpr SIZE_T ValueCount = sizeof(TestValues)/sizeof(TestValues[0]);
-	bool Values[ValueCount];
 
-	FPlatformMemory::Memcpy(Values, TestValues, sizeof(TestValues));
-
-	// Validate test values
-	{
-		const bool* TestValue0 = &Values[0];
-		const uint8 ExpectedValue0 = 0;
-		UE_NET_ASSERT_EQ(FPlatformMemory::Memcmp(TestValue0, &ExpectedValue0, 1), 0);
-
-		const bool* TestValue1 = &Values[1];
-		const uint8 ExpectedValue1 = 1;
-		UE_NET_ASSERT_EQ(FPlatformMemory::Memcmp(TestValue1, &ExpectedValue1, 1), 0);
-
-		const bool* TestValue2 = &Values[2];
-		const uint8 ExpectedValue2 = 128;
-		UE_NET_ASSERT_EQ(FPlatformMemory::Memcmp(TestValue2, &ExpectedValue2, 1), 0);
-
-		const bool* TestValue3 = &Values[3];
-		const uint8 ExpectedValue3 = 255;
-		UE_NET_ASSERT_EQ(FPlatformMemory::Memcmp(TestValue3, &ExpectedValue3, 1), 0);
-	}
+	bool ExpectedValues[ValueCount] = {};
+	bool ReturnValues[ValueCount] = {};
+	bool ReadValues[ValueCount] = {};
 
 	// We expect WriteBool to write a 1 if the bool isn't exactly false (zero) and write a 0 when it's false.
-	// We also expect WriteBool to return the value written.
-	bool ExpectedValues[sizeof(Values)/sizeof(Values[0])];
-	bool ReturnValues[sizeof(Values)/sizeof(Values[0])];
-	bool ReadValues[sizeof(Values)/sizeof(Values[0])];
+	// The most important thing is it must succeed in returning what it wrote.
+	// Fill in expected values for standard true and false. These must succeed.
 	for (SIZE_T ValueIt = 0; ValueIt < ValueCount; ++ValueIt)
 	{
-		uint8 Value;
-		FPlatformMemory::Memcpy(&Value, &Values[ValueIt], 1);
-		ExpectedValues[ValueIt] = Value != 0;
+		if (TestValues[ValueIt] <= 1)
+		{
+			ExpectedValues[ValueIt] = (TestValues[ValueIt] != 0);
+		}
 	}
 
 	Writer.InitBytes(Buffer.GetBuffer(), Buffer.GetBufferCapacity());
-	for (const bool& Value : Values)
+	for (const uint8& Value : TestValues)
 	{
-		ReturnValues[&Value - Values] = Writer.WriteBool(Value);
+		const bool bWroteTrue = Writer.WriteBool(BitCast<bool>(Value));
+		ReturnValues[&Value - TestValues] = bWroteTrue;
+		if (Value > 1)
+		{
+			ExpectedValues[&Value - TestValues] = bWroteTrue;
+			if (!bWroteTrue)
+			{
+				UE_NET_LOG(FString::Printf(TEXT("Expected WriteBool of non-zero value (%u) to write true."), Value));
+			}
+		}
 	}
 	Writer.CommitWrites();
 
 	// Validate return value is the written value
 	UE_NET_ASSERT_EQ(FPlatformMemory::Memcmp(ReturnValues, ExpectedValues, sizeof(ReturnValues)), 0);
 
-
 	Reader.InitBits(Buffer.GetBuffer(), Writer.GetPosBits());
 	for (SIZE_T ValueIt = 0; ValueIt < ValueCount; ++ValueIt)
 	{
