Iris - Test RPC stomp detection with AutoRTFM enabled

#rnx

diff --git a/Engine/Plugins/Runtime/ReplicationSystemTestPlugin/Source/Private/Tests/ReplicationSystem/RPC/TestObjectRPC.cpp b/Engine/Plugins/Runtime/ReplicationSystemTestPlugin/Source/Private/Tests/ReplicationSystem/RPC/TestObjectRPC.cpp
--- a/Engine/Plugins/Runtime/ReplicationSystemTestPlugin/Source/Private/Tests/ReplicationSystem/RPC/TestObjectRPC.cpp
+++ b/Engine/Plugins/Runtime/ReplicationSystemTestPlugin/Source/Private/Tests/ReplicationSystem/RPC/TestObjectRPC.cpp
@@ -11,6 +11,7 @@
 #include "Iris/ReplicationSystem/ReplicationSystemInternal.h"
 #include "Tests/EnsureScope.h"
 #include "Tests/CheckScope.h"
+#include "AutoRTFM.h"
 
 namespace UE::Net::Private
 {
@@ -429,7 +430,7 @@ namespace UE::Net::Private
 	// This test is written to test that we can control if ordered unreliable attachments are pruned or not after each tick.
 	UE_NET_TEST_FIXTURE(FRPCTestFixture, TestOvercommitOrderedUnreliableRPC)
 	{
-		// This tests tests the behavior with and without DropUnsentOrderedUnreliableAttachmentAtEndOfTick enabled
+		// This tests the behavior with and without DropUnsentOrderedUnreliableAttachmentAtEndOfTick enabled
 		IConsoleVariable* CVarOverride = IConsoleManager::Get().FindConsoleVariable(TEXT("net.Iris.ReplicationWriter.MaxUnsentOrderedUnreliableAttachmentAtEndOfTick"));
 		check(CVarOverride != nullptr && CVarOverride->IsVariableInt());
 		const int32 OldCVarValue = CVarOverride->GetInt();
@@ -540,4 +541,55 @@ namespace UE::Net::Private
 			UE_NET_ASSERT_EQ(ServerObject->ServerUnreliableWithExportRPCCallCounter, ExpectedPerPacket);
 		}
 	}
+
+#if UE_AUTORTFM
+	UE_NET_TEST_FIXTURE(FRPCTestFixture, AbortingRPCWithStompDectionDoesNotCrash)
+	{
+		IConsoleVariable* CVarStompDetectionEnabled = IConsoleManager::Get().FindConsoleVariable(TEXT("net.Iris.RPC.ServerStompDetectionEnabled"));
+		UE_NET_ASSERT_NE(CVarStompDetectionEnabled, nullptr);
+		UE_NET_ASSERT_TRUE(CVarStompDetectionEnabled->IsVariableBool());
+		const bool OldCVarValue = CVarStompDetectionEnabled->GetBool();
+
+		// Enable RPC quantized state stomp detection
+		CVarStompDetectionEnabled->Set(true, ECVF_SetByCode);
+
+		ON_SCOPE_EXIT
+		{
+			// Restore cvar
+			CVarStompDetectionEnabled->Set(OldCVarValue, ECVF_SetByCode);
+		};
+
+		// Add a client
+		FReplicationSystemTestClient* Client = CreateClient();
+
+		// Spawn object on server
+		UTestReplicatedObjectWithRPC* ServerObject = Server->CreateObject<UTestReplicatedObjectWithRPC>();
+
+		ServerObject->bIsServerObject = true;
+		ServerObject->ReplicationSystem = Server->GetReplicationSystem();
+		Server->ReplicationSystem->SetOwningNetConnection(ServerObject->NetRefHandle, 0x01);
+
+		// Send and deliver packet
+		Server->UpdateAndSend({Client});
+
+		UTestReplicatedObjectWithRPC* ClientObject = Cast<UTestReplicatedObjectWithRPC>(Client->GetReplicationBridge()->GetReplicatedObject(ServerObject->NetRefHandle));
+		
+		// Verify that created server handle now also exists on client
+		UE_NET_ASSERT_NE(ClientObject, nullptr);
+
+		auto Transaction = AutoRTFM::Transact([&]()
+			{
+				// Call an RPC server->client
+				ServerObject->ClientRPCWithParam(0);
+
+				AutoRTFM::AbortTransaction();
+			});
+
+		// Send and deliver packet
+		Server->UpdateAndSend({Client});
+
+		// The RPC was aborted and should not have been called
+		UE_NET_ASSERT_FALSE(ClientObject->bClientRPCCalled);
+	}
+#endif
 }
\ No newline at end of file
