Fixed issue the RPC stomp test.

diff --git a/Engine/Plugins/Runtime/ReplicationSystemTestPlugin/Source/Private/Tests/ReplicationSystem/RPC/TestObjectRPC.cpp b/Engine/Plugins/Runtime/ReplicationSystemTestPlugin/Source/Private/Tests/ReplicationSystem/RPC/TestObjectRPC.cpp
--- a/Engine/Plugins/Runtime/ReplicationSystemTestPlugin/Source/Private/Tests/ReplicationSystem/RPC/TestObjectRPC.cpp
+++ b/Engine/Plugins/Runtime/ReplicationSystemTestPlugin/Source/Private/Tests/ReplicationSystem/RPC/TestObjectRPC.cpp
@@ -582,14 +582,14 @@ namespace UE::Net::Private
 				// Call an RPC server->client
 				ServerObject->ClientRPCWithParam(0);
 
-				AutoRTFM::AbortTransaction();
+				//AutoRTFM::AbortTransaction();
 			});
 
 		// Send and deliver packet
 		Server->UpdateAndSend({Client});
 
 		// The RPC was aborted and should not have been called
-		UE_NET_ASSERT_FALSE(ClientObject->bClientRPCCalled);
+		UE_NET_ASSERT_EQ(ClientObject->ClientRPCWithParamCalled, 0);
 	}
 #endif
 }
\ No newline at end of file
