Fix issue with RPC stomp test.

diff --git a/Engine/Plugins/Runtime/ReplicationSystemTestPlugin/Source/Private/Tests/ReplicationSystem/RPC/TestObjectRPC.cpp b/Engine/Plugins/Runtime/ReplicationSystemTestPlugin/Source/Private/Tests/ReplicationSystem/RPC/TestObjectRPC.cpp
--- a/Engine/Plugins/Runtime/ReplicationSystemTestPlugin/Source/Private/Tests/ReplicationSystem/RPC/TestObjectRPC.cpp
+++ b/Engine/Plugins/Runtime/ReplicationSystemTestPlugin/Source/Private/Tests/ReplicationSystem/RPC/TestObjectRPC.cpp
@@ -580,9 +580,9 @@ namespace UE::Net::Private
 		auto Transaction = AutoRTFM::Transact([&]()
 			{
 				// Call an RPC server->client
-				ServerObject->ClientRPCWithParam(0);
+				ServerObject->ClientRPCWithParam(1);
 
-				//AutoRTFM::AbortTransaction();
+				AutoRTFM::AbortTransaction();
 			});
 
 		// Send and deliver packet
