Iris - Fixed compile error

#rnx

diff --git a/Engine/Plugins/Runtime/ReplicationSystemTestPlugin/Source/Private/Tests/Serialization/TestRemoteObjectReferenceNetSerializer.cpp b/Engine/Plugins/Runtime/ReplicationSystemTestPlugin/Source/Private/Tests/Serialization/TestRemoteObjectReferenceNetSerializer.cpp
--- a/Engine/Plugins/Runtime/ReplicationSystemTestPlugin/Source/Private/Tests/Serialization/TestRemoteObjectReferenceNetSerializer.cpp
+++ b/Engine/Plugins/Runtime/ReplicationSystemTestPlugin/Source/Private/Tests/Serialization/TestRemoteObjectReferenceNetSerializer.cpp
@@ -21,7 +21,7 @@ void UTestReplicatedObjectWithRemoteReference::RemoteRPCWithRemoteReferenceParam
 }
 
 // These tests require remote reference support to be compiled in
-#if UE_WITH_REMOTE_OBJECT_REFERENCE
+#if UE_WITH_REMOTE_OBJECT_HANDLE
 
 namespace UE::Net::Private
 {
