Iris
- Fixed bug with FastArraySerialzierFragment if using not replicating items

diff --git a/Engine/Plugins/Runtime/ReplicationSystemTestPlugin/Source/Private/Tests/ReplicationState/TestFastArrayReplicationState.cpp b/Engine/Plugins/Runtime/ReplicationSystemTestPlugin/Source/Private/Tests/ReplicationState/TestFastArrayReplicationState.cpp
--- a/Engine/Plugins/Runtime/ReplicationSystemTestPlugin/Source/Private/Tests/ReplicationState/TestFastArrayReplicationState.cpp
+++ b/Engine/Plugins/Runtime/ReplicationSystemTestPlugin/Source/Private/Tests/ReplicationState/TestFastArrayReplicationState.cpp
@@ -201,9 +201,7 @@ UE_NET_TEST_FIXTURE(FReplicationSystemServerClientTestFixture, TestAddItem)
 	FReplicationSystemTestClient* Client = CreateClient();
 
 	// Send and deliver packet
-	Server->NetUpdate();
-	Server->SendAndDeliverTo(Client, true);
-	Server->PostSendUpdate();
+	Server->UpdateAndSend({Client});
 
 	// Store Pointer to objects
 	UTestFastArrayReplicationState_FastArray_TestClassFastArray* ClientObject = Cast<UTestFastArrayReplicationState_FastArray_TestClassFastArray>(Client->GetReplicationBridge()->GetReplicatedObject(ServerObject->NetRefHandle));
@@ -230,9 +228,7 @@ UE_NET_TEST_FIXTURE(FReplicationSystemServerClientTestFixture, TestAddItem)
 	ServerObject->FastArray.Edit().Add(NewEntry);
 
 	// Send and deliver packet
-	Server->NetUpdate();
-	Server->SendAndDeliverTo(Client, true);
-	Server->PostSendUpdate();
+	Server->UpdateAndSend({Client});
 
 	// Verify that we received the data
 	UE_NET_ASSERT_EQ(ServerObject->FastArray.GetItemArray().Num(), ClientObject->FastArray.GetItemArray().Num());
@@ -253,9 +249,7 @@ UE_NET_TEST_FIXTURE(FReplicationSystemServerClientTestFixture, TestRemoveItem)
 	FReplicationSystemTestClient* Client = CreateClient();
 
 	// Send and deliver packet
-	Server->NetUpdate();
-	Server->SendAndDeliverTo(Client, true);
-	Server->PostSendUpdate();
+	Server->UpdateAndSend({Client});
 
 	// Store Pointer to objects
 	UTestFastArrayReplicationState_FastArray_TestClassFastArray* ClientObject = Cast<UTestFastArrayReplicationState_FastArray_TestClassFastArray>(Client->GetReplicationBridge()->GetReplicatedObject(ServerObject->NetRefHandle));
@@ -282,9 +276,7 @@ UE_NET_TEST_FIXTURE(FReplicationSystemServerClientTestFixture, TestRemoveItem)
 	ServerObject->FastArray.Edit().Add(NewEntry);
 
 	// Send and deliver packet
-	Server->NetUpdate();
-	Server->SendAndDeliverTo(Client, true);
-	Server->PostSendUpdate();
+	Server->UpdateAndSend({Client});
 
 	// Verify that we received the data
 	UE_NET_ASSERT_EQ(ServerObject->FastArray.GetItemArray().Num(), ClientObject->FastArray.GetItemArray().Num());
@@ -300,9 +292,7 @@ UE_NET_TEST_FIXTURE(FReplicationSystemServerClientTestFixture, TestRemoveItem)
 	ServerObject->FastArray.Edit().Remove(0);
 
 	// Send and deliver packet
-	Server->NetUpdate();
-	Server->SendAndDeliverTo(Client, true);
-	Server->PostSendUpdate();
+	Server->UpdateAndSend({Client});
 
 	// Verify that we received the data
 	UE_NET_ASSERT_TRUE(ClientObject!= nullptr);
@@ -321,9 +311,7 @@ UE_NET_TEST_FIXTURE(FReplicationSystemServerClientTestFixture, TestLocalItemAreP
 	FReplicationSystemTestClient* Client = CreateClient();
 
 	// Send and deliver packet
-	Server->NetUpdate();
-	Server->SendAndDeliverTo(Client, true);
-	Server->PostSendUpdate();
+	Server->UpdateAndSend({Client});
 
 	// Store Pointer to objects
 	UTestFastArrayReplicationState_FastArray_TestClassFastArray* ClientObject = Cast<UTestFastArrayReplicationState_FastArray_TestClassFastArray>(Client->GetReplicationBridge()->GetReplicatedObject(ServerObject->NetRefHandle));
@@ -359,9 +347,7 @@ UE_NET_TEST_FIXTURE(FReplicationSystemServerClientTestFixture, TestLocalItemAreP
 	ServerObject->FastArray.Edit().Add(NewEntry);
 
 	// Send and deliver packet
-	Server->NetUpdate();
-	Server->SendAndDeliverTo(Client, true);
-	Server->PostSendUpdate();
+	Server->UpdateAndSend({Client});
 
 	// Verify that we received the data
 	UE_NET_ASSERT_EQ(3, ClientObject->FastArray.GetItemArray().Num());
@@ -379,9 +365,7 @@ UE_NET_TEST_FIXTURE(FReplicationSystemServerClientTestFixture, TestLocalItemAreP
 	ServerObject->FastArray.Edit().Remove(0);
 
 	// Send and deliver packet
-	Server->NetUpdate();
-	Server->SendAndDeliverTo(Client, true);
-	Server->PostSendUpdate();
+	Server->UpdateAndSend({Client});
 
 	// Verify that we received the remove
 	UE_NET_ASSERT_EQ(2, ClientObject->FastArray.GetItemArray().Num());
@@ -404,9 +388,7 @@ UE_NET_TEST_FIXTURE(FReplicationSystemServerClientTestFixture, TestModifyItem)
 	FReplicationSystemTestClient* Client = CreateClient();
 
 	// Send and deliver packet
-	Server->NetUpdate();
-	Server->SendAndDeliverTo(Client, true);
-	Server->PostSendUpdate();
+	Server->UpdateAndSend({Client});
 
 	// Store Pointer to objects
 	UTestFastArrayReplicationState_FastArray_TestClassFastArray* ClientObject = Cast<UTestFastArrayReplicationState_FastArray_TestClassFastArray>(Client->GetReplicationBridge()->GetReplicatedObject(ServerObject->NetRefHandle));
@@ -433,9 +415,7 @@ UE_NET_TEST_FIXTURE(FReplicationSystemServerClientTestFixture, TestModifyItem)
 	ServerObject->FastArray.Edit().Add(NewEntry);
 
 	// Send and deliver packet
-	Server->NetUpdate();
-	Server->SendAndDeliverTo(Client, true);
-	Server->PostSendUpdate();
+	Server->UpdateAndSend({Client});
 
 	// Verify that we received the data
 	UE_NET_ASSERT_EQ(ServerObject->FastArray.GetItemArray().Num(), ClientObject->FastArray.GetItemArray().Num());
@@ -455,9 +435,7 @@ UE_NET_TEST_FIXTURE(FReplicationSystemServerClientTestFixture, TestModifyItem)
 	ServerObject->FastArray.Edit()[0].RepInt32 = 3U;
 	
 	// Send and deliver packet
-	Server->NetUpdate();
-	Server->SendAndDeliverTo(Client, true);
-	Server->PostSendUpdate();
+	Server->UpdateAndSend({Client});
 
 	UE_NET_ASSERT_FALSE(ClientObject->FastArray.bHitReplicatedAdd);
 	UE_NET_ASSERT_TRUE(ClientObject->FastArray.bHitReplicatedChange);
@@ -477,9 +455,7 @@ UE_NET_TEST_FIXTURE(FReplicationSystemServerClientTestFixture, TestInsertItem)
 	FReplicationSystemTestClient* Client = CreateClient();
 
 	// Send and deliver packet
-	Server->NetUpdate();
-	Server->SendAndDeliverTo(Client, true);
-	Server->PostSendUpdate();
+	Server->UpdateAndSend({Client});
 
 	// Store Pointer to objects
 	UTestFastArrayReplicationState_FastArray_TestClassFastArray* ClientObject = Cast<UTestFastArrayReplicationState_FastArray_TestClassFastArray>(Client->GetReplicationBridge()->GetReplicatedObject(ServerObject->NetRefHandle));
@@ -514,9 +490,7 @@ UE_NET_TEST_FIXTURE(FReplicationSystemServerClientTestFixture, TestInsertItem)
 	ServerObject->FastArray.Edit()[1].RepInt32 = 13U;
 
 	// Send and deliver packet
-	Server->NetUpdate();
-	Server->SendAndDeliverTo(Client, true);
-	Server->PostSendUpdate();
+	Server->UpdateAndSend({Client});
 
 	// Verify that we received the data
 	UE_NET_ASSERT_EQ(ServerObject->FastArray.GetItemArray().Num(), ClientObject->FastArray.GetItemArray().Num());
@@ -536,9 +510,7 @@ UE_NET_TEST_FIXTURE(FReplicationSystemServerClientTestFixture, TestInsertItem)
 	ServerObject->FastArray.Edit()[0].RepInt32 = 3U;
 	
 	// Send and deliver packet
-	Server->NetUpdate();
-	Server->SendAndDeliverTo(Client, true);
-	Server->PostSendUpdate();
+	Server->UpdateAndSend({Client});
 
 	UE_NET_ASSERT_FALSE(ClientObject->FastArray.bHitReplicatedAdd);
 	UE_NET_ASSERT_TRUE(ClientObject->FastArray.bHitReplicatedChange);
@@ -559,9 +531,7 @@ UE_NET_TEST_FIXTURE(FReplicationSystemServerClientTestFixture, TestModifyItemDoe
 	FReplicationSystemTestClient* Client = CreateClient();
 
 	// Send and deliver packet
-	Server->NetUpdate();
-	Server->SendAndDeliverTo(Client, true);
-	Server->PostSendUpdate();
+	Server->UpdateAndSend({Client});
 
 	// Store Pointer to objects
 	UTestFastArrayReplicationState_FastArray_TestClassFastArray* ClientObject = Cast<UTestFastArrayReplicationState_FastArray_TestClassFastArray>(Client->GetReplicationBridge()->GetReplicatedObject(ServerObject->NetRefHandle));
@@ -591,9 +561,7 @@ UE_NET_TEST_FIXTURE(FReplicationSystemServerClientTestFixture, TestModifyItemDoe
 	FastArrayEditor.Add(NewEntry);
 
 	// Send and deliver packet
-	Server->NetUpdate();
-	Server->SendAndDeliverTo(Client, true);
-	Server->PostSendUpdate();
+	Server->UpdateAndSend({Client});
 
 	// Verify that we received the data
 	UE_NET_ASSERT_EQ(FastArrayEditor.Num(), ClientObject->FastArray.GetItemArray().Num());
@@ -617,9 +585,7 @@ UE_NET_TEST_FIXTURE(FReplicationSystemServerClientTestFixture, TestModifyItemDoe
 	FastArrayEditor.Edit(0).RepInt32 = 3;
 	
 	// Send and deliver packet
-	Server->NetUpdate();
-	Server->SendAndDeliverTo(Client, true);
-	Server->PostSendUpdate();
+	Server->UpdateAndSend({Client});
 
 	UE_NET_ASSERT_FALSE(ClientObject->FastArray.bHitReplicatedAdd);
 	UE_NET_ASSERT_TRUE(ClientObject->FastArray.bHitReplicatedChange);
@@ -685,9 +651,7 @@ UE_NET_TEST_FIXTURE(FReplicationSystemServerClientTestFixture, TestChangeMaskWra
 	FReplicationSystemTestClient* Client = CreateClient();
 
 	// Send and deliver packet
-	Server->NetUpdate();
-	Server->SendAndDeliverTo(Client, true);
-	Server->PostSendUpdate();
+	Server->UpdateAndSend({Client});
 
 	// Store Pointer to objects
 	UTestFastArrayReplicationState_FastArray_TestClassFastArray* ClientObject = Cast<UTestFastArrayReplicationState_FastArray_TestClassFastArray>(Client->GetReplicationBridge()->GetReplicatedObject(ServerObject->NetRefHandle));
@@ -717,9 +681,7 @@ UE_NET_TEST_FIXTURE(FReplicationSystemServerClientTestFixture, TestChangeMaskWra
 	int32 CurrentItemCount = FastArrayEditor.Num();
 
 	// Send and deliver packet
-	Server->NetUpdate();
-	Server->SendAndDeliverTo(Client, true);
-	Server->PostSendUpdate();
+	Server->UpdateAndSend({Client});
 
 	// Verify that we received the data
 	UE_NET_ASSERT_EQ(CurrentItemCount, ClientObject->FastArray.GetItemArray().Num());
@@ -745,9 +707,7 @@ UE_NET_TEST_FIXTURE(FReplicationSystemServerClientTestFixture, TestChangeMaskWra
 	}
 	
 	// Send and deliver packet
-	Server->NetUpdate();
-	Server->SendAndDeliverTo(Client, true);
-	Server->PostSendUpdate();
+	Server->UpdateAndSend({Client});
 
 	UE_NET_ASSERT_EQ(ServerObject->FastArray.GetItemArray()[CurrentItemCount - 1].RepInt32, ClientObject->FastArray.GetItemArray()[CurrentItemCount - 1].RepInt32);
 
@@ -764,9 +724,7 @@ UE_NET_TEST_FIXTURE(FReplicationSystemServerClientTestFixture, TestChangeMaskWra
 	}
 
 	// Send and deliver packet
-	Server->NetUpdate();
-	Server->SendAndDeliverTo(Client, true);
-	Server->PostSendUpdate();
+	Server->UpdateAndSend({Client});
 
 	// We should both have a change and an add
 	UE_NET_ASSERT_TRUE(ClientObject->FastArray.bHitReplicatedAdd);
@@ -807,10 +765,8 @@ UE_NET_TEST_FIXTURE(FReplicationSystemServerClientTestFixture, TestPostReplicate
 	Server->ReplicationSystem->AddToGroup(Server->GetReplicationSystem()->GetNotReplicatedNetObjectGroup(), ServerReferencedObjectA->NetRefHandle);
 	Server->ReplicationSystem->AddToGroup(Server->GetReplicationSystem()->GetNotReplicatedNetObjectGroup(), ServerReferencedObjectB->NetRefHandle);
 	
-	// Replicate
-	Server->NetUpdate();
-	Server->SendAndDeliverTo(Client, true);
-	Server->PostSendUpdate();
+	// Send and deliver packet
+	Server->UpdateAndSend({Client});
 
 	// Verify that object has been spawned on client
 	auto ClientObject = Client->GetObjectAs<UTestFastArrayReplicationState_FastArray_TestClassFastArray>(ServerObject->NetRefHandle);
@@ -837,10 +793,8 @@ UE_NET_TEST_FIXTURE(FReplicationSystemServerClientTestFixture, TestPostReplicate
 	// Enable replication for ServerReferenceObjectA
 	Server->ReplicationSystem->RemoveFromGroup(Server->GetReplicationSystem()->GetNotReplicatedNetObjectGroup(), ServerReferencedObjectA->NetRefHandle);
 
-	// Replicate
-	Server->NetUpdate();
-	Server->SendAndDeliverTo(Client, true);
-	Server->PostSendUpdate();
+	// Send and deliver packet
+	Server->UpdateAndSend({Client});
 
 	// Verify that we managed to resolve objectA even though we have other unresolved objects in the array
 	UE_NET_ASSERT_TRUE(ClientFastArray.GetItemArray()[0].ObjectRef != nullptr);
@@ -856,10 +810,8 @@ UE_NET_TEST_FIXTURE(FReplicationSystemServerClientTestFixture, TestPostReplicate
 	// Enable replication for ServerReferenceObjectB 
 	Server->ReplicationSystem->RemoveFromGroup(Server->GetReplicationSystem()->GetNotReplicatedNetObjectGroup(), ServerReferencedObjectB->NetRefHandle);
 
-	// Replicate
-	Server->NetUpdate();
-	Server->SendAndDeliverTo(Client, true);
-	Server->PostSendUpdate();
+	// Send and deliver packet
+	Server->UpdateAndSend({Client});
 
 	// Verify that we managed to resolve both objects and that the 
 	UE_NET_ASSERT_TRUE(ClientFastArray.GetItemArray()[0].ObjectRef != nullptr);
@@ -877,11 +829,8 @@ UE_NET_TEST_FIXTURE(FReplicationSystemServerClientTestFixture, TestFastArrayWith
 
 	UTestFastArrayReplicationState_FastArray_TestClassFastArrayWithExtraProperty* ServerObject = Server->CreateObject<UTestFastArrayReplicationState_FastArray_TestClassFastArrayWithExtraProperty>();
 
-
 	// Send and deliver packet
-	Server->NetUpdate();
-	Server->SendAndDeliverTo(Client, true);
-	Server->PostSendUpdate();
+	Server->UpdateAndSend({Client});
 
 	// Store Pointer to objects
 	UTestFastArrayReplicationState_FastArray_TestClassFastArrayWithExtraProperty* ClientObject = Cast<UTestFastArrayReplicationState_FastArray_TestClassFastArrayWithExtraProperty>(Client->GetReplicationBridge()->GetReplicatedObject(ServerObject->NetRefHandle));
@@ -909,9 +858,7 @@ UE_NET_TEST_FIXTURE(FReplicationSystemServerClientTestFixture, TestFastArrayWith
 	ServerObject->FastArray.Edit().Add(NewEntry);
 
 	// Send and deliver packet
-	Server->NetUpdate();
-	Server->SendAndDeliverTo(Client, true);
-	Server->PostSendUpdate();
+	Server->UpdateAndSend({Client});
 
 	// Verify that we received the data
 	UE_NET_ASSERT_EQ(ServerObject->FastArray.GetItemArray().Num(), ClientObject->FastArray.GetItemArray().Num());
@@ -928,9 +875,7 @@ UE_NET_TEST_FIXTURE(FReplicationSystemServerClientTestFixture, TestFastArrayWith
 	ServerObject->FastArray.ExtraInt = 3;
 
 	// Send and deliver packet
-	Server->NetUpdate();
-	Server->SendAndDeliverTo(Client, true);
-	Server->PostSendUpdate();
+	Server->UpdateAndSend({Client});
 
 	// Verify that we received the data
 	UE_NET_ASSERT_EQ(ServerObject->FastArray.ExtraInt, ClientObject->FastArray.ExtraInt);
@@ -944,9 +889,7 @@ UE_NET_TEST_FIXTURE(FReplicationSystemServerClientTestFixture, TestLostDataDirti
 	FReplicationSystemTestClient* Client = CreateClient();
 
 	// Send and deliver packet
-	Server->NetUpdate();
-	Server->SendAndDeliverTo(Client, true);
-	Server->PostSendUpdate();
+	Server->UpdateAndSend({Client});
 
 	// Store Pointer to objects
 	UTestFastArrayReplicationState_FastArray_TestClassFastArray* ClientObject = Cast<UTestFastArrayReplicationState_FastArray_TestClassFastArray>(Client->GetReplicationBridge()->GetReplicatedObject(ServerObject->NetRefHandle));
@@ -980,9 +923,7 @@ UE_NET_TEST_FIXTURE(FReplicationSystemServerClientTestFixture, TestLostDataDirti
 	Server->DeliverTo(Client, true);
 
 	// Send and deliver packet
-	Server->NetUpdate();
-	Server->SendAndDeliverTo(Client, true);
-	Server->PostSendUpdate();
+	Server->UpdateAndSend({Client});
 
 	// Verify that we received the expected data, for both elements
 	UE_NET_ASSERT_EQ(ServerObject->FastArray.GetItemArray().Num(), ClientObject->FastArray.GetItemArray().Num());
@@ -990,5 +931,223 @@ UE_NET_TEST_FIXTURE(FReplicationSystemServerClientTestFixture, TestLostDataDirti
 	UE_NET_ASSERT_EQ(ServerObject->FastArray.GetItemArray()[1].RepInt32, ClientObject->FastArray.GetItemArray()[1].RepInt32);
 }
 
+UE_NET_TEST_FIXTURE(FReplicationSystemServerClientTestFixture, TestNativeFastArrayWithNotReplicatedItems)
+{
+	IConsoleVariable* CVarToOverride = IConsoleManager::Get().FindConsoleVariable(TEXT("net.Iris.UseNativeFastArray"));
+	UE_NET_ASSERT_NE(CVarToOverride, nullptr);
+	UE_NET_ASSERT_TRUE(CVarToOverride->IsVariableBool());
+	const bool bPrevUseNativeFastArray = CVarToOverride->GetBool();
+	ON_SCOPE_EXIT
+	{
+		CVarToOverride->Set(bPrevUseNativeFastArray, ECVF_SetByCode);
+	};
+
+	// Use native fastarray
+	CVarToOverride->Set(true, ECVF_SetByCode);
+
+	UTestFastArrayReplicationState_FastArray_TestClassFastArray* ServerObject = Server->CreateObject<UTestFastArrayReplicationState_FastArray_TestClassFastArray>();
+
+	// Add a client
+	FReplicationSystemTestClient* Client = CreateClient();
+
+	// Send and deliver packet
+	Server->UpdateAndSend({Client});
+
+	// Store Pointer to objects
+	UTestFastArrayReplicationState_FastArray_TestClassFastArray* ClientObject = Cast<UTestFastArrayReplicationState_FastArray_TestClassFastArray>(Client->GetReplicationBridge()->GetReplicatedObject(ServerObject->NetRefHandle));
+
+	// Verify initial replication
+	UE_NET_ASSERT_TRUE(ClientObject!= nullptr);
+
+	// Add Replicated item
+	{
+		FTestFastArrayReplicationState_FastArrayItem ReplicatedEntry;
+		ReplicatedEntry.RepInt32 = 0U;
+		ServerObject->FastArray.Edit().Add(ReplicatedEntry);
+	}
+
+	// Add local explicitly not Replicated item
+	{
+		FTestFastArrayReplicationState_FastArrayItem ReplicatedEntry;
+		ReplicatedEntry.bSkipWriteFastArrayItem = true;
+		ReplicatedEntry.RepInt32 = 1000U;
+		ServerObject->FastArray.Edit().AddLocal(ReplicatedEntry);
+	}
+
+	// Add Replicated item
+	{
+		FTestFastArrayReplicationState_FastArrayItem ReplicatedEntry;
+		ReplicatedEntry.RepInt32 = 1U;
+		ServerObject->FastArray.Edit().Add(ReplicatedEntry);
+	}
+
+	// Add local explicitly not Replicated item
+	{
+		FTestFastArrayReplicationState_FastArrayItem ReplicatedEntry;
+		ReplicatedEntry.bSkipWriteFastArrayItem = true;
+		ReplicatedEntry.RepInt32 = 1000U;
+		ServerObject->FastArray.Edit().AddLocal(ReplicatedEntry);
+	}
+
+	{
+		// We do expect an ensure if we are using the Native IrisFastArraySerializer
+		FTestEnsureScope SuppressEnsureScope;
+
+		// Send and deliver packet
+		Server->UpdateAndSend({Client});
+	}
+
+	// Verify status
+	UE_NET_ASSERT_EQ(ClientObject->FastArray.GetItemArray().Num(), 2);
+	UE_NET_ASSERT_EQ(ServerObject->FastArray.GetItemArray()[0].RepInt32, ClientObject->FastArray.GetItemArray()[0].RepInt32);
+	UE_NET_ASSERT_EQ(ServerObject->FastArray.GetItemArray()[2].RepInt32, ClientObject->FastArray.GetItemArray()[1].RepInt32);
+
+	// Remove not replicated item at index 1
+	ServerObject->FastArray.Edit().Remove(1);
+	
+	// Remove replicated item at index 0
+	ServerObject->FastArray.Edit().Remove(0);
+
+	{
+		// We do expect an ensure if we are using the Native IrisFastArraySerializer
+		FTestEnsureScope SuppressEnsureScope;
+
+		// Send and deliver packet
+		Server->UpdateAndSend({Client});
+	}
+
+	// Verify status
+	UE_NET_ASSERT_EQ(ClientObject->FastArray.GetItemArray().Num(), 1);
+	UE_NET_ASSERT_EQ(ServerObject->FastArray.GetItemArray()[0].RepInt32, ClientObject->FastArray.GetItemArray()[0].RepInt32);
+
+	// Remove local item (now index 1)
+	ServerObject->FastArray.Edit().Remove(1);
+
+	{
+		// We do expect an ensure if we are using the Native IrisFastArraySerializer
+		FTestEnsureScope SuppressEnsureScope;
+
+		// Send and deliver packet
+		Server->UpdateAndSend({Client});
+	}
+
+	// Verify status, should be same as last time
+	UE_NET_ASSERT_EQ(ClientObject->FastArray.GetItemArray().Num(), 1);
+	UE_NET_ASSERT_EQ(ServerObject->FastArray.GetItemArray()[0].RepInt32, ClientObject->FastArray.GetItemArray()[0].RepInt32);
+
+	// Remove remaining replicated item (now index 0)
+	ServerObject->FastArray.Edit().Remove(0);
+
+	{
+		// We do expect an ensure if we are using the Native IrisFastArraySerializer
+		FTestEnsureScope SuppressEnsureScope;
+
+		// Send and deliver packet
+		Server->UpdateAndSend({Client});
+	}
+
+	// Verify status, should be same as last time
+	UE_NET_ASSERT_EQ(ClientObject->FastArray.GetItemArray().Num(), 0);
+}
+
+UE_NET_TEST_FIXTURE(FReplicationSystemServerClientTestFixture, TestFastArrayWithNotReplicatedItems)
+{
+	IConsoleVariable* CVarToOverride = IConsoleManager::Get().FindConsoleVariable(TEXT("net.Iris.UseNativeFastArray"));
+	UE_NET_ASSERT_NE(CVarToOverride, nullptr);
+	UE_NET_ASSERT_TRUE(CVarToOverride->IsVariableBool());
+	const bool bPrevUseNativeFastArray = CVarToOverride->GetBool();
+	ON_SCOPE_EXIT
+	{
+		CVarToOverride->Set(bPrevUseNativeFastArray, ECVF_SetByCode);
+	};
+
+	// Use "normal" fastarray
+	CVarToOverride->Set(false, ECVF_SetByCode);
+
+	UTestFastArrayReplicationState_FastArray_TestClassFastArray* ServerObject = Server->CreateObject<UTestFastArrayReplicationState_FastArray_TestClassFastArray>();
+
+	// Add a client
+	FReplicationSystemTestClient* Client = CreateClient();
+
+	// Send and deliver packet
+	Server->UpdateAndSend({Client});
+
+	// Store Pointer to objects
+	UTestFastArrayReplicationState_FastArray_TestClassFastArray* ClientObject = Cast<UTestFastArrayReplicationState_FastArray_TestClassFastArray>(Client->GetReplicationBridge()->GetReplicatedObject(ServerObject->NetRefHandle));
+
+	// Verify initial replication
+	UE_NET_ASSERT_TRUE(ClientObject!= nullptr);
+
+	// Add Replicated item
+	{
+		FTestFastArrayReplicationState_FastArrayItem ReplicatedEntry;
+		ReplicatedEntry.RepInt32 = 0U;
+		ServerObject->FastArray.Edit().Add(ReplicatedEntry);
+	}
+
+	// Add local explicitly not Replicated item
+	{
+		FTestFastArrayReplicationState_FastArrayItem ReplicatedEntry;
+		ReplicatedEntry.bSkipWriteFastArrayItem = true;
+		ReplicatedEntry.RepInt32 = 1000U;
+		ServerObject->FastArray.Edit().AddLocal(ReplicatedEntry);
+	}
+
+	// Add Replicated item
+	{
+		FTestFastArrayReplicationState_FastArrayItem ReplicatedEntry;
+		ReplicatedEntry.RepInt32 = 1U;
+		ServerObject->FastArray.Edit().Add(ReplicatedEntry);
+	}
+
+	// Add local explicitly not Replicated item
+	{
+		FTestFastArrayReplicationState_FastArrayItem ReplicatedEntry;
+		ReplicatedEntry.bSkipWriteFastArrayItem = true;
+		ReplicatedEntry.RepInt32 = 1000U;
+		ServerObject->FastArray.Edit().AddLocal(ReplicatedEntry);
+	}
+
+	// Send and deliver packet
+	Server->UpdateAndSend({Client});
+
+	// Verify status
+	UE_NET_ASSERT_EQ(ClientObject->FastArray.GetItemArray().Num(), 2);
+	UE_NET_ASSERT_EQ(ServerObject->FastArray.GetItemArray()[0].RepInt32, ClientObject->FastArray.GetItemArray()[0].RepInt32);
+	UE_NET_ASSERT_EQ(ServerObject->FastArray.GetItemArray()[2].RepInt32, ClientObject->FastArray.GetItemArray()[1].RepInt32);
+
+	// Remove not replicated item at index 1
+	ServerObject->FastArray.Edit().Remove(1);
+	
+	// Remove replicated item at index 0
+	ServerObject->FastArray.Edit().Remove(0);
+
+	// Send and deliver packet
+	Server->UpdateAndSend({Client});
+
+	// Verify status
+	UE_NET_ASSERT_EQ(ClientObject->FastArray.GetItemArray().Num(), 1);
+	UE_NET_ASSERT_EQ(ServerObject->FastArray.GetItemArray()[0].RepInt32, ClientObject->FastArray.GetItemArray()[0].RepInt32);
+
+	// Remove local item (now index 1)
+	ServerObject->FastArray.Edit().Remove(1);
+
+	// Send and deliver packet
+	Server->UpdateAndSend({Client});
+
+	// Verify status, should be same as last time
+	UE_NET_ASSERT_EQ(ClientObject->FastArray.GetItemArray().Num(), 1);
+	UE_NET_ASSERT_EQ(ServerObject->FastArray.GetItemArray()[0].RepInt32, ClientObject->FastArray.GetItemArray()[0].RepInt32);
+
+	// Remove remaining replicated item (now index 0)
+	ServerObject->FastArray.Edit().Remove(0);
+
+	// Send and deliver packet
+	Server->UpdateAndSend({Client});
+
+	// Verify status, should be same as last time
+	UE_NET_ASSERT_EQ(ClientObject->FastArray.GetItemArray().Num(), 0);
+
+}
 
 }
diff --git a/Engine/Plugins/Runtime/ReplicationSystemTestPlugin/Source/Private/Tests/ReplicationState/TestFastArrayReplicationState.h b/Engine/Plugins/Runtime/ReplicationSystemTestPlugin/Source/Private/Tests/ReplicationState/TestFastArrayReplicationState.h
--- a/Engine/Plugins/Runtime/ReplicationSystemTestPlugin/Source/Private/Tests/ReplicationState/TestFastArrayReplicationState.h
+++ b/Engine/Plugins/Runtime/ReplicationSystemTestPlugin/Source/Private/Tests/ReplicationState/TestFastArrayReplicationState.h
@@ -28,6 +28,9 @@ struct FTestFastArrayReplicationState_FastArrayItem : public FFastArraySerialize
 	UPROPERTY()
 	TObjectPtr<UObject> ObjectRef = nullptr;
 
+	UPROPERTY(NotReplicated)
+	bool bSkipWriteFastArrayItem = false;
+
 	// Callbacks
 	void PostReplicatedAdd(const struct FTestFastArrayReplicationState_FastArraySerializer& InArraySerializer);
 	void PostReplicatedChange(const struct FTestFastArrayReplicationState_FastArraySerializer& InArraySerializer);
@@ -49,8 +52,7 @@ struct FTestFastArrayReplicationState_FastArraySerializer : public FIrisFastArra
 	{
 	}
 
-	// Test of TIrisFastArrayEditor interface
-	// This is just to see if it works out as expected
+	// Test of experimental TIrisFastArrayEditor interface
 	typedef TArray<FTestFastArrayReplicationState_FastArrayItem> ItemArrayType;
 	const ItemArrayType& GetItemArray() const { return Items; }
 	ItemArrayType& GetItemArray() { return Items; }
@@ -64,6 +66,23 @@ struct FTestFastArrayReplicationState_FastArraySerializer : public FIrisFastArra
 	uint8 bHitPostReplicatedReceive : 1;
 
 	bool bPostReplicatedReceiveWasHitWithUnresolvedReferences;
+
+	// Should really never have non replicated elements in a fastarray especially on server, but that does not mean that we should not test it.
+	template<typename Type, typename SerializerType>
+	bool ShouldWriteFastArrayItem(const Type& Item, const bool bIsWritingOnClient) const
+	{
+		if (Item.bSkipWriteFastArrayItem)
+		{
+			return false;
+		}
+
+		if (bIsWritingOnClient)
+		{
+			return Item.ReplicationID != INDEX_NONE;
+		}
+
+		return true;
+	}
 	
 	void PostReplicatedReceive(const FFastArraySerializer::FPostReplicatedReceiveParameters& Parameters)
 	{
diff --git a/Engine/Source/Runtime/Net/Iris/Public/Iris/ReplicationState/IrisFastArraySerializer.h b/Engine/Source/Runtime/Net/Iris/Public/Iris/ReplicationState/IrisFastArraySerializer.h
--- a/Engine/Source/Runtime/Net/Iris/Public/Iris/ReplicationState/IrisFastArraySerializer.h
+++ b/Engine/Source/Runtime/Net/Iris/Public/Iris/ReplicationState/IrisFastArraySerializer.h
@@ -17,8 +17,12 @@ struct FIrisFastArraySerializerPrivateAccessor;
  * Specialization of FFastArraySerializer in order to add state tracking support for Iris
  * Current usage is to inherit from this struct instead of FFastArraySerializer, backwards compatible with existing system as it simply forwards calls to MarkDirty/MarkItemDirty
  * This class could be named FFastArrayReplicationState, but kept the FIrisFastArraySerializer to match old naming for the time being
+ *
+ * NOTE: IrisFastArraySerializer, does not support having local not replicated items in the array.
+ * If ShouldWriteFastArrayItem is overridden and filters out items, an ensure will be triggered. 
+ * Logic should still work, but extra elements will be replicated and filtered out on receiving end.
  */
-USTRUCT()
+USTRUCT(Experimental)
 struct FIrisFastArraySerializer : public FFastArraySerializer
 {
 	// At the moment as we have no way to specify this per derived type, currently we reserve a fixed range of bits used for the changemask, the first bit is used for the array itself
diff --git a/Engine/Source/Runtime/Net/Iris/Public/Iris/ReplicationState/Private/IrisFastArraySerializerInternal.h b/Engine/Source/Runtime/Net/Iris/Public/Iris/ReplicationState/Private/IrisFastArraySerializerInternal.h
--- a/Engine/Source/Runtime/Net/Iris/Public/Iris/ReplicationState/Private/IrisFastArraySerializerInternal.h
+++ b/Engine/Source/Runtime/Net/Iris/Public/Iris/ReplicationState/Private/IrisFastArraySerializerInternal.h
@@ -139,7 +139,7 @@ void TIrisFastArrayEditor<FastArrayType>::MarkItemDirty(FastArrayItemType& Item)
 				for (FastArrayItemType& ArrayItem : MakeArrayView(&Item, Array.Num() - Index))
 				{
 					const bool bIsWritingOnClient = false;
-					if (!FastArray.template ShouldWriteFastArrayItem<FastArrayItemType, FastArrayType>(ArrayItem, bIsWritingOnClient))
+					if (FastArray.template ShouldWriteFastArrayItem<FastArrayItemType, FastArrayType>(ArrayItem, bIsWritingOnClient))
 					{
 						if (ArrayItem.ReplicationID == INDEX_NONE)
 						{
@@ -195,9 +195,14 @@ void TIrisFastArrayEditor<FastArrayType>::Remove(int32 ItemIdx)
 		// If the Item is at the end we only need to mark the array as dirty
 		ItemArray.RemoveAt(ItemIdx);
 		MarkArrayDirty();
+
 		if (ItemIdx < ItemArray.Num())
 		{
-			FIrisFastArraySerializerPrivateAccessor::MarkAllArrayItemsDirty(FastArray, ItemIdx);
+			const FReplicationStateHeader& Header = FIrisFastArraySerializerPrivateAccessor::GetReplicationStateHeader(FastArray);
+			if (Header.IsBound())
+			{
+				FIrisFastArraySerializerPrivateAccessor::MarkAllArrayItemsDirty(FastArray, ItemIdx);
+			}
 		}
 	}
 }
diff --git a/Engine/Source/Runtime/Net/Iris/Public/Iris/ReplicationSystem/FastArrayReplicationFragment.h b/Engine/Source/Runtime/Net/Iris/Public/Iris/ReplicationSystem/FastArrayReplicationFragment.h
--- a/Engine/Source/Runtime/Net/Iris/Public/Iris/ReplicationSystem/FastArrayReplicationFragment.h
+++ b/Engine/Source/Runtime/Net/Iris/Public/Iris/ReplicationSystem/FastArrayReplicationFragment.h
@@ -210,10 +210,10 @@ bool TFastArrayReplicationFragment<FastArrayItemType, FastArrayType>::PollAllSta
 	// First we must resize target to match size of source data
 	bool bMarkArrayDirty = false;
 	const uint32 ElementCount = SrcWrappedArray->Num();
+	const uint32 OldDstElementCount = DstWrappedArray->Num();
 	if (DstWrappedArray->Num() != ElementCount)
 	{
 		DstWrappedArray->SetNum(ElementCount);
-		bMarkArrayDirty = true;
 	}
 
 	const FReplicationStateDescriptor* ArrayElementDescriptor = GetArrayElementDescriptor();
@@ -231,6 +231,8 @@ bool TFastArrayReplicationFragment<FastArrayItemType, FastArrayType>::PollAllSta
 	FastArrayItemType* DstItems = DstWrappedArray->GetData();
 	FastArrayItemType* SrcItems = SrcWrappedArray->GetData();
 
+	// We keep a separate count as we do not care about items that should not be replicated.
+	int32 DstItemCount = 0;
 	{
 #if WITH_PUSH_MODEL
 		// Disable push model by temporarily setting the FastArray's RepIndex to none.
@@ -241,8 +243,9 @@ bool TFastArrayReplicationFragment<FastArrayItemType, FastArrayType>::PollAllSta
 		// Iterate over array entries and copy the statedata using internal data if it has changed
 		for (int32 ElementIt = 0, ElementEndIt = ElementCount; ElementIt < ElementEndIt; ++ElementIt)
 		{
+			const int32 DstElementIt = DstItemCount;
 			FastArrayItemType& SrcItem = SrcItems[ElementIt];
-			FastArrayItemType& DstItem = DstItems[ElementIt];
+			FastArrayItemType& DstItem = DstItems[DstElementIt];
 
 			const bool bIsWritingOnClient = false;
 			if (SrcArraySerializer->template ShouldWriteFastArrayItem<FastArrayItemType, FastArrayType>(SrcItem, bIsWritingOnClient))
@@ -261,14 +264,25 @@ bool TFastArrayReplicationFragment<FastArrayItemType, FastArrayType>::PollAllSta
 					// Mark element as dirty and mark array as dirty as well.
 					if (ChangeMaskBitCount)
 					{
-						MemberChangeMask.SetBit((ElementIt % ChangeMaskBitCount) + ChangeMaskBitOffset);
+						MemberChangeMask.SetBit((DstElementIt % ChangeMaskBitCount) + ChangeMaskBitOffset);
 					}
 					bMarkArrayDirty = true;
 				}
+				++DstItemCount;
 			}
 		}
 	}
 
+	// Set actual num replicated items.
+	if (DstWrappedArray->Num() != DstItemCount)
+	{
+		// Set actual num, but keeping the allocated size of the source array
+		DstWrappedArray->SetNum(DstItemCount, EAllowShrinking::No);
+	}
+
+	// Mark dirty if the new filtered size differs from our previous state.
+	bMarkArrayDirty = bMarkArrayDirty || (OldDstElementCount != DstItemCount);
+
 	// We update this after the poll since every call to MarkItem() dirty will Increase the ArrayReplicationKey
 	DstArraySerializer->ArrayReplicationKey = SrcArraySerializer->ArrayReplicationKey;
 
@@ -471,6 +485,24 @@ bool TNativeFastArrayReplicationFragment<FastArrayItemType, FastArrayType, Polli
 					bMarkArrayDirty = true;
 				}
 			}
+			else
+			{
+				ensureMsgf(false, TEXT("Native IrisFastArraySerializer does not support local non-replicated items, use FastArraySerializer intead if this is required."));
+
+				// Even if the native variant do not really support not replicated entries we still mark the item dirty and update the stored data to keep logic working.
+				if (SrcItem.ReplicationKey != DstItem.ReplicationKey || SrcItem.ReplicationID != DstItem.ReplicationID)
+				{
+					DstItem.ReplicationKey = SrcItem.ReplicationKey;
+					DstItem.ReplicationID = SrcItem.ReplicationID;
+
+					// Mark element as dirty and mark array as dirty as well.
+					if (ChangeMaskBitCount)
+					{
+						MemberChangeMask.SetBit((ElementIt % ChangeMaskBitCount) + ChangeMaskBitOffset);
+					}
+					bMarkArrayDirty = true;
+				}
+			}
 		}
 
 		// We update this after the poll since every call to MarkItem() dirty will Increase the ArrayReplicationKey
diff --git a/Engine/Source/Runtime/Net/Iris/Public/Iris/ReplicationSystem/Private/FastArrayReplicationFragmentInternal.h b/Engine/Source/Runtime/Net/Iris/Public/Iris/ReplicationSystem/Private/FastArrayReplicationFragmentInternal.h
--- a/Engine/Source/Runtime/Net/Iris/Public/Iris/ReplicationSystem/Private/FastArrayReplicationFragmentInternal.h
+++ b/Engine/Source/Runtime/Net/Iris/Public/Iris/ReplicationSystem/Private/FastArrayReplicationFragmentInternal.h
@@ -277,7 +277,7 @@ void FFastArrayReplicationFragmentHelper::ApplyReplicatedState(FastArrayType* Ds
 			else
 			{
 				// Since we zero initialize our replicated properties we can end up with ReplicationID == 0 when receiving partial changes which should be ignored.
-				if (SrcItems[It].ReplicationID != 0)
+				if (!(SrcItems[It].ReplicationID == INDEX_NONE || SrcItems[It].ReplicationID == 0))
 				{
 					int32 AddedIndex = DstWrappedArray->Add(SrcItems[It]);
 
