Iris
- Test that caused NCL due to not-split part sent by hugeobject path failed to fit in packet.

diff --git a/Engine/Plugins/Runtime/ReplicationSystemTestPlugin/Source/Private/Tests/ReplicationSystem/TestObjectSplitting.cpp b/Engine/Plugins/Runtime/ReplicationSystemTestPlugin/Source/Private/Tests/ReplicationSystem/TestObjectSplitting.cpp
--- a/Engine/Plugins/Runtime/ReplicationSystemTestPlugin/Source/Private/Tests/ReplicationSystem/TestObjectSplitting.cpp
+++ b/Engine/Plugins/Runtime/ReplicationSystemTestPlugin/Source/Private/Tests/ReplicationSystem/TestObjectSplitting.cpp
@@ -1268,4 +1268,74 @@ UE_NET_TEST_FIXTURE(FSplitObjectTestFixture, TestReliableAttachmentAddedAfterSpl
 }
 #endif
 
+// Special test for network connection loss due to bitstream error on receiving end. 
+// Basically we need to craft data so that it hits the hugeobject path due to combined size of exports and attachments.
+// Attachments are then dropped reducing the size of written data to fit within the hugeobject split threshold but close enough for it to not fit when including the overhead of the raw blob. 
+// Prior to fix this would count the object as written even if it was just enqueued for next packet which resulted in bitstream error.
+// NOTE: This test is compiled out due to being very sensitive to code changes affects bits written/config settings etc.
+#if 0
+UE_NET_TEST_FIXTURE(FSplitObjectTestFixture, SplitHugeObjectNoPartsInPacket)
+{
+	FReplicationSystemTestClient* Client = CreateClient();
+	RegisterNetBlobHandlers(Client);
+
+	// Limit packet size to setup this case
+	Server->SetMaxSendPacketSize(200);
+
+	UTestReplicatedIrisObject::FComponents Components;
+	Components.DynamicStateComponentCount = 1;
+	Components.ObjectReferenceComponentCount = 1;
+	UTestReplicatedIrisObject* ServerObject = Server->CreateObject(Components);
+	
+	// Add another object to send packet even if split object does not fit in packet.
+	UTestReplicatedIrisObject* DummyObject = Server->CreateObject(Components);
+
+	// Give the object we expect to be split high priority
+	Server->ReplicationSystem->SetStaticPriority(ServerObject->NetRefHandle, 100.f);
+
+	TRefCountPtr<FNetObjectAttachment> Attachment;
+	Attachment = ServerMockNetObjectAttachmentHandler->CreateReliableNetObjectAttachment(82);
+	FNetObjectReference AttachmentTarget = FObjectReferenceCache::MakeNetObjectReference(ServerObject->NetRefHandle);
+	Server->GetReplicationSystem()->QueueNetObjectAttachment(Client->ConnectionIdOnServer, AttachmentTarget, Attachment);
+
+	static uint32 PayloadByteCount = 112U;
+	SetObjectPayloadByteCount(ServerObject, PayloadByteCount);
+
+	// Exports to balance the size of the written attachments.
+	UTestReplicatedIrisPushModelComponentWithObjectReference* Component = ServerObject->ObjectReferenceComponents[0].Get();
+	Component->RawObjectPtrRef = UReplicationBridge::StaticClass();
+
+	// Send and deliver packet
+	Server->UpdateAndSend({Client});
+
+	// The object should not exist as we expect it to be split in the first packet but not fit
+	{
+		UTestReplicatedIrisObject* ClientObject = Cast<UTestReplicatedIrisObject>(Client->GetReplicationBridge()->GetReplicatedObject(ServerObject->NetRefHandle));
+		UE_NET_ASSERT_EQ(ClientObject, nullptr);
+	}
+
+	// Bump the packet size, so that the unsplit rpc will fit.
+	Server->SetMaxSendPacketSize(400);
+	Server->UpdateAndSend({Client});
+
+	// The object should now exist as we expect it to have been sent as an un-split hugeobject, but attachments should not have been delivered yet.
+	{
+		UTestReplicatedIrisObject* ClientObject = Cast<UTestReplicatedIrisObject>(Client->GetReplicationBridge()->GetReplicatedObject(ServerObject->NetRefHandle));
+		UE_NET_ASSERT_NE(ClientObject, nullptr);
+
+		const UMockNetObjectAttachmentHandler::FCallCounts AttachmentCallCounts = ClientMockNetObjectAttachmentHandler->GetFunctionCallCounts();
+		UE_NET_ASSERT_EQ(AttachmentCallCounts.OnNetBlobReceived, 0U);
+	}
+
+	// Send another packet to verify that we now got the attachement
+	Server->UpdateAndSend({Client});
+
+	// Verify the attachment made it through
+	{
+		const UMockNetObjectAttachmentHandler::FCallCounts AttachmentCallCounts = ClientMockNetObjectAttachmentHandler->GetFunctionCallCounts();
+		UE_NET_ASSERT_EQ(AttachmentCallCounts.OnNetBlobReceived, 1U);
+	}
+}
+#endif
+
 }
