Iris
- Fixed issue with ObjectPoller not force polling ObjectReferences for pushbased fragments in mixed protocols after GC

diff --git a/Engine/Plugins/Runtime/ReplicationSystemTestPlugin/Source/Private/Tests/ReplicationSystem/TestReplicationOperations.cpp b/Engine/Plugins/Runtime/ReplicationSystemTestPlugin/Source/Private/Tests/ReplicationSystem/TestReplicationOperations.cpp
--- a/Engine/Plugins/Runtime/ReplicationSystemTestPlugin/Source/Private/Tests/ReplicationSystem/TestReplicationOperations.cpp
+++ b/Engine/Plugins/Runtime/ReplicationSystemTestPlugin/Source/Private/Tests/ReplicationSystem/TestReplicationOperations.cpp
@@ -502,9 +502,10 @@ private:
 
 UE_NET_TEST_FIXTURE(FTestReplicationOperationsForObjectsFixture, StaleObjectPointerIsUpdated)
 {
+	TGuardConsoleVariable<bool> CVarOverride(TEXT("gc.GarbageEliminationEnabled"), true);
+
 	// Create two objects, one which is referencing the other. Verify reference gets updated when first object is destroyed.
 	UTestReplicatedIrisObject* Object0 = CreateObject();
-	FNetRefHandle HandleToObject0 = ReplicationBridge->BeginReplication(Object0);
 
 	UTestReplicatedIrisObject* Object1 = CreateObject();
 	constexpr uint32 ObjectReferenceComponentCount = 1;
@@ -549,7 +550,57 @@ UE_NET_TEST_FIXTURE(FTestReplicationOperationsForObjectsFixture, StaleObjectPoin
 	// Compare the two state buffers. Despite that we haven't manually changed any object references we expect them to differ thanks to garbage collection support.
 	UE_NET_ASSERT_FALSE(FReplicationProtocolOperationsInternal::IsEqualQuantizedState(SerializationContext, StateBuffers1[0], StateBuffers1[1], Protocol1));
 }
-	
+
+UE_NET_TEST_FIXTURE(FTestReplicationOperationsForObjectsFixture, StaleObjectPointerPolledAfterGC)
+{
+	TGuardConsoleVariable<bool> CVarOverride(TEXT("gc.GarbageEliminationEnabled"), true);
+
+	// Create two objects, one which is referencing the other. Verify reference gets updated when first object is destroyed.
+	UTestReplicatedIrisObject* Object0 = CreateObject();
+
+	UTestReplicatedIrisObject* Object1 = CreateObject();
+	constexpr uint32 ObjectReferenceComponentCount = 1;
+	Object1->AddComponents(UTestReplicatedIrisObject::FComponents{0,0,0,0,ObjectReferenceComponentCount});
+	Object1->ObjectReferenceComponents[0]->RawObjectPtrRef = Object0;
+	FNetRefHandle HandleToObject1 = ReplicationBridge->BeginReplication(Object1);
+
+	const FReplicationInstanceProtocol* InstanceProtocol1 = ReplicationBridge->GetReplicationInstanceProtocol(HandleToObject1);
+	const FReplicationProtocol* Protocol1 = ReplicationSystem->GetReplicationProtocol(HandleToObject1);
+
+	// Verify protocol, if it fails, the test needs to be updated
+	UE_NET_ASSERT_EQ(Protocol1->ReplicationStateCount, 2U);
+	UE_NET_ASSERT_EQ((uint32)Protocol1->ReplicationStateDescriptors[1]->MemberCount, 3U);
+	UE_NET_ASSERT_TRUE(FCString::Strcmp(Protocol1->ReplicationStateDescriptors[1]->MemberDebugDescriptors[1].DebugName->Name, TEXT("RawObjectPtrRef")) == 0);
+
+	// Update polling
+	ReplicationSystem->NetUpdate(0.033f);
+	ReplicationSystem->PostSendUpdate();
+
+	// Get polled value
+	UObject* PolledPointerValue0 = nullptr;
+	((FPropertyReplicationFragment*)InstanceProtocol1->Fragments[1])->GetPropertyReplicationState()->GetPropertyValue(1, &PolledPointerValue0);
+
+	// Should be same as source
+	UE_NET_ASSERT_TRUE(PolledPointerValue0 == Object1->ObjectReferenceComponents[0]->RawObjectPtrRef);
+
+	// Destroy first object and invalidate references to it.
+	ReplicationBridge->EndReplication(Object0);
+	DestroyObject(Object0);
+	constexpr bool bPerformFullPurge = false;
+	CollectGarbage(GARBAGE_COLLECTION_KEEPFLAGS, bPerformFullPurge);
+
+	// Update polling
+	ReplicationSystem->NetUpdate(0.033f);
+	ReplicationSystem->PostSendUpdate();
+
+	// Get polled value
+	UObject* PolledPointerValue1 = nullptr;
+	((FPropertyReplicationFragment*)InstanceProtocol1->Fragments[1])->GetPropertyReplicationState()->GetPropertyValue(1, &PolledPointerValue1);
+
+	// Polled value should now be nullptr
+	UE_NET_ASSERT_TRUE(PolledPointerValue1 == nullptr);
+}
+
 class FTestReplicationOperationsOnInitStateFixture : public FTestReplicationOperationsFixture
 {
 protected:
diff --git a/Engine/Source/Runtime/Net/Iris/Private/Iris/ReplicationSystem/Polling/ObjectPoller.cpp b/Engine/Source/Runtime/Net/Iris/Private/Iris/ReplicationSystem/Polling/ObjectPoller.cpp
--- a/Engine/Source/Runtime/Net/Iris/Private/Iris/ReplicationSystem/Polling/ObjectPoller.cpp
+++ b/Engine/Source/Runtime/Net/Iris/Private/Iris/ReplicationSystem/Polling/ObjectPoller.cpp
@@ -423,9 +423,9 @@ void FObjectPoller::PushModelPollObject(FInternalNetRefIndex ObjectIndex, FNetSt
 		UE_NET_TRACE_POLL_OBJECT_SCOPE(ObjectData.RefHandle, Timer);
 
 		// If the object has fragments with pushed based properties, and is not marked dirty and object is affected by GC we need to make sure that we refresh cached references for all fragments with push based properties
-		const bool bIsFullPushBasedObject = EnumHasAnyFlags(InstanceTraits, EReplicationInstanceProtocolTraits::HasFullPushBasedDirtiness);
+		const bool bHasPushBasedFragments = EnumHasAnyFlags(InstanceTraits, EReplicationInstanceProtocolTraits::HasPushBasedDirtiness);
 		const bool bHasObjectReferences = EnumHasAnyFlags(InstanceTraits, EReplicationInstanceProtocolTraits::HasObjectReference);
-		const bool bNeedsRefreshOfCachedObjectReferences = ((!(bWantsFullPoll | bIsDirtyObject)) & bIsGCAffectedObject & bIsFullPushBasedObject & bHasObjectReferences);
+		const bool bNeedsRefreshOfCachedObjectReferences = ((!(bWantsFullPoll | bIsDirtyObject)) & bIsGCAffectedObject & bHasPushBasedFragments & bHasObjectReferences);
 		if (bNeedsRefreshOfCachedObjectReferences)
 		{
 			// Only states which has full push based dirtiness need to be updated as the other states will be at least partially polled anyway.
