Iris
- Addressed issues with PacketNotification getting stuck in deadlock if both ends overshoots ack sequence history.
- Now includes explicit sequency number for sequencyhistory to keep progressing sequencenumbers while flushing the ack sequence history.
- Bumped NetVersion as this includes changes to header.

diff --git a/Engine/Plugins/Runtime/ReplicationSystemTestPlugin/Source/Private/Tests/Core/TestNetPacketNotify.cpp b/Engine/Plugins/Runtime/ReplicationSystemTestPlugin/Source/Private/Tests/Core/TestNetPacketNotify.cpp
--- a/Engine/Plugins/Runtime/ReplicationSystemTestPlugin/Source/Private/Tests/Core/TestNetPacketNotify.cpp
+++ b/Engine/Plugins/Runtime/ReplicationSystemTestPlugin/Source/Private/Tests/Core/TestNetPacketNotify.cpp
@@ -47,7 +47,7 @@ struct FNetPacketNotifyTestUtil : public FNetworkAutomationTestSuiteFixture
 	static void InitHistory(FNetPacketNotify::SequenceHistoryT& History, const T& DataToSet)
 	{
 		const SIZE_T Count = sizeof(T) / sizeof(DataToSet[0]);
-		static_assert(Count < FNetPacketNotify::SequenceHistoryT::WordCount, "DataToSet must be smaller than HistoryBuffer");
+		static_assert(Count <= FNetPacketNotify::SequenceHistoryT::WordCount, "DataToSet array cannot be larger HistoryBuffer");
 
 		for (SIZE_T It=0; It < Count; ++It)
 		{
@@ -124,42 +124,99 @@ struct FNetPacketNotifyTestUtil : public FNetworkAutomationTestSuiteFixture
 	struct FTestNode
 	{
 		FNetPacketNotify PacketNotify;
-		TArray<FNetPacketNotify::SequenceNumberT> AcceptedInPackets;
-		TArray<FNetPacketNotify::SequenceNumberT> AcknowledgedPackets;
+
+		uint32 InPacketId = -1;
+		uint32 OutPacketId = 0;
+		uint32 LastNotifiedPacketId = -1;
+
+		TArray<uint32> AcceptedInPackets;
+		TArray<uint32> AcknowledgedPackets;
+		TArray<uint32> ExpectedAcceptedInPackets;
 
 		FBitWriter Writer{FNetPacketNotify::MaxSequenceHistoryLength, true};
 
 		FTestNode(FNetPacketNotify::SequenceNumberT FirstSequence = 0)
 		{
-			PacketNotify.Init(FNetPacketNotify::SequenceNumberT(FirstSequence.Get() - 1), FirstSequence);
+			InPacketId = FirstSequence.Get() - 1U;
+			OutPacketId = FirstSequence.Get();
+			LastNotifiedPacketId = OutPacketId - 1;
+
+			PacketNotify.Init(FNetPacketNotify::SequenceNumberT::SequenceT(InPacketId), FNetPacketNotify::SequenceNumberT::SequenceT(OutPacketId));
 		}
 
 		FTestNode(FNetPacketNotify::SequenceNumberT InSequnce, FNetPacketNotify::SequenceNumberT OutSequnce)
 		{
-			PacketNotify.Init(InSequnce, OutSequnce);
+			InPacketId = InSequnce.Get();
+			OutPacketId = OutSequnce.Get();
+			LastNotifiedPacketId = OutPacketId - 1;
+
+			PacketNotify.Init(FNetPacketNotify::SequenceNumberT::SequenceT(InPacketId), FNetPacketNotify::SequenceNumberT::SequenceT(OutPacketId));
 		}
 
-		const FBitWriter& Send()
+		const FBitWriter& Send(bool ResetWriter = true)
 		{
 			const SIZE_T CurrentHistoryWordCount = FMath::Clamp<SIZE_T>((PacketNotify.GetCurrentSequenceHistoryLength() + FNetPacketNotify::SequenceHistoryT::BitsPerWord - 1u) / FNetPacketNotify::SequenceHistoryT::BitsPerWord, 1u, FNetPacketNotify::SequenceHistoryT::WordCount);
 
 			FNetPacketNotify::FNotificationHeader Data;
 
-			Writer.Reset();
-			PacketNotify.WriteHeader(Writer);
+			if (ResetWriter)
+			{
+				Writer.Reset();
+			}
 
+			PacketNotify.WriteHeader(Writer);
 			PacketNotify.CommitAndIncrementOutSeq();
 
+			++OutPacketId;
+
 			return Writer;
 		}
 
-		void Receive(const FBitWriter& IncomingData)
+		const FBitWriter& SendBurst(int32 NumPackets, bool bResetWriter = true)
 		{
-			auto HandleAck = [&](FNetPacketNotify::SequenceNumberT Seq, bool bDelivered)
+			if (bResetWriter)
+			{
+				Writer.Reset();
+			}
+			while (NumPackets > 0)
 			{
+				Send(false);
+				--NumPackets;
+			}
+
+			return Writer;
+		}
+
+		// Read IncomingData and (optional) compare received acks with state from the from node.
+		// Returns false if an error is detected
+		bool Receive(const FBitWriter& IncomingData, const FTestNode* From = nullptr)
+		{
+			bool bSuccess = true;
+
+			auto HandlePacketNotification = [&](FNetPacketNotify::SequenceNumberT Seq, bool bDelivered)
+			{
+				++LastNotifiedPacketId;
+
+				// Sanity check
+				if (!ensure(FNetPacketNotify::SequenceNumberT(FNetPacketNotify::SequenceNumberT::SequenceT(LastNotifiedPacketId)) == Seq))
+				{
+					bSuccess = false;
+				}
+
 				if (bDelivered)
 				{
-					AcknowledgedPackets.Add(Seq);
+					AcknowledgedPackets.Add(LastNotifiedPacketId);
+					if (!ensure(!From || From->AcceptedInPackets.Contains(LastNotifiedPacketId)))
+					{
+						bSuccess = false;
+					}
+				}
+				else
+				{
+					if (!ensure(!From || !From->AcceptedInPackets.Contains(LastNotifiedPacketId)))
+					{
+						bSuccess = false;
+					}
 				}
 			};
 
@@ -167,42 +224,140 @@ struct FNetPacketNotifyTestUtil : public FNetworkAutomationTestSuiteFixture
 
 			FNetPacketNotify::FNotificationHeader Header;
 			PacketNotify.ReadHeader(Header, Reader);
+			
+			const uint32 SequenceDelta = PacketNotify.Update(Header, HandlePacketNotification);
+			if (SequenceDelta > 0U)
+			{
+				InPacketId += SequenceDelta;
+
+				// Sanity check
+				check(FNetPacketNotify::SequenceNumberT(FNetPacketNotify::SequenceNumberT::SequenceT(InPacketId)) == PacketNotify.InSeq);
+
+				if (!PacketNotify.IsWaitingForSequenceHistoryFlush())
+				{
+					PacketNotify.AckSeq(FNetPacketNotify::SequenceNumberT::SequenceT(InPacketId));
+					AcceptedInPackets.Add(InPacketId);
+				}
+			}
 
-			const uint32 SequenceDelta = PacketNotify.Update(Header, HandleAck);
-			if (SequenceDelta > 0U && !PacketNotify.IsWaitingForSequenceHistoryFlush())
+			return bSuccess;
+		}
+
+		// Read IncomingData and (optional) compare received acks with state from the from node.
+		// Returns false if an error is detected
+		bool ReceiveBurst(const FBitWriter& IncomingData, const FTestNode* From = nullptr)
+		{
+			bool bSuccess = true;
+
+			auto HandlePacketNotification = [&](FNetPacketNotify::SequenceNumberT Seq, bool bDelivered)
 			{
-				PacketNotify.AckSeq(PacketNotify.GetInSeq());
-				AcceptedInPackets.Add(PacketNotify.GetInSeq());
+				++LastNotifiedPacketId;
+
+				// Sanity check
+				if (!ensure(FNetPacketNotify::SequenceNumberT(FNetPacketNotify::SequenceNumberT::SequenceT(LastNotifiedPacketId)) == Seq))
+				{
+					bSuccess = false;
+				}
+
+				if (bDelivered)
+				{
+					AcknowledgedPackets.Add(LastNotifiedPacketId);
+					if (!ensure(!From || From->AcceptedInPackets.Contains(LastNotifiedPacketId)))
+					{
+						bSuccess = false;
+					}
+				}
+				else
+				{
+					if (!ensure(!From || !From->AcceptedInPackets.Contains(LastNotifiedPacketId)))
+					{
+						bSuccess = false;
+					}
+				}
+			};
+
+			FBitReader Reader(IncomingData.GetData(), IncomingData.GetNumBits());
+
+			while (Reader.GetBytesLeft() > 0)
+			{
+				FNetPacketNotify::FNotificationHeader Header;
+				PacketNotify.ReadHeader(Header, Reader);
+
+				const uint32 SequenceDelta = PacketNotify.Update(Header, HandlePacketNotification);
+				if (SequenceDelta > 0U)
+				{
+					InPacketId += SequenceDelta;
+
+					// Sanity check
+					check(FNetPacketNotify::SequenceNumberT(FNetPacketNotify::SequenceNumberT::SequenceT(InPacketId)) == PacketNotify.InSeq);
+				
+					if (!PacketNotify.IsWaitingForSequenceHistoryFlush())
+					{
+						PacketNotify.AckSeq(FNetPacketNotify::SequenceNumberT::SequenceT(InPacketId));
+						AcceptedInPackets.Add(InPacketId);
+					}
+				}
 			}
+
+			return bSuccess;
 		}
 
 		void SendAndDeliverTo(FTestNode& Target)
 		{
 			Target.Receive(Send());
 		}
-
-		void SendAndDeliverTo(FTestNode& Target, TArrayView<const FNetPacketNotify::SequenceNumberT> SequenceNumbers)
+	
+		void SendRange(TArrayView<const uint32> PacketIds, bool bResetWriter = true)
 		{
-			for (FNetPacketNotify::SequenceNumberT Seq : SequenceNumbers)
+			if (bResetWriter)
+			{
+				Writer.Reset();
+			}
+			for (const uint32 PacketId : PacketIds)
 			{
 				// Skip missing seqs
-				while (Seq != PacketNotify.GetOutSeq())
+				FBitWriterMark Mark(Writer);
+				while (PacketId != OutPacketId)
+				{					
+					Send(false);
+				}
+				Mark.Pop(Writer);
+				Send(false);
+			}
+		}
+
+		void SendAndDeliverTo(FTestNode& Target, TArrayView<const uint32> PacketIds)
+		{
+			SendRange(PacketIds);
+			Target.ReceiveBurst(Writer);
+		}
+
+		bool ValidateAcknowledgedPackets(TArrayView<const uint32> ExpectedPacketIds) const
+		{
+			if (AcknowledgedPackets.Num() != ExpectedPacketIds.Num())
+			{
+				return false;
+			}
+			const int32 MaxCount = FMath::Min(AcknowledgedPackets.Num(), ExpectedPacketIds.Num());
+			for (int32 I = 0; I < MaxCount; ++I)
+			{
+				if (AcknowledgedPackets[I] != ExpectedPacketIds[I])
 				{
-					Send();
+					UE_LOG(LogNetTraffic, Error, TEXT("Ack %u != Exp %u index %u"), AcknowledgedPackets[I], ExpectedPacketIds[I], I);
+					return false;
 				}
-				// Send and receive seq
-				Target.Receive(Send());
 			}
+			return true;
 		}
 
-		bool ValidateAcknowledgedPackets(TArrayView<const FNetPacketNotify::SequenceNumberT> ExpectedSequenceNumbers) const
+		bool ValidateAcceptedPackets(TArrayView<const uint32> ExpectedPacketIds) const
 		{
-			return AcknowledgedPackets == ExpectedSequenceNumbers;
+			return AcceptedInPackets == ExpectedPacketIds;
 		}
 
-		bool ValidateAcceptedPackets(TArrayView<const FNetPacketNotify::SequenceNumberT> ExpectedSequenceNumbers) const
+		bool ValidateExpectedAcceptedPackets()
 		{
-			return AcceptedInPackets == ExpectedSequenceNumbers;
+			return ValidateAcceptedPackets(MakeArrayView(ExpectedAcceptedInPackets));
 		}
 	};
 };
@@ -443,7 +598,7 @@ UE_NET_TEST_FIXTURE(FNetPacketNotifyTestUtil, TestIgnoresOutOfBound)
 	FTestNode Src(FNetPacketNotify::SequenceNumberT::SeqNumberHalf);
 	FTestNode Dst;
 
-	// Pretend to send and accept packet in both directions
+	// Pretend to send packet in both directions, which should be ignored by dst
 	Src.SendAndDeliverTo(Dst);
 	Dst.SendAndDeliverTo(Src);
 
@@ -458,7 +613,7 @@ UE_NET_TEST_FIXTURE(FNetPacketNotifyTestUtil, TestOutOfBoundDoesNotTriggerFlush)
 	FTestNode Dst;
 
 	// Pretend to send and accept packet in both directions
-	Src.SendAndDeliverTo(Dst , {FNetPacketNotify::SequenceNumberT::SeqNumberMax, 0U});
+	Src.SendAndDeliverTo(Dst , {FNetPacketNotify::SequenceNumberT::SeqNumberMax, FNetPacketNotify::SequenceNumberT::SeqNumberMax + 1U});
 	Dst.SendAndDeliverTo(Src);
 
 	// Validate expected state
@@ -476,13 +631,16 @@ UE_NET_TEST_FIXTURE(FNetPacketNotifyTestUtil, TestBurstReceiveValidRangeWrapArou
 	FTestNode Dst(FNetPacketNotify::SequenceNumberT::SeqNumberHalf + 2);
 
 	// Pretend to send and accept packet in both directions
-	Src.SendAndDeliverTo(Dst, {FNetPacketNotify::SequenceNumberT::SeqNumberMax, 0U});
+	Dst.ExpectedAcceptedInPackets.Add(FNetPacketNotify::SequenceNumberT::SeqNumberMax);
+	Dst.ExpectedAcceptedInPackets.Add(FNetPacketNotify::SequenceNumberT::SeqNumberMax + 1U);
+
+	Src.SendAndDeliverTo(Dst, MakeArrayView(Dst.ExpectedAcceptedInPackets));
 	Dst.SendAndDeliverTo(Src);
 
 	// Validate expected state
 	UE_NET_ASSERT_FALSE(Dst.PacketNotify.IsWaitingForSequenceHistoryFlush());
-	UE_NET_ASSERT_TRUE(Src.ValidateAcknowledgedPackets({FNetPacketNotify::SequenceNumberT::SeqNumberMax, 0U}));
-	UE_NET_ASSERT_TRUE(Dst.ValidateAcceptedPackets({FNetPacketNotify::SequenceNumberT::SeqNumberMax, 0U}));
+	UE_NET_ASSERT_TRUE(Dst.ValidateExpectedAcceptedPackets());
+	UE_NET_ASSERT_TRUE(Src.ValidateAcknowledgedPackets(MakeArrayView(Dst.AcceptedInPackets)));
 }
 
 UE_NET_TEST_FIXTURE(FNetPacketNotifyTestUtil, TestAcceptsHugeDeltaWithNoPreviousAcks)
@@ -500,8 +658,8 @@ UE_NET_TEST_FIXTURE(FNetPacketNotifyTestUtil, TestAcceptsHugeDeltaWithNoPrevious
 	Dst.SendAndDeliverTo(Src);
 
 	// Validate expected state
-	UE_NET_ASSERT_TRUE(Src.ValidateAcknowledgedPackets({FirstOutSeq}));
-	UE_NET_ASSERT_TRUE(Dst.ValidateAcceptedPackets({FirstOutSeq}));
+	UE_NET_ASSERT_TRUE(Src.ValidateAcknowledgedPackets({FirstOutSeq.Get()}));
+	UE_NET_ASSERT_TRUE(Dst.ValidateAcceptedPackets({FirstOutSeq.Get()}));
 }
 
 UE_NET_TEST_FIXTURE(FNetPacketNotifyTestUtil, TestBurstReceiveValidRange)
@@ -549,7 +707,7 @@ UE_NET_TEST_FIXTURE(FNetPacketNotifyTestUtil, TestReceiveDoesNotTriggersWaitForS
 	FTestNode Dst;
 
 	// Send a specific sequence that would overflow window if we had already acked data
-	// we have no previous acks we can accept the sequence and allow the other end to inject any missing nacks
+	// we have no previous acks we can accept the sequence and allow the other end to inject any Fing nacks
 	Src.SendAndDeliverTo(Dst, {FNetPacketNotify::MaxSequenceHistoryLength + 1, FNetPacketNotify::MaxSequenceHistoryLength + 2});
 
 	// Respond
@@ -639,6 +797,243 @@ UE_NET_TEST_FIXTURE(FNetPacketNotifyTestUtil, TestGetBothEndsWaitForSequenceWind
 	UE_NET_ASSERT_TRUE(Src.ValidateAcceptedPackets({0, FNetPacketNotify::MaxSequenceHistoryLength + 1, FNetPacketNotify::MaxSequenceHistoryLength + 2, FNetPacketNotify::MaxSequenceHistoryLength + 3}));
 }
 
+UE_NET_TEST_FIXTURE(FNetPacketNotifyTestUtil, TestGetBothEndsWaitForSequenceWindowFlushDueToSendWindowOverflow)
+{
+	LOG_SCOPE_VERBOSITY_OVERRIDE(LogNetTraffic, ELogVerbosity::Error);
+
+	FNetPacketNotify::SequenceNumberT SrcStartingSeq(1024);
+	FNetPacketNotify::SequenceNumberT DstStartingSeq(0);
+
+	FTestNode Src(DstStartingSeq, SrcStartingSeq + 1);
+	FTestNode Dst(SrcStartingSeq, DstStartingSeq + 1);
+
+	// Send off bursts on both ends, over committing send window which will trigger SequenceHistoryFlush.
+	Src.SendBurst(FNetPacketNotify::MaxSequenceHistoryLength + 1);
+	Dst.SendBurst(FNetPacketNotify::MaxSequenceHistoryLength + 1);
+
+	// Receive the data
+	Dst.ReceiveBurst(Src.Writer, &Src);
+	Src.ReceiveBurst(Dst.Writer, &Dst);
+
+	// Verify that we triggered SequenceHistoryFlush on both ends
+	UE_NET_ASSERT_TRUE(Src.PacketNotify.IsWaitingForSequenceHistoryFlush());
+	UE_NET_ASSERT_TRUE(Dst.PacketNotify.IsWaitingForSequenceHistoryFlush());
+
+	// We a full round trip to reset SequenceHistory
+	Src.SendBurst(1);
+	Dst.SendBurst(1);
+	Src.ReceiveBurst(Dst.Writer, &Dst);
+	Dst.ReceiveBurst(Src.Writer, &Src);
+
+	Src.SendBurst(1);
+	Dst.SendBurst(1);
+	Src.ReceiveBurst(Dst.Writer, &Dst);
+	Dst.ReceiveBurst(Src.Writer, &Src);
+
+	// Verify that we recovered on both ends
+	UE_NET_ASSERT_FALSE(Src.PacketNotify.IsWaitingForSequenceHistoryFlush());
+	UE_NET_ASSERT_FALSE(Dst.PacketNotify.IsWaitingForSequenceHistoryFlush());
+
+	// Now we should be able to acknowledge data again
+	Src.SendAndDeliverTo(Dst);
+	Dst.SendAndDeliverTo(Src);
+
+	// Verify that we delivered what we expected
+	UE_NET_ASSERT_TRUE(Src.ValidateAcknowledgedPackets(Dst.AcceptedInPackets));
+	// Dst will not the latest packet accepted by Src.
+	UE_NET_ASSERT_TRUE(Dst.ValidateAcknowledgedPackets(MakeArrayView(Src.AcceptedInPackets.GetData(), Src.AcceptedInPackets.Num() - 1)));
+}
+
+UE_NET_TEST_FIXTURE(FNetPacketNotifyTestUtil, TestGetBothEndsWaitForSequenceWindowFlushDueToBurstPacketLoss)
+{
+	LOG_SCOPE_VERBOSITY_OVERRIDE(LogNetTraffic, ELogVerbosity::Error);
+
+	FNetPacketNotify::SequenceNumberT SrcStartingSeq(1024);
+	FNetPacketNotify::SequenceNumberT DstStartingSeq(0);
+
+	FTestNode Src(DstStartingSeq, SrcStartingSeq + 1);
+	FTestNode Dst(SrcStartingSeq, DstStartingSeq + 1);
+
+	// Send off bursts on both ends, over committing send window but do not deliver
+	Src.SendBurst(FNetPacketNotify::MaxSequenceHistoryLength + 1);
+	Dst.SendBurst(FNetPacketNotify::MaxSequenceHistoryLength + 1);
+
+	// Send new data that we will deliver, which should SequenceHistoryFlush but with no previously acked data we can do immediate reset
+	Src.SendBurst(1);
+	Dst.SendBurst(1);
+	
+	// Receive the data
+	Dst.ReceiveBurst(Src.Writer);
+	Src.ReceiveBurst(Dst.Writer);
+
+	// Verify that we managed to avoid requesting a flush
+	UE_NET_ASSERT_FALSE(Src.PacketNotify.IsWaitingForSequenceHistoryFlush());
+	UE_NET_ASSERT_FALSE(Dst.PacketNotify.IsWaitingForSequenceHistoryFlush());
+
+	// Verify that we delivered what we expected
+	UE_NET_ASSERT_TRUE(Src.ValidateAcceptedPackets({DstStartingSeq.Get() + FNetPacketNotify::MaxSequenceHistoryLength + 2U}));
+	UE_NET_ASSERT_TRUE(Dst.ValidateAcceptedPackets({SrcStartingSeq.Get() + FNetPacketNotify::MaxSequenceHistoryLength + 2U}));
+}
+
+UE_NET_TEST_FIXTURE(FNetPacketNotifyTestUtil, TestGetBothEndsWaitForSequenceWindowFlushWithPartiallyFullHistory)
+{
+	LOG_SCOPE_VERBOSITY_OVERRIDE(LogNetTraffic, ELogVerbosity::Error);
+
+	FNetPacketNotify::SequenceNumberT SrcStartingSeq(1024);
+	FNetPacketNotify::SequenceNumberT DstStartingSeq(0);
+
+	FTestNode Src(DstStartingSeq, SrcStartingSeq + 1);
+	FTestNode Dst(SrcStartingSeq, DstStartingSeq + 1);
+
+	// Send off a range with a packet to deliver and a bunch of missing ones, overshooting the sequence history
+	Src.SendRange({SrcStartingSeq.Get() + FNetPacketNotify::MaxSequenceHistoryLength/2U, SrcStartingSeq.Get() + FNetPacketNotify::MaxSequenceHistoryLength + 1U});
+	Dst.SendRange({DstStartingSeq.Get() + FNetPacketNotify::MaxSequenceHistoryLength/2U, DstStartingSeq.Get() + FNetPacketNotify::MaxSequenceHistoryLength + 1U});
+
+	// Receive the data
+	Dst.ReceiveBurst(Src.Writer);
+	Src.ReceiveBurst(Dst.Writer);
+
+	// We only expect the first packet to be accepted
+	Src.ExpectedAcceptedInPackets.Add(DstStartingSeq.Get() + FNetPacketNotify::MaxSequenceHistoryLength/2U);
+	Dst.ExpectedAcceptedInPackets.Add(SrcStartingSeq.Get() + FNetPacketNotify::MaxSequenceHistoryLength/2U);
+
+	// Verify that we delivered what we expected
+	UE_NET_ASSERT_TRUE(Src.ValidateExpectedAcceptedPackets());
+	UE_NET_ASSERT_TRUE(Dst.ValidateExpectedAcceptedPackets());
+
+	// Verify that we as expected triggered a SequenceHistoryFlush
+	UE_NET_ASSERT_TRUE(Src.PacketNotify.IsWaitingForSequenceHistoryFlush());
+	UE_NET_ASSERT_TRUE(Dst.PacketNotify.IsWaitingForSequenceHistoryFlush());
+
+	// Tell remote that we are flushing sequence history
+	Src.SendBurst(1);
+	Dst.SendBurst(1);
+	Dst.ReceiveBurst(Src.Writer, &Src);
+	Src.ReceiveBurst(Dst.Writer, &Dst);
+
+	// Verify that we discarded the packets as both ends are waiting for SequenceHistoryFlush
+	UE_NET_ASSERT_TRUE(Src.ValidateExpectedAcceptedPackets());
+	UE_NET_ASSERT_TRUE(Dst.ValidateExpectedAcceptedPackets());
+
+	// Send a burst, overshooting while in recovery (which should be fine)
+	Src.SendBurst(FNetPacketNotify::MaxSequenceHistoryLength);
+	Dst.SendBurst(FNetPacketNotify::MaxSequenceHistoryLength);
+
+	// We expect next sent packets to be delivered
+	Src.ExpectedAcceptedInPackets.Add(Dst.OutPacketId);
+	Dst.ExpectedAcceptedInPackets.Add(Src.OutPacketId);
+
+	// Deliver some data, completing the flush
+	Src.SendBurst(1);
+	Dst.SendBurst(1);
+	Dst.ReceiveBurst(Src.Writer, &Src);
+	Src.ReceiveBurst(Dst.Writer, &Dst);
+
+	// Verify that we managed to recover and are ready to proceed as normal
+	UE_NET_ASSERT_FALSE(Src.PacketNotify.IsWaitingForSequenceHistoryFlush());
+	UE_NET_ASSERT_FALSE(Dst.PacketNotify.IsWaitingForSequenceHistoryFlush());
+
+	// Verify that we delivered what we expected
+	UE_NET_ASSERT_TRUE(Src.ValidateExpectedAcceptedPackets());
+	UE_NET_ASSERT_TRUE(Dst.ValidateExpectedAcceptedPackets());
+}
+
+// Stress test to find random cases.
+#if 0
+UE_NET_TEST_FIXTURE(FNetPacketNotifyTestUtil, TestGetBothEndsWaitForSequenceWindowFlushRandomized)
+{
+	LOG_SCOPE_VERBOSITY_OVERRIDE(LogNetTraffic, ELogVerbosity::Verbose);
+
+	FNetPacketNotify::SequenceNumberT SrcStartingSeq((uint16)FMath::RandRange(0, (int32)FNetPacketNotify::SequenceNumberT::SeqNumberMax));
+	FNetPacketNotify::SequenceNumberT DstStartingSeq((uint16)FMath::RandRange(0, (int32)FNetPacketNotify::SequenceNumberT::SeqNumberMax));
+
+	FTestNode Src(DstStartingSeq, SrcStartingSeq + 1);;
+	FTestNode Dst(SrcStartingSeq, DstStartingSeq + 1);;
+	
+	SIZE_T TestIterations = 1024;
+
+	uint32 SrcHitFlushCount = 0U;
+	uint32 DstHitFlushCount = 0U;
+	uint32 SrcResetFlushCount = 0U;
+	uint32 DstResetFlushCount = 0U;
+	uint32 SrcStuckFlushCount = 0U;
+	uint32 DstStuckFlushCount = 0U;
+
+	FMath::SRandInit(347856243);
+	FMath::RandInit(347856243);
+
+	const uint32 DeltaBatch = 128; //32
+	while(TestIterations--)
+	{
+		// Verify that we as expected triggered a SequenceHistoryFlush
+		bool bSrcWasFlush = Src.PacketNotify.IsWaitingForSequenceHistoryFlush();
+		bool bDstWasFlush = Dst.PacketNotify.IsWaitingForSequenceHistoryFlush();
+
+		SrcStuckFlushCount += bSrcWasFlush ? 1U : 0U;
+		DstStuckFlushCount += bDstWasFlush ? 1U : 0U;
+
+		{
+			int32 NumPacketsToSend = FMath::RandRange(0, DeltaBatch);
+			Src.SendBurst(NumPacketsToSend);
+		}
+		{
+			int32 NumPacketsToSend = FMath::RandRange(0, DeltaBatch);
+			Dst.SendBurst(NumPacketsToSend);
+		}
+		if (FMath::RandBool())
+		{
+			// Receive and verify success
+			UE_NET_ASSERT_TRUE(Dst.ReceiveBurst(Src.Writer, &Src));
+		}
+		if (FMath::RandBool())
+		{
+			// Receive and verify success
+			UE_NET_ASSERT_TRUE(Src.ReceiveBurst(Dst.Writer, &Dst));
+		}
+		
+		if (Src.PacketNotify.IsWaitingForSequenceHistoryFlush())
+		{			
+			SrcHitFlushCount += bSrcWasFlush ? 0U : 1U;
+		}
+		else
+		{
+			SrcResetFlushCount += bSrcWasFlush ? 1U : 0U;			
+			SrcStuckFlushCount = 0U;
+		}
+		if (Dst.PacketNotify.IsWaitingForSequenceHistoryFlush())
+		{			
+			DstHitFlushCount += bDstWasFlush ? 0U : 1U;
+		}
+		else
+		{
+			DstResetFlushCount += bDstWasFlush ? 1U : 0U;
+			DstStuckFlushCount = 0U;
+		}
+
+		UE_NET_ASSERT_TRUE(SrcStuckFlushCount < 100);
+		UE_NET_ASSERT_TRUE(DstStuckFlushCount < 100);
+	};
+
+	Src.SendBurst(1);
+	Dst.SendBurst(1);
+	Dst.ReceiveBurst(Src.Writer, &Src);
+	Src.ReceiveBurst(Dst.Writer, &Dst);
+
+	Src.SendBurst(1);
+	Dst.SendBurst(1);
+	Dst.ReceiveBurst(Src.Writer, &Src);
+	Src.ReceiveBurst(Dst.Writer, &Dst);
+
+	// Verify that we acknowledged everything we accepted on both ends.
+	UE_NET_ASSERT_TRUE(Src.ValidateAcknowledgedPackets(MakeArrayView(Dst.AcceptedInPackets.GetData(), Dst.AcceptedInPackets.Num() - 1)));
+	UE_NET_ASSERT_TRUE(Dst.ValidateAcknowledgedPackets(MakeArrayView(Src.AcceptedInPackets.GetData(), Src.AcceptedInPackets.Num() - 1)));
+
+	// Verify that we managed to recover
+	UE_NET_ASSERT_FALSE(Src.PacketNotify.IsWaitingForSequenceHistoryFlush());
+	UE_NET_ASSERT_FALSE(Dst.PacketNotify.IsWaitingForSequenceHistoryFlush());
+}
+#endif
+
 }
 }
 
diff --git a/Engine/Source/Runtime/Core/Public/Misc/EngineNetworkCustomVersion.h b/Engine/Source/Runtime/Core/Public/Misc/EngineNetworkCustomVersion.h
--- a/Engine/Source/Runtime/Core/Public/Misc/EngineNetworkCustomVersion.h
+++ b/Engine/Source/Runtime/Core/Public/Misc/EngineNetworkCustomVersion.h
@@ -96,6 +96,7 @@ struct FEngineNetworkCustomVersion
 		ClientHandshakeId = 41,					// Bump version to support serialization of APlayerController::ClientHandeshakeId.
 		PawnRemoteViewPitchTo16Bit = 42,		// Bump version to support higher precision PawnRemoteViewPitch
 		CloseChildConnection = 43,				// Bump version to support closing child connections with NMT_CloseChildConnection.
+		ExplicitAckHistorySeq = 44,				// Bump version to add explicit sequence for ack history.
 
 		// -----<new versions can be added above this line>-------------------------------------------------
 		VersionPlusOne,
diff --git a/Engine/Source/Runtime/Engine/Classes/Engine/NetConnection.h b/Engine/Source/Runtime/Engine/Classes/Engine/NetConnection.h
--- a/Engine/Source/Runtime/Engine/Classes/Engine/NetConnection.h
+++ b/Engine/Source/Runtime/Engine/Classes/Engine/NetConnection.h
@@ -82,7 +82,7 @@ enum { RELIABLE_BUFFER = 512 }; // Power of 2 >= 1.
 enum { MAX_PACKETID = FNetPacketNotify::SequenceNumberT::SeqNumberCount };  // Power of 2 >= 1, covering guaranteed loss/misorder time.
 enum { MAX_CHSEQUENCE = 1024 }; // Power of 2 >RELIABLE_BUFFER, covering loss/misorder time.
 enum { MAX_BUNCH_HEADER_BITS = 256 };
-enum { MAX_PACKET_RELIABLE_SEQUENCE_HEADER_BITS = 32 /*PackedHeader*/ + FNetPacketNotify::SequenceHistoryT::MaxSizeInBits };
+enum { MAX_PACKET_RELIABLE_SEQUENCE_HEADER_BITS = 32 /*PackedHeader*/ + 16 /*HistorySequence*/ + FNetPacketNotify::SequenceHistoryT::MaxSizeInBits };
 enum { MAX_PACKET_INFO_HEADER_BITS = 1 /*bHasPacketInfo*/ + NetConnectionHelper::NumBitsForJitterClockTimeInHeader + 1 /*bHasServerFrameTime*/ + 8 /*ServerFrameTime*/};
 enum { MAX_PACKET_HEADER_BITS = MAX_PACKET_RELIABLE_SEQUENCE_HEADER_BITS + MAX_PACKET_INFO_HEADER_BITS  };
 enum { MAX_PACKET_TRAILER_BITS = 1 };
diff --git a/Engine/Source/Runtime/Engine/Private/Net/NetPacketNotify.cpp b/Engine/Source/Runtime/Engine/Private/Net/NetPacketNotify.cpp
--- a/Engine/Source/Runtime/Engine/Private/Net/NetPacketNotify.cpp
+++ b/Engine/Source/Runtime/Engine/Private/Net/NetPacketNotify.cpp
@@ -62,7 +62,6 @@ FNetPacketNotify::SequenceNumberT FNetPacketNotify::UpdateInAckSeqAck(SequenceNu
 		// verify that we have a matching sequence number
 		if (AckData.OutSeq == AckedSeq)
 		{
-			UE_LOG_PACKET_NOTIFY(TEXT("FNetPacketNotify::UpdateInAckSeqAck - InAckSeqAck: %u"), AckData.InAckSeq.Get());
 			return AckData.InAckSeq;
 		}
 	}
@@ -93,7 +92,6 @@ void FNetPacketNotify::SetWaitForSequenceHistoryFlush()
 void FNetPacketNotify::AckSeq(SequenceNumberT AckedSeq, bool IsAck)
 {
 	check( AckedSeq == InSeq);
-
 	while (AckedSeq > InAckSeq)
 	{
 		++InAckSeq;
@@ -102,68 +100,53 @@ void FNetPacketNotify::AckSeq(SequenceNumberT AckedSeq, bool IsAck)
 
 		UE_LOG_PACKET_NOTIFY(TEXT("FNetPacketNotify::AckSeq - AckedSeq: %u, IsAck %u, AckHistorySize: %d"), InAckSeq.Get(), bReportAcked ? 1u : 0u, GetCurrentSequenceHistoryLength());
 
-		InSeqHistory.AddDeliveryStatus(bReportAcked);		
+		InSeqHistory.AddDeliveryStatus(bReportAcked);
 	}
 }
 
 FNetPacketNotify::SequenceNumberT::DifferenceT FNetPacketNotify::InternalUpdate(const FNotificationHeader& NotificationData, SequenceNumberT::DifferenceT InSeqDelta)
 {
-	// We must check if we will overflow our outgoing ack-window, if we do and it contains processed data we must initiate a re-sync of ack-sequence history.
-	// This is done by ignoring any new packets until we are in sync again. 
-	// This would typically only occur in situations where we would have had huge packet loss or spikes on the receiving end.
-	if (!IsWaitingForSequenceHistoryFlush() && !WillSequenceFitInSequenceHistory(NotificationData.Seq))
+	if (!WillSequenceFitInSequenceHistory(NotificationData.Seq) && !IsWaitingForSequenceHistoryFlush())
 	{
+		// If we will overflow our outgoing ack sequence history and it contains accepted processed packets we must initiate a re-sync of ack sequence history.
+		// This is done by ignoring any new packets until we are in sync again.
+		// This would typically only occur in situations where we would have had huge packet loss or spikes on the receiving end.
 		if (GetHasUnacknowledgedAcks())
 		{
 			SetWaitForSequenceHistoryFlush();
+
+			// Mark everything we can as lost up until the end of the sequence history, as we will not accept anything until we are confirmed as flushed
+			SequenceNumberT NewInSeqToAck(NotificationData.Seq);
+			NewInSeqToAck = SequenceNumberT(InAckSeqAck.Get() + MaxSequenceHistoryLength);
+
+			if (NewInSeqToAck >= InAckSeq)
+			{
+				// Need to set InSeq to clamped value before naks
+				InSeq = NewInSeqToAck;
+				AckSeq(NewInSeqToAck, false);
+
+				UE_LOG_PACKET_NOTIFY(TEXT("FNetPacketNotify::Update - Waiting for sequence history flush - Rejected: %u Last InAckSeq: %u"), NotificationData.Seq.Get(), InAckSeq.Get());
+			}
 		}
 		else
 		{
-			// We can reset if we have no previous acks and can then safely synthesize nacks on the receiving end
+			// We can reset if we have no previous acks and then can safely synthesize nacks on the receiving end
 			const SequenceNumberT NewInAckSeqAck(NotificationData.Seq.Get() - 1);
-			UE_LOG_PACKET_NOTIFY_WARNING(TEXT("FNetPacketNotify::Reset SequenceHistory - New InSeqDelta: %u Old: %u"), NewInAckSeqAck.Get(), InAckSeqAck.Get());
+			UE_LOG_PACKET_NOTIFY_WARNING(TEXT("FNetPacketNotify::Reset SequenceHistory - As it is empty. New InAckSeqAck: %u Old: %u"), NewInAckSeqAck.Get(), InAckSeqAck.Get());
 			InAckSeqAck = NewInAckSeqAck;
+			InSeqHistory.Reset();			
 		}
 	}
 
-	if (!IsWaitingForSequenceHistoryFlush())
-	{
-		// Just accept the incoming sequence, under normal circumstances NetConnection explicitly handles the acks.
-		InSeq = NotificationData.Seq;
+	// Update InSeq
+	InSeq = NotificationData.Seq;
 
-		return InSeqDelta;
-	}
-	else
+	if (IsWaitingForSequenceHistoryFlush())
 	{
-		// Until we have flushed the history we treat incoming packets as lost while still advancing ack window as far as we can.
-		SequenceNumberT NewInSeqToAck(NotificationData.Seq);
-
-		// Still waiting on flush, but we can fill up the history
-		if (!WillSequenceFitInSequenceHistory(NotificationData.Seq) && GetHasUnacknowledgedAcks())
-		{
-			// Mark everything we can as lost up until the end of the sequence history 
-			NewInSeqToAck = SequenceNumberT(InAckSeqAck.Get() + (MaxSequenceHistoryLength - GetCurrentSequenceHistoryLength()));
-		}
-
-		if (NewInSeqToAck >= InSeq)
-		{
-			const SequenceNumberT::DifferenceT AdjustedSequenceDelta = SequenceNumberT::Diff(NewInSeqToAck, InSeq);
-
-			InSeq = NewInSeqToAck;
-
-			// Nack driven from here
-			AckSeq(NewInSeqToAck, false);
-
-			UE_LOG_PACKET_NOTIFY(TEXT("FNetPacketNotify::Update - Waiting for sequence history flush - Rejected: %u Accepted: InSeq: %u Adjusted delta %d"), NotificationData.Seq.Get(), InSeq.Get(), AdjustedSequenceDelta);
-			
-			return AdjustedSequenceDelta;
-		}
-		else
-		{
-			UE_LOG_PACKET_NOTIFY(TEXT("FNetPacketNotify::Update - Waiting for sequence history flush - Rejected: %u Accepted: InSeq: %u"), NotificationData.Seq.Get(), InSeq.Get());
-			return 0;
-		}
+		UE_LOG_PACKET_NOTIFY(TEXT("FNetPacketNotify::Update - Waiting for sequence history flush - Ignoring current InSeq: %u Last acked InSeq: %u"), NotificationData.Seq.Get(), InAckSeq.Get());
 	}
+
+	return InSeqDelta;
 };
 
 namespace 
@@ -179,11 +162,21 @@ namespace
 		enum { HistoryWordCountMask	= (1 << HistoryWordCountBits) - 1 };
 		enum { AckSeqShift			= HistoryWordCountBits };
 		enum { SeqShift				= AckSeqShift + FNetPacketNotify::SequenceNumberBits };
+		enum { ResetHistoryValue	= HistoryWordCountMask };
 		
-		static_assert(FNetPacketNotify::SequenceHistoryT::WordCount <= (1 << HistoryWordCountBits), "Sequence history word count (minus one) must fit within HistoryWordCountBits. Check FNetPacketNotify::MaxSequenceHistoryLength.");
+		// Validate expected WordCount, we reserve the last word in order to indicate if we are trying to reset the ack window
+		static_assert(FNetPacketNotify::SequenceHistoryT::WordCount <= ((1 << HistoryWordCountBits) -1), "Sequence history word count (minus one) must fit within HistoryWordCountBits. Check FNetPacketNotify::MaxSequenceHistoryLength.");
 
-		static uint32 Pack(SequenceNumberT Seq, SequenceNumberT AckedSeq, SIZE_T HistoryWordCount)
+		static uint32 Pack(SequenceNumberT Seq, SequenceNumberT AckedSeq, SIZE_T HistoryWordCount, bool bResetHistory)
 		{
+			// We reserve ResetHistoryValue to express if we are currently resetting sequence ack history
+			check(HistoryWordCount != ResetHistoryValue);
+			if (bResetHistory)
+			{
+				check(HistoryWordCount + 1 == FNetPacketNotify::SequenceHistoryT::WordCount);
+				HistoryWordCount = ResetHistoryValue;
+			}
+
 			uint32 Packed = 0u;
 
 			Packed |= Seq.Get() << SeqShift;
@@ -202,9 +195,16 @@ namespace
 // These methods must always write and read the exact same number of bits, that is the reason for not using WriteInt/WrittedWrappedInt
 bool FNetPacketNotify::WriteHeader(FBitWriter& Writer, bool bRefresh)
 {
-	// we always write at least 1 word
+	// We always write at least 1 word
 	SIZE_T CurrentHistoryWordCount = FMath::Clamp<SIZE_T>((GetCurrentSequenceHistoryLength() + SequenceHistoryT::BitsPerWord - 1u) / SequenceHistoryT::BitsPerWord, 1u, SequenceHistoryT::WordCount);
 
+	// If we are waiting for SequenceHistoryFlush we will always write the full sequence history
+	const bool bIsWaitingForSequenceHistoryFlush = IsWaitingForSequenceHistoryFlush();
+	if (bIsWaitingForSequenceHistoryFlush && CurrentHistoryWordCount != SequenceHistoryT::WordCount)
+	{
+		CurrentHistoryWordCount = SequenceHistoryT::WordCount;
+	}
+
 	// We can only do a refresh if we do not need more space for the history
 	if (bRefresh && (CurrentHistoryWordCount > WrittenHistoryWordCount))
 	{
@@ -213,14 +213,15 @@ bool FNetPacketNotify::WriteHeader(FBitWriter& Writer, bool bRefresh)
 
 	// How many words of ack data should we write? If this is a refresh we must write the same size as the original header
 	WrittenHistoryWordCount = bRefresh ? WrittenHistoryWordCount : CurrentHistoryWordCount;
-	// This is the last InAck we have acknowledged at this time
+
+	// This is the last InSeq we have acknowledged at this time
 	WrittenInAckSeq = InAckSeq;
 
 	SequenceNumberT::SequenceT Seq = OutSeq.Get();
-	SequenceNumberT::SequenceT AckedSeq = InAckSeq.Get();
+	SequenceNumberT::SequenceT AckedSeq = InSeq.Get();
 
 	// Pack data into a uint
-	uint32 PackedHeader = FPackedHeader::Pack(Seq, AckedSeq, WrittenHistoryWordCount - 1);
+	uint32 PackedHeader = FPackedHeader::Pack(Seq, AckedSeq, WrittenHistoryWordCount - 1, bIsWaitingForSequenceHistoryFlush);
 
 	// Write packed header
 	Writer << PackedHeader;
@@ -228,7 +229,11 @@ bool FNetPacketNotify::WriteHeader(FBitWriter& Writer, bool bRefresh)
 	// Write ack history
 	InSeqHistory.Write(Writer, WrittenHistoryWordCount);
 
-	UE_LOG_PACKET_NOTIFY(TEXT("FNetPacketNotify::WriteHeader - Seq %u, AckedSeq %u bReFresh %u HistorySizeInWords %" SIZE_T_FMT), Seq, AckedSeq, bRefresh ? 1u : 0u, WrittenHistoryWordCount);
+	// Write AckedHistorySeq
+	uint16 AckedHistorySeq = InAckSeq.Get();
+	Writer << AckedHistorySeq;
+
+	UE_LOG_PACKET_NOTIFY(TEXT("FNetPacketNotify::WriteHeader - Seq %u, AckedSeq %u bReFresh %u bWaitingForSequenceHistoryFlush %u AckedHistorySeq %u HistorySizeInWords %" SIZE_T_FMT), Seq, AckedSeq, bRefresh ? 1u : 0u, bIsWaitingForSequenceHistoryFlush ? 1u : 0u, AckedHistorySeq, WrittenHistoryWordCount);
 
 	return true;
 }
@@ -239,15 +244,30 @@ bool FNetPacketNotify::ReadHeader(FNotificationHeader& Data, FBitReader& Reader)
 	uint32 PackedHeader = 0;	
 	Reader << PackedHeader;
 
-	// unpack
+	// Unpack
 	Data.Seq = FPackedHeader::GetSeq(PackedHeader);
 	Data.AckedSeq = FPackedHeader::GetAckedSeq(PackedHeader);
-	Data.HistoryWordCount = FPackedHeader::GetHistoryWordCount(PackedHeader) + 1;
-
+	const SIZE_T HistoryWordCount = FPackedHeader::GetHistoryWordCount(PackedHeader);
+	if (HistoryWordCount == FPackedHeader::ResetHistoryValue)
+	{
+		Data.HistoryWordCount = SequenceHistoryT::WordCount;
+		Data.bRemoteIsWaitingForHistoryReset = true;
+	}
+	else
+	{
+		Data.HistoryWordCount = FPlatformMath::Min(HistoryWordCount + 1, SequenceHistoryT::WordCount);
+		Data.bRemoteIsWaitingForHistoryReset = false;
+	}
+	
 	// Read ack history
 	Data.History.Read(Reader, Data.HistoryWordCount);
 
-	UE_LOG_PACKET_NOTIFY(TEXT("FNetPacketNotify::ReadHeader - Seq %u, AckedSeq %u HistorySizeInWords %" SIZE_T_FMT), Data.Seq.Get(), Data.AckedSeq.Get(), Data.HistoryWordCount);
+	// Read AckedHistorySeq
+	uint16 AckedHistorySeq = Data.AckedSeq.Get();
+	Reader << AckedHistorySeq;
+	Data.AckedHistorySeq = SequenceNumberT(AckedHistorySeq);
+
+	UE_LOG_PACKET_NOTIFY(TEXT("FNetPacketNotify::ReadHeader - Seq %u, AckedSeq %u RemoteWaitingForFlush %u AckedHistorySeq %u HistorySizeInWords %" SIZE_T_FMT), Data.Seq.Get(), Data.AckedSeq.Get(), Data.bRemoteIsWaitingForHistoryReset, Data.AckedHistorySeq.Get(), Data.HistoryWordCount);
 
 	return Reader.IsError() == false;
 }
@@ -258,7 +278,7 @@ FNetPacketNotify::SequenceNumberT FNetPacketNotify::CommitAndIncrementOutSeq()
 	check(WrittenHistoryWordCount != 0);
 
 	// Add entry to the ack-record so that we can update the InAckSeqAck when we received the ack for this OutSeq.
-	AckRecord.Enqueue( {OutSeq, WrittenInAckSeq} );
+	AckRecord.Enqueue( {OutSeq, WrittenInAckSeq } );
 	WrittenHistoryWordCount = 0u;
 	
 	return ++OutSeq;
diff --git a/Engine/Source/Runtime/Engine/Public/Net/NetPacketNotify.h b/Engine/Source/Runtime/Engine/Public/Net/NetPacketNotify.h
--- a/Engine/Source/Runtime/Engine/Public/Net/NetPacketNotify.h
+++ b/Engine/Source/Runtime/Engine/Public/Net/NetPacketNotify.h
@@ -46,10 +46,18 @@ public:
 
 	struct FNotificationHeader
 	{
+		// Bitfield with acknowledged packets.
 		SequenceHistoryT History;
-		SIZE_T HistoryWordCount;
+		// Size of history in words
+		SIZE_T HistoryWordCount = 0;
+		// Incoming sequence number
 		SequenceNumberT Seq;
+		// Latest sequence that the remote has received. Note: It does not mean that the packet was accepted.
 		SequenceNumberT AckedSeq;
+		// Sequence number of the first entry in the sequence history, normally the same as AckedSeq
+		SequenceNumberT AckedHistorySeq;
+		// True if remote is flushing history due to sequence history overflow.
+		bool bRemoteIsWaitingForHistoryReset = false;
 	};
 
 	/** Constructor */
@@ -145,7 +153,7 @@ private:
 
 	// Track incoming sequence data
 	SequenceHistoryT InSeqHistory;		// BitBuffer containing a bitfield describing the history of received packets
-	SequenceNumberT InSeq;				// Last sequence number received and accepted from remote
+	SequenceNumberT InSeq;				// Last sequence number received from remote
 	SequenceNumberT InAckSeq;			// Last sequence number received from remote that we have acknowledged, this is needed since we support accepting a packet but explicitly not acknowledge it as received.
 	SequenceNumberT InAckSeqAck;		// Last sequence number received from remote that we have acknowledged and also knows that the remote has received the ack, used to calculate how big our history must be
 	SequenceNumberT WaitingForFlushSeqAck;
@@ -207,18 +215,21 @@ void FNetPacketNotify::ProcessReceivedAcks(const FNotificationHeader& Notificati
 {
 	if (NotificationData.AckedSeq > OutAckSeq)
 	{
+		// Total ack count including implicit naks
 		SequenceNumberT::DifferenceT AckCount = SequenceNumberT::Diff(NotificationData.AckedSeq, OutAckSeq);
-
-		UE_LOG_PACKET_NOTIFY(TEXT("Notification::ProcessReceivedAcks - AckedSeq: %u, OutAckSeq: %u AckCount: %u"), NotificationData.AckedSeq.Get(), OutAckSeq.Get(), AckCount);
+		UE_LOG_PACKET_NOTIFY(TEXT("Notification::ProcessReceivedAcks - AckedSeq: %u, OutAckSeq: %u AckCount: %u RemoteWaitingForFlush: %u"), NotificationData.AckedSeq.Get(), OutAckSeq.Get(), AckCount, NotificationData.bRemoteIsWaitingForHistoryReset ? 1U : 0U);
 
 		// Update InAckSeqAck used to track the needed number of bits to transmit our ack history
 		// Note: As we might reset sequence history we need to check if we already have advanced the InAckSeqAck
-		const SequenceNumberT NewInAckSeqAck = UpdateInAckSeqAck(AckCount, NotificationData.AckedSeq);
-		if (NewInAckSeqAck > InAckSeqAck)
 		{
-			InAckSeqAck = NewInAckSeqAck;
+			const SequenceNumberT NewInAckSeqAck = UpdateInAckSeqAck(AckCount, NotificationData.AckedSeq);
+			if (NewInAckSeqAck > InAckSeqAck)
+			{
+				UE_LOG_PACKET_NOTIFY(TEXT("FNetPacketNotify::ProcessReceivedAcks - Advance InAckSeqAck: %u"), NewInAckSeqAck.Get());
+				InAckSeqAck = NewInAckSeqAck;
+			}
 		}
-		
+
 		// ExpectedAck = OutAckSeq + 1
 		SequenceNumberT CurrentAck(OutAckSeq);
 		++CurrentAck;
@@ -228,37 +239,90 @@ void FNetPacketNotify::ProcessReceivedAcks(const FNotificationHeader& Notificati
 		const SequenceNumberT::DifferenceT HistoryBits = static_cast<SequenceNumberT::DifferenceT>(NotificationData.HistoryWordCount * SequenceHistoryT::BitsPerWord);
 
 		// Warn if the received sequence number is greater than our history buffer, since if that is the case we have to treat the data as lost
-		// Note: This should normally not be a problem since we try to flush the sequence history before allowing an overshoot of the sequence history window on the sending side.
+		// Note: This should not be a problem as we have a mechanism to re-sync history
 		// If this occurs with no hitches on server or client, there might be reason to investigate if too much data is being sent in which case the the size sequence history might have to be increased.
-		if (AckCount > (SequenceNumberT::DifferenceT)(SequenceHistoryT::Size))
+		SequenceNumberT::DifferenceT MissingAckCount = 0U;
+		if (AckCount > HistoryBits)
 		{
-			UE_LOG_PACKET_NOTIFY_WARNING(TEXT("FNetPacketNotify::ProcessReceivedAcks - Missed Acks: AckedSeq: %u, OutAckSeq: %u, FirstMissingSeq: %u Count: %u"), NotificationData.AckedSeq.Get(), OutAckSeq.Get(), CurrentAck.Get(), AckCount - (SequenceNumberT::DifferenceT)(SequenceHistoryT::Size));
+			MissingAckCount = AckCount - HistoryBits;
+			UE_LOG_PACKET_NOTIFY(TEXT("FNetPacketNotify::ProcessReceivedAcks - Missed Acks: AckedSeq: %u, OutAckSeq: %u, FirstMissingSeq: %u Count: %u"), NotificationData.AckedSeq.Get(), OutAckSeq.Get(), CurrentAck.Get(), MissingAckCount);
 		}
 
-		// Everything not found in the history buffer is treated as lost
-		while (AckCount > HistoryBits)
+		if (NotificationData.bRemoteIsWaitingForHistoryReset)
 		{
-			--AckCount;
-			InFunc(CurrentAck, false);
-			++CurrentAck;
+			// Re-sync sequence history
+			// Extra care has to be taken, we might have to synthesize naks both before and after the acks carried in the sequence history.
+			if (NotificationData.AckedHistorySeq > OutAckSeq)
+			{
+				SequenceNumberT::DifferenceT AckCountInHistory = SequenceNumberT::Diff(NotificationData.AckedHistorySeq, OutAckSeq);
+
+				// Synthesize missing naks before start of sequence history, this can occur if we get burst packet drops and then overflow history
+				while (AckCountInHistory > HistoryBits)
+				{
+					--AckCountInHistory;
+					--AckCount;
+					InFunc(CurrentAck, false);
+					++CurrentAck;
+				}
+
+				// Acknowledge packets based on information from sequence history
+				while (AckCountInHistory > 0)
+				{
+					--AckCountInHistory;
+					--AckCount;
+
+					const uint32 AckIndex = AckCountInHistory;
+
+					UE_LOG_PACKET_NOTIFY(TEXT("Notification::ProcessReceivedAcks Seq: %u - IsAck: %u HistoryIndex: %u"), CurrentAck.Get(), NotificationData.History.IsDelivered(AckIndex) ? 1u : 0u, AckIndex);
+					InFunc(CurrentAck, NotificationData.History.IsDelivered(AckIndex));
+					++CurrentAck;
+				}
+			}
+
+			// Synthesize missing naks after history window as the packets are discarded during re-sync
+			while (AckCount > 0)
+			{
+				--AckCount;
+				InFunc(CurrentAck, false);
+				++CurrentAck;
+			}
 		}
-
-		// For sequence numbers contained in the history we lookup the delivery status from the history
-		while (AckCount > 0)
+		else
 		{
-			--AckCount;
-			UE_LOG_PACKET_NOTIFY(TEXT("Notification::ProcessReceivedAcks Seq: %u - IsAck: %u HistoryIndex: %u"), CurrentAck.Get(), NotificationData.History.IsDelivered(AckCount) ? 1u : 0u, AckCount);
-			InFunc(CurrentAck, NotificationData.History.IsDelivered(AckCount));
-			++CurrentAck;
+			// Normal path, missing naks are at the beginning
+
+			// Synthesize missing naks before start of history window, this can occur if we get burst packet drops with no pending acks
+			while (AckCount > HistoryBits)
+			{
+				--AckCount;
+				InFunc(CurrentAck, false);
+				++CurrentAck;
+			}
+
+			// For sequence numbers contained in the history we lookup the delivery status from the history
+			while (AckCount > 0)
+			{
+				--AckCount;
+				UE_LOG_PACKET_NOTIFY(TEXT("Notification::ProcessReceivedAcks Seq: %u - IsAck: %u HistoryIndex: %u"), CurrentAck.Get(), NotificationData.History.IsDelivered(AckCount) ? 1u : 0u, AckCount);
+				InFunc(CurrentAck, NotificationData.History.IsDelivered(AckCount));
+				++CurrentAck;
+			}
 		}
+			
+		// Reset history if we are done with reset of sequence history.
+		if (IsWaitingForSequenceHistoryFlush() && NotificationData.AckedSeq >= WaitingForFlushSeqAck)
+		{
+			UE_LOG_PACKET_NOTIFY(TEXT("ClearIsWaitingForSequenceHistoryFlush %u"), NotificationData.AckedSeq.Get());
+			InSeqHistory.Reset();
+			WaitingForFlushSeqAck = NotificationData.AckedSeq;
+		}
+
 		OutAckSeq = NotificationData.AckedSeq;
 
-		// Are we done waiting for an reset of the ack history?
+		// Deal with sequence wraparound.
 		if (OutAckSeq > WaitingForFlushSeqAck)
 		{
 			WaitingForFlushSeqAck = OutAckSeq;
 		}
 	}
 }
-
-
