Iris
- Fixed som issues reported by BugHawk

diff --git a/Engine/Plugins/Runtime/ReplicationSystemTestPlugin/Source/Private/Tests/Core/TestNetPacketNotify.cpp b/Engine/Plugins/Runtime/ReplicationSystemTestPlugin/Source/Private/Tests/Core/TestNetPacketNotify.cpp
--- a/Engine/Plugins/Runtime/ReplicationSystemTestPlugin/Source/Private/Tests/Core/TestNetPacketNotify.cpp
+++ b/Engine/Plugins/Runtime/ReplicationSystemTestPlugin/Source/Private/Tests/Core/TestNetPacketNotify.cpp
@@ -793,8 +793,8 @@ UE_NET_TEST_FIXTURE(FNetPacketNotifyTestUtil, TestGetBothEndsWaitForSequenceWind
 	// Validate expected state
 	UE_NET_ASSERT_TRUE(Src.ValidateAcknowledgedPackets({0, FNetPacketNotify::MaxSequenceHistoryLength + 1, FNetPacketNotify::MaxSequenceHistoryLength + 2}));
 	UE_NET_ASSERT_TRUE(Src.ValidateAcceptedPackets({0, FNetPacketNotify::MaxSequenceHistoryLength + 1, FNetPacketNotify::MaxSequenceHistoryLength + 2, FNetPacketNotify::MaxSequenceHistoryLength + 3}));
-	UE_NET_ASSERT_TRUE(Src.ValidateAcknowledgedPackets({0, FNetPacketNotify::MaxSequenceHistoryLength + 1, FNetPacketNotify::MaxSequenceHistoryLength + 2}));
-	UE_NET_ASSERT_TRUE(Src.ValidateAcceptedPackets({0, FNetPacketNotify::MaxSequenceHistoryLength + 1, FNetPacketNotify::MaxSequenceHistoryLength + 2, FNetPacketNotify::MaxSequenceHistoryLength + 3}));
+	UE_NET_ASSERT_TRUE(Dst.ValidateAcknowledgedPackets({0, FNetPacketNotify::MaxSequenceHistoryLength + 1, FNetPacketNotify::MaxSequenceHistoryLength + 2}));
+	UE_NET_ASSERT_TRUE(Dst.ValidateAcceptedPackets({0, FNetPacketNotify::MaxSequenceHistoryLength + 1, FNetPacketNotify::MaxSequenceHistoryLength + 2, FNetPacketNotify::MaxSequenceHistoryLength + 3}));
 }
 
 UE_NET_TEST_FIXTURE(FNetPacketNotifyTestUtil, TestGetBothEndsWaitForSequenceWindowFlushDueToSendWindowOverflow)
diff --git a/Engine/Source/Runtime/Engine/Private/Net/NetPacketNotify.cpp b/Engine/Source/Runtime/Engine/Private/Net/NetPacketNotify.cpp
--- a/Engine/Source/Runtime/Engine/Private/Net/NetPacketNotify.cpp
+++ b/Engine/Source/Runtime/Engine/Private/Net/NetPacketNotify.cpp
@@ -267,7 +267,7 @@ bool FNetPacketNotify::ReadHeader(FNotificationHeader& Data, FBitReader& Reader)
 	Reader << AckedHistorySeq;
 	Data.AckedHistorySeq = SequenceNumberT(AckedHistorySeq);
 
-	UE_LOG_PACKET_NOTIFY(TEXT("FNetPacketNotify::ReadHeader - Seq %u, AckedSeq %u RemoteWaitingForFlush %u AckedHistorySeq %u HistorySizeInWords %" SIZE_T_FMT), Data.Seq.Get(), Data.AckedSeq.Get(), Data.bRemoteIsWaitingForHistoryReset, Data.AckedHistorySeq.Get(), Data.HistoryWordCount);
+	UE_LOG_PACKET_NOTIFY(TEXT("FNetPacketNotify::ReadHeader - Seq %u, AckedSeq %u RemoteWaitingForFlush %u AckedHistorySeq %u HistorySizeInWords %" SIZE_T_FMT), Data.Seq.Get(), Data.AckedSeq.Get(), Data.bRemoteIsWaitingForHistoryReset ? 1U : 0U, Data.AckedHistorySeq.Get(), Data.HistoryWordCount);
 
 	return Reader.IsError() == false;
 }
