Iris
- Exposed experimental method to query relevancy of a replicated object for a specific connection.

diff --git a/Engine/Plugins/Runtime/ReplicationSystemTestPlugin/Source/Private/Tests/ReplicationSystem/Filtering/TestFiltering.cpp b/Engine/Plugins/Runtime/ReplicationSystemTestPlugin/Source/Private/Tests/ReplicationSystem/Filtering/TestFiltering.cpp
--- a/Engine/Plugins/Runtime/ReplicationSystemTestPlugin/Source/Private/Tests/ReplicationSystem/Filtering/TestFiltering.cpp
+++ b/Engine/Plugins/Runtime/ReplicationSystemTestPlugin/Source/Private/Tests/ReplicationSystem/Filtering/TestFiltering.cpp
@@ -2742,4 +2742,65 @@ UE_NET_TEST_FIXTURE(FTestFilteringFixture, OwningConnectionIsSetProperly)
 	UE_NET_ASSERT_EQ(Server->ReplicationSystem->GetOwningNetConnection(ServerObject2SubObject->NetRefHandle), Client->ConnectionIdOnServer);
 }
 
+UE_NET_TEST_FIXTURE(FTestFilteringFixture, TestIsRelevantForConnection)
+{
+	// Add client
+	FReplicationSystemTestClient* Client = CreateClient();
+
+	// Spawn object on server
+	UReplicatedTestObject* ServerObject = Server->CreateObject(0, 0);
+
+	// Send and deliver packets
+	Server->NetUpdate();
+	Server->SendAndDeliverTo(Client, DeliverPacket);
+	Server->PostSendUpdate();
+
+	UE_NET_ASSERT_TRUE(Server->ReplicationSystem->IsRelevantForConnection(ServerObject->NetRefHandle, Client->ConnectionIdOnServer));
+
+	// Object should now exist on client
+	UE_NET_ASSERT_NE(Client->GetReplicationBridge()->GetReplicatedObject(ServerObject->NetRefHandle), nullptr);
+
+	// Turn on owner filter
+	Server->ReplicationSystem->SetFilter(ServerObject->NetRefHandle, ToOwnerFilterHandle);
+
+	// Send and deliver packets
+	Server->NetUpdate();
+	Server->SendAndDeliverTo(Client, DeliverPacket);
+	Server->PostSendUpdate();
+
+	UE_NET_ASSERT_FALSE(Server->ReplicationSystem->IsRelevantForConnection(ServerObject->NetRefHandle, Client->ConnectionIdOnServer));
+
+	// As object is now filtered it should be deleted on the client
+	UE_NET_ASSERT_EQ(Client->GetReplicationBridge()->GetReplicatedObject(ServerObject->NetRefHandle), nullptr);
+
+	// Set the client as owner
+	Server->ReplicationSystem->SetOwningNetConnection(ServerObject->NetRefHandle, Client->ConnectionIdOnServer);
+
+	// Send and deliver packets
+	Server->NetUpdate();
+	Server->SendAndDeliverTo(Client, DeliverPacket);
+	Server->PostSendUpdate();
+
+	UE_NET_ASSERT_TRUE(Server->ReplicationSystem->IsRelevantForConnection(ServerObject->NetRefHandle, Client->ConnectionIdOnServer));
+
+	// The object should have been created again
+	UE_NET_ASSERT_NE(Client->GetReplicationBridge()->GetReplicatedObject(ServerObject->NetRefHandle), nullptr);
+
+	// Finally, remove the owning connection
+	Server->ReplicationSystem->SetOwningNetConnection(ServerObject->NetRefHandle, 0);
+
+	// Send and deliver packets
+	Server->NetUpdate();
+	Server->SendAndDeliverTo(Client, DeliverPacket);
+	Server->PostSendUpdate();
+
+	UE_NET_ASSERT_FALSE(Server->ReplicationSystem->IsRelevantForConnection(ServerObject->NetRefHandle, Client->ConnectionIdOnServer));
+
+	// The client is no longer owning the object so it should be deleted
+	UE_NET_ASSERT_EQ(Client->GetReplicationBridge()->GetReplicatedObject(ServerObject->NetRefHandle), nullptr);
+
+	Server->DestroyObject(ServerObject);
+}
+
+
 } // end namespace UE::Net::Private
diff --git a/Engine/Source/Runtime/Net/Iris/Private/Iris/ReplicationSystem/ReplicationSystem.cpp b/Engine/Source/Runtime/Net/Iris/Private/Iris/ReplicationSystem/ReplicationSystem.cpp
--- a/Engine/Source/Runtime/Net/Iris/Private/Iris/ReplicationSystem/ReplicationSystem.cpp
+++ b/Engine/Source/Runtime/Net/Iris/Private/Iris/ReplicationSystem/ReplicationSystem.cpp
@@ -1307,6 +1307,30 @@ const UDataStream* UReplicationSystem::GetDataStream(uint32 ConnectionId, FName
 	return nullptr;
 }
 
+bool UReplicationSystem::IsRelevantForConnection(FNetRefHandle Handle, uint32 ConnectionId) const
+{
+	using namespace UE::Net;
+	using namespace UE::Net::Private;
+
+	const FReplicationConnections& Connections = Impl->ReplicationSystemInternal.GetConnections();
+	if (!ensureMsgf(Connections.IsValidConnection(ConnectionId), TEXT("Invalid ConnectionId %u passed to UReplicationSystem::GetDataStream."), ConnectionId))
+	{
+		return false;
+	}
+
+	const FNetRefHandleManager& NetRefHandleManager = Impl->ReplicationSystemInternal.GetNetRefHandleManager();
+	const FInternalNetRefIndex ObjectInternalIndex = NetRefHandleManager.GetInternalIndex(Handle);
+	if (ObjectInternalIndex == FNetRefHandleManager::InvalidInternalIndex)
+	{
+		return false;
+	}
+
+	const FReplicationFiltering& Filtering = Impl->ReplicationSystemInternal.GetFiltering();
+	const  FNetBitArrayView ObjectsInScope = Filtering.GetRelevantObjectsInScope(ConnectionId);
+
+	return ObjectsInScope.GetNumBits() > ObjectInternalIndex ? ObjectsInScope.IsBitSet(ObjectInternalIndex) : false;
+}
+
 void UReplicationSystem::CloseDataStream(uint32 ConnectionId, FName Name)
 {
 	UE::Net::Private::FReplicationConnections& Connections = Impl->ReplicationSystemInternal.GetConnections();
diff --git a/Engine/Source/Runtime/Net/Iris/Public/Iris/ReplicationSystem/ReplicationSystem.h b/Engine/Source/Runtime/Net/Iris/Public/Iris/ReplicationSystem/ReplicationSystem.h
--- a/Engine/Source/Runtime/Net/Iris/Public/Iris/ReplicationSystem/ReplicationSystem.h
+++ b/Engine/Source/Runtime/Net/Iris/Public/Iris/ReplicationSystem/ReplicationSystem.h
@@ -744,6 +744,10 @@ public:
 
 	IRISCORE_API UE::Net::FReplicationSystemDelegates& GetDelegates();
 
+	/** Function to query if a particular replicated object is relevant to a connection or not. */
+	UE_EXPERIMENTAL(5.7, "IsRelevantForConnection() is experimental and should not be used.")
+	IRISCORE_API bool IsRelevantForConnection(FNetRefHandle NetRefHandle, uint32 ConnectionId) const;
+
 public:
 	// For internal use and not exported.
 	UE::Net::Private::FReplicationSystemInternal* GetReplicationSystemInternal();
