Iris - Fix bugs in ReplicationFiltering's GetDependentObjectFilterStatus.

diff --git a/Engine/Plugins/Runtime/ReplicationSystemTestPlugin/Source/Private/Tests/ReplicationSystem/TestDependentObjects.cpp b/Engine/Plugins/Runtime/ReplicationSystemTestPlugin/Source/Private/Tests/ReplicationSystem/TestDependentObjects.cpp
--- a/Engine/Plugins/Runtime/ReplicationSystemTestPlugin/Source/Private/Tests/ReplicationSystem/TestDependentObjects.cpp
+++ b/Engine/Plugins/Runtime/ReplicationSystemTestPlugin/Source/Private/Tests/ReplicationSystem/TestDependentObjects.cpp
@@ -10,6 +10,7 @@
 namespace UE::Net::Private
 {
 	
+// Use this fixture for all tests in this file regardless of whether the test requires it or not. That makes it easy to only run or never run these tests.
 class FTestFilteredDependentObjectFixture : public FReplicationSystemServerClientTestFixture
 {
 protected:
@@ -93,7 +94,7 @@ private:
 	TArray<FNetObjectFilterDefinition> OriginalFilterDefinitions;
 };
 
-UE_NET_TEST_FIXTURE(FReplicationSystemServerClientTestFixture, TestDependentObjectDroppedDataIsRetransmitted)
+UE_NET_TEST_FIXTURE(FTestFilteredDependentObjectFixture, TestDependentObjectDroppedDataIsRetransmitted)
 {
 	UReplicationSystem* ReplicationSystem = Server->ReplicationSystem;
 
@@ -152,7 +153,7 @@ UE_NET_TEST_FIXTURE(FReplicationSystemServerClientTestFixture, TestDependentObje
 }
 
 // Dependent objects
-UE_NET_TEST_FIXTURE(FReplicationSystemServerClientTestFixture, TestDependentObjectWithZeroPrioOnlyReplicatesWithParent)
+UE_NET_TEST_FIXTURE(FTestFilteredDependentObjectFixture, TestDependentObjectWithZeroPrioOnlyReplicatesWithParent)
 {
 	UReplicationSystem* ReplicationSystem = Server->ReplicationSystem;
 	UReplicatedTestObjectBridge* Bridge = Server->GetReplicationBridge();
@@ -200,7 +201,7 @@ UE_NET_TEST_FIXTURE(FReplicationSystemServerClientTestFixture, TestDependentObje
 }
 
 // Chained dependent objects
-UE_NET_TEST_FIXTURE(FReplicationSystemServerClientTestFixture, TestChainedDependentObjectWithZeroPrio)
+UE_NET_TEST_FIXTURE(FTestFilteredDependentObjectFixture, TestChainedDependentObjectWithZeroPrio)
 {
 	UReplicationSystem* ReplicationSystem = Server->ReplicationSystem;
 	UReplicatedTestObjectBridge* Bridge = Server->GetReplicationBridge();
@@ -258,7 +259,7 @@ UE_NET_TEST_FIXTURE(FReplicationSystemServerClientTestFixture, TestChainedDepend
 	UE_NET_ASSERT_NE(ClientDependentObject1, nullptr);
 }
 
-UE_NET_TEST_FIXTURE(FReplicationSystemServerClientTestFixture, TestDependentObjectPollFrequency)
+UE_NET_TEST_FIXTURE(FTestFilteredDependentObjectFixture, TestDependentObjectPollFrequency)
 {
 	UReplicationSystem* ReplicationSystem = Server->ReplicationSystem;
 
@@ -329,7 +330,7 @@ UE_NET_TEST_FIXTURE(FReplicationSystemServerClientTestFixture, TestDependentObje
 	UE_NET_ASSERT_EQ(ClientDependentObject->IntA, ServerDependentObject->IntA);
 }
 
-UE_NET_TEST_FIXTURE(FReplicationSystemServerClientTestFixture, TestDependentObjectPollWithDirtyParent)
+UE_NET_TEST_FIXTURE(FTestFilteredDependentObjectFixture, TestDependentObjectPollWithDirtyParent)
 {
 	UReplicationSystem* ReplicationSystem = Server->ReplicationSystem;
 
@@ -408,7 +409,7 @@ UE_NET_TEST_FIXTURE(FReplicationSystemServerClientTestFixture, TestDependentObje
 	UE_NET_ASSERT_GE(RepCountDependent, RepCount1);
 }
 
-UE_NET_TEST_FIXTURE(FReplicationSystemServerClientTestFixture, TestDependentObjectIsPolledIfParentIsMarkedDirty)
+UE_NET_TEST_FIXTURE(FTestFilteredDependentObjectFixture, TestDependentObjectIsPolledIfParentIsMarkedDirty)
 {
 	UReplicationSystem* ReplicationSystem = Server->ReplicationSystem;
 
@@ -485,7 +486,7 @@ UE_NET_TEST_FIXTURE(FReplicationSystemServerClientTestFixture, TestDependentObje
 	UE_NET_ASSERT_EQ(ClientDependentObject->IntA, ServerDependentObject->IntA);
 }
 
-UE_NET_TEST_FIXTURE(FReplicationSystemServerClientTestFixture, TestDependentObjectScheduledAfterParent)
+UE_NET_TEST_FIXTURE(FTestFilteredDependentObjectFixture, TestDependentObjectScheduledAfterParent)
 {
 	UReplicationSystem* ReplicationSystem = Server->ReplicationSystem;
 	UReplicatedTestObjectBridge* Bridge = Server->GetReplicationBridge();
@@ -534,7 +535,7 @@ UE_NET_TEST_FIXTURE(FReplicationSystemServerClientTestFixture, TestDependentObje
 	UE_NET_ASSERT_GT(ClientDependentObject->LastRepOrderCounter, ClientObject->LastRepOrderCounter);
 }
 
-UE_NET_TEST_FIXTURE(FReplicationSystemServerClientTestFixture, TestDependentObjectScheduledBeforeParent)
+UE_NET_TEST_FIXTURE(FTestFilteredDependentObjectFixture, TestDependentObjectScheduledBeforeParent)
 {
 	UReplicationSystem* ReplicationSystem = Server->ReplicationSystem;
 	UReplicatedTestObjectBridge* Bridge = Server->GetReplicationBridge();
@@ -583,7 +584,7 @@ UE_NET_TEST_FIXTURE(FReplicationSystemServerClientTestFixture, TestDependentObje
 	UE_NET_ASSERT_LT(ClientDependentObject->LastRepOrderCounter, ClientObject->LastRepOrderCounter);
 }
 
-UE_NET_TEST_FIXTURE(FReplicationSystemServerClientTestFixture, TestDependentObjectScheduledBeforeParentIfInitialState)
+UE_NET_TEST_FIXTURE(FTestFilteredDependentObjectFixture, TestDependentObjectScheduledBeforeParentIfInitialState)
 {
 	UReplicationSystem* ReplicationSystem = Server->ReplicationSystem;
 	UReplicatedTestObjectBridge* Bridge = Server->GetReplicationBridge();
@@ -632,7 +633,7 @@ UE_NET_TEST_FIXTURE(FReplicationSystemServerClientTestFixture, TestDependentObje
 	UE_NET_ASSERT_GT(ClientDependentObject->LastRepOrderCounter, ClientObject->LastRepOrderCounter);
 }
 
-UE_NET_TEST_FIXTURE(FReplicationSystemServerClientTestFixture, TestNestedDependentObjectScheduledBeforeParent)
+UE_NET_TEST_FIXTURE(FTestFilteredDependentObjectFixture, TestNestedDependentObjectScheduledBeforeParent)
 {
 	UReplicationSystem* ReplicationSystem = Server->ReplicationSystem;
 	UReplicatedTestObjectBridge* Bridge = Server->GetReplicationBridge();
@@ -689,7 +690,7 @@ UE_NET_TEST_FIXTURE(FReplicationSystemServerClientTestFixture, TestNestedDepende
 	UE_NET_ASSERT_LT(ClientDependentObject->LastRepOrderCounter, ClientObject->LastRepOrderCounter);
 }
 
-UE_NET_TEST_FIXTURE(FReplicationSystemServerClientTestFixture, TestNestedDependentObjectScheduledBeforeParents)
+UE_NET_TEST_FIXTURE(FTestFilteredDependentObjectFixture, TestNestedDependentObjectScheduledBeforeParents)
 {
 	UReplicationSystem* ReplicationSystem = Server->ReplicationSystem;
 	UReplicatedTestObjectBridge* Bridge = Server->GetReplicationBridge();
@@ -746,7 +747,7 @@ UE_NET_TEST_FIXTURE(FReplicationSystemServerClientTestFixture, TestNestedDepende
 	UE_NET_ASSERT_LT(ClientObject->LastRepOrderCounter, ClientDependentObject->LastRepOrderCounter);
 }
 
-UE_NET_TEST_FIXTURE(FReplicationSystemServerClientTestFixture, TestLateAddedNestedDependentObjectScheduledBeforeParents)
+UE_NET_TEST_FIXTURE(FTestFilteredDependentObjectFixture, TestLateAddedNestedDependentObjectScheduledBeforeParents)
 {
 	UReplicationSystem* ReplicationSystem = Server->ReplicationSystem;
 	UReplicatedTestObjectBridge* Bridge = Server->GetReplicationBridge();
@@ -803,7 +804,7 @@ UE_NET_TEST_FIXTURE(FReplicationSystemServerClientTestFixture, TestLateAddedNest
 	UE_NET_ASSERT_LT(ClientNestedDependentObject->LastRepOrderCounter, ClientObject->LastRepOrderCounter);
 }
 
-UE_NET_TEST_FIXTURE(FReplicationSystemServerClientTestFixture, TestLateAddedNestedDependentObjectPendingWaitOnCreateConfirmation)
+UE_NET_TEST_FIXTURE(FTestFilteredDependentObjectFixture, TestLateAddedNestedDependentObjectPendingWaitOnCreateConfirmation)
 {
 	UReplicationSystem* ReplicationSystem = Server->ReplicationSystem;
 	UReplicatedTestObjectBridge* Bridge = Server->GetReplicationBridge();
@@ -893,7 +894,7 @@ UE_NET_TEST_FIXTURE(FTestFilteredDependentObjectFixture, TestDependentObjectIsFi
 	UE_NET_ASSERT_NE(ClientDependentObject, nullptr);
 }
 
-UE_NET_TEST_FIXTURE(FReplicationSystemServerClientTestFixture, TestDependentObjectScheduledBeforeParentWithHighPriority)
+UE_NET_TEST_FIXTURE(FTestFilteredDependentObjectFixture, TestDependentObjectScheduledBeforeParentWithHighPriority)
 {
 	UReplicationSystem* ReplicationSystem = Server->ReplicationSystem;
 	UReplicatedTestObjectBridge* Bridge = Server->GetReplicationBridge();
@@ -946,7 +947,7 @@ UE_NET_TEST_FIXTURE(FReplicationSystemServerClientTestFixture, TestDependentObje
 	UE_NET_ASSERT_LT(ClientDependentObject->LastRepOrderCounter, ClientObject->LastRepOrderCounter);
 }
 
-UE_NET_TEST_FIXTURE(FReplicationSystemServerClientTestFixture, TestCircularDependencyIsDenied)
+UE_NET_TEST_FIXTURE(FTestFilteredDependentObjectFixture, TestCircularDependencyIsDenied)
 {
 	UReplicatedTestObjectBridge* Bridge = Server->GetReplicationBridge();
 
@@ -973,4 +974,204 @@ UE_NET_TEST_FIXTURE(FReplicationSystemServerClientTestFixture, TestCircularDepen
 	}
 }
 
+// Create a dependency chain with three objects. Mark dependent objects as not routed. All three objects should replicate.
+UE_NET_TEST_FIXTURE(FTestFilteredDependentObjectFixture, DependentObjectIsInScopeIfParentChainIsInScope)
+{
+	UReplicationSystem* ReplicationSystem = Server->ReplicationSystem;
+	UReplicatedTestObjectBridge* Bridge = Server->GetReplicationBridge();
+
+	// Add a client
+	FReplicationSystemTestClient* Client = CreateClient();
+
+	// Objects are created in arbitrary order
+	{
+		UTestReplicatedIrisObject* ServerRootObject = Server->CreateObject();
+		UTestReplicatedIrisObject* ServerParentObject = Server->CreateObject();
+		UTestReplicatedIrisObject* ServerDependentObject = Server->CreateObject();
+
+		ReplicationSystem->SetFilter(ServerParentObject->NetRefHandle, NotRoutedFilterHandle);
+		ReplicationSystem->SetFilter(ServerDependentObject->NetRefHandle, NotRoutedFilterHandle);
+
+		Bridge->AddDependentObject(ServerRootObject->NetRefHandle, ServerParentObject->NetRefHandle);
+		Bridge->AddDependentObject(ServerParentObject->NetRefHandle, ServerDependentObject->NetRefHandle);
+
+		Server->UpdateAndSend({ Client });
+
+		UE_NET_ASSERT_TRUE(Client->IsResolvableNetRefHandle(ServerRootObject->NetRefHandle));
+		UE_NET_ASSERT_TRUE(Client->IsResolvableNetRefHandle(ServerParentObject->NetRefHandle));
+		UE_NET_ASSERT_TRUE(Client->IsResolvableNetRefHandle(ServerDependentObject->NetRefHandle));
+
+		// Cleanup a bit
+		Server->DestroyObject(ServerDependentObject);
+		Server->DestroyObject(ServerParentObject);
+		Server->DestroyObject(ServerRootObject);
+
+		Server->UpdateAndSend({ Client });
+	}
+
+	// Objects are created in reverse order to the previous test
+	{
+		UTestReplicatedIrisObject* ServerDependentObject = Server->CreateObject();
+		UTestReplicatedIrisObject* ServerParentObject = Server->CreateObject();
+		UTestReplicatedIrisObject* ServerRootObject = Server->CreateObject();
+
+		ReplicationSystem->SetFilter(ServerParentObject->NetRefHandle, NotRoutedFilterHandle);
+		ReplicationSystem->SetFilter(ServerDependentObject->NetRefHandle, NotRoutedFilterHandle);
+
+		Bridge->AddDependentObject(ServerRootObject->NetRefHandle, ServerParentObject->NetRefHandle);
+		Bridge->AddDependentObject(ServerParentObject->NetRefHandle, ServerDependentObject->NetRefHandle);
+
+		Server->UpdateAndSend({ Client });
+
+		UE_NET_ASSERT_TRUE(Client->IsResolvableNetRefHandle(ServerRootObject->NetRefHandle));
+		UE_NET_ASSERT_TRUE(Client->IsResolvableNetRefHandle(ServerParentObject->NetRefHandle));
+		UE_NET_ASSERT_TRUE(Client->IsResolvableNetRefHandle(ServerDependentObject->NetRefHandle));
+
+		// Cleanup a bit
+		Server->DestroyObject(ServerDependentObject);
+		Server->DestroyObject(ServerParentObject);
+		Server->DestroyObject(ServerRootObject);
+
+		Server->UpdateAndSend({ Client });
+	}
+}
+
+// Create a dependency chain with three objects. Mark dependent objects as not routed. Exclusion filter out middle object. Only the root object should replicate.
+UE_NET_TEST_FIXTURE(FTestFilteredDependentObjectFixture, DependentObjectIsNotInScopeIfParentIsNotInScope)
+{
+	UReplicationSystem* ReplicationSystem = Server->ReplicationSystem;
+	UReplicatedTestObjectBridge* Bridge = Server->GetReplicationBridge();
+
+	// Add a client
+	FReplicationSystemTestClient* Client = CreateClient();
+
+	// Objects are created in arbitrary order
+	{
+		UTestReplicatedIrisObject* ServerRootObject = Server->CreateObject();
+		UTestReplicatedIrisObject* ServerParentObject = Server->CreateObject();
+		UTestReplicatedIrisObject* ServerDependentObject = Server->CreateObject();
+
+		ReplicationSystem->SetFilter(ServerParentObject->NetRefHandle, NotRoutedFilterHandle);
+		ReplicationSystem->SetFilter(ServerDependentObject->NetRefHandle, NotRoutedFilterHandle);
+
+		Bridge->AddDependentObject(ServerRootObject->NetRefHandle, ServerParentObject->NetRefHandle);
+		Bridge->AddDependentObject(ServerParentObject->NetRefHandle, ServerDependentObject->NetRefHandle);
+
+		FNetObjectGroupHandle ExclusionGroup = ReplicationSystem->CreateGroup(NAME_None);
+		ReplicationSystem->AddToGroup(ExclusionGroup, ServerParentObject->NetRefHandle);
+		ReplicationSystem->AddExclusionFilterGroup(ExclusionGroup);
+
+		Server->UpdateAndSend({ Client });
+
+		UE_NET_ASSERT_TRUE(Client->IsResolvableNetRefHandle(ServerRootObject->NetRefHandle));
+		UE_NET_ASSERT_FALSE(Client->IsResolvableNetRefHandle(ServerParentObject->NetRefHandle));
+		UE_NET_ASSERT_FALSE(Client->IsResolvableNetRefHandle(ServerDependentObject->NetRefHandle));
+
+		// Cleanup a bit
+		ReplicationSystem->DestroyGroup(ExclusionGroup);
+
+		Server->DestroyObject(ServerDependentObject);
+		Server->DestroyObject(ServerParentObject);
+		Server->DestroyObject(ServerRootObject);
+
+		Server->UpdateAndSend({ Client });
+	}
+
+	// Objects are created in reverse order to the previous test
+	{
+		UTestReplicatedIrisObject* ServerDependentObject = Server->CreateObject();
+		UTestReplicatedIrisObject* ServerParentObject = Server->CreateObject();
+		UTestReplicatedIrisObject* ServerRootObject = Server->CreateObject();
+
+		ReplicationSystem->SetFilter(ServerParentObject->NetRefHandle, NotRoutedFilterHandle);
+		ReplicationSystem->SetFilter(ServerDependentObject->NetRefHandle, NotRoutedFilterHandle);
+
+		Bridge->AddDependentObject(ServerRootObject->NetRefHandle, ServerParentObject->NetRefHandle);
+		Bridge->AddDependentObject(ServerParentObject->NetRefHandle, ServerDependentObject->NetRefHandle);
+
+		FNetObjectGroupHandle ExclusionGroup = ReplicationSystem->CreateGroup(NAME_None);
+		ReplicationSystem->AddToGroup(ExclusionGroup, ServerParentObject->NetRefHandle);
+		ReplicationSystem->AddExclusionFilterGroup(ExclusionGroup);
+
+		Server->UpdateAndSend({ Client });
+
+		UE_NET_ASSERT_TRUE(Client->IsResolvableNetRefHandle(ServerRootObject->NetRefHandle));
+		UE_NET_ASSERT_FALSE(Client->IsResolvableNetRefHandle(ServerParentObject->NetRefHandle));
+		UE_NET_ASSERT_FALSE(Client->IsResolvableNetRefHandle(ServerDependentObject->NetRefHandle));
+
+		// Cleanup a bit
+		ReplicationSystem->DestroyGroup(ExclusionGroup);
+
+		Server->DestroyObject(ServerDependentObject);
+		Server->DestroyObject(ServerParentObject);
+		Server->DestroyObject(ServerRootObject);
+
+		Server->UpdateAndSend({ Client });
+	}
+}
+
+// Create a dependency chain with three objects. Mark all objects as not routed. No objects should replicate.
+UE_NET_TEST_FIXTURE(FTestFilteredDependentObjectFixture, DependentObjectIsNotInScopeIfRootIsDynamicallyFilteredOut)
+{
+	UReplicationSystem* ReplicationSystem = Server->ReplicationSystem;
+	UReplicatedTestObjectBridge* Bridge = Server->GetReplicationBridge();
+
+	// Add a client
+	FReplicationSystemTestClient* Client = CreateClient();
+
+	// Objects are created in arbitrary order
+	{
+		UTestReplicatedIrisObject* ServerRootObject = Server->CreateObject();
+		UTestReplicatedIrisObject* ServerParentObject = Server->CreateObject();
+		UTestReplicatedIrisObject* ServerDependentObject = Server->CreateObject();
+
+		ReplicationSystem->SetFilter(ServerRootObject->NetRefHandle, NotRoutedFilterHandle);
+		ReplicationSystem->SetFilter(ServerParentObject->NetRefHandle, NotRoutedFilterHandle);
+		ReplicationSystem->SetFilter(ServerDependentObject->NetRefHandle, NotRoutedFilterHandle);
+
+		Bridge->AddDependentObject(ServerRootObject->NetRefHandle, ServerParentObject->NetRefHandle);
+		Bridge->AddDependentObject(ServerParentObject->NetRefHandle, ServerDependentObject->NetRefHandle);
+
+		Server->UpdateAndSend({ Client });
+
+		UE_NET_ASSERT_FALSE(Client->IsResolvableNetRefHandle(ServerRootObject->NetRefHandle));
+		UE_NET_ASSERT_FALSE(Client->IsResolvableNetRefHandle(ServerParentObject->NetRefHandle));
+		UE_NET_ASSERT_FALSE(Client->IsResolvableNetRefHandle(ServerDependentObject->NetRefHandle));
+
+		// Cleanup a bit
+		Server->DestroyObject(ServerDependentObject);
+		Server->DestroyObject(ServerParentObject);
+		Server->DestroyObject(ServerRootObject);
+
+		Server->UpdateAndSend({ Client });
+	}
+
+	// Objects are created in reverse order to the previous test
+	{
+		UTestReplicatedIrisObject* ServerDependentObject = Server->CreateObject();
+		UTestReplicatedIrisObject* ServerParentObject = Server->CreateObject();
+		UTestReplicatedIrisObject* ServerRootObject = Server->CreateObject();
+
+		ReplicationSystem->SetFilter(ServerRootObject->NetRefHandle, NotRoutedFilterHandle);
+		ReplicationSystem->SetFilter(ServerParentObject->NetRefHandle, NotRoutedFilterHandle);
+		ReplicationSystem->SetFilter(ServerDependentObject->NetRefHandle, NotRoutedFilterHandle);
+
+		Bridge->AddDependentObject(ServerRootObject->NetRefHandle, ServerParentObject->NetRefHandle);
+		Bridge->AddDependentObject(ServerParentObject->NetRefHandle, ServerDependentObject->NetRefHandle);
+
+		Server->UpdateAndSend({ Client });
+
+		UE_NET_ASSERT_FALSE(Client->IsResolvableNetRefHandle(ServerRootObject->NetRefHandle));
+		UE_NET_ASSERT_FALSE(Client->IsResolvableNetRefHandle(ServerParentObject->NetRefHandle));
+		UE_NET_ASSERT_FALSE(Client->IsResolvableNetRefHandle(ServerDependentObject->NetRefHandle));
+
+		// Cleanup a bit
+		Server->DestroyObject(ServerDependentObject);
+		Server->DestroyObject(ServerParentObject);
+		Server->DestroyObject(ServerRootObject);
+
+		Server->UpdateAndSend({ Client });
+	}
+}
+
 }
diff --git a/Engine/Source/Runtime/Net/Iris/Private/Iris/ReplicationSystem/Filtering/ReplicationFiltering.cpp b/Engine/Source/Runtime/Net/Iris/Private/Iris/ReplicationSystem/Filtering/ReplicationFiltering.cpp
--- a/Engine/Source/Runtime/Net/Iris/Private/Iris/ReplicationSystem/Filtering/ReplicationFiltering.cpp
+++ b/Engine/Source/Runtime/Net/Iris/Private/Iris/ReplicationSystem/Filtering/ReplicationFiltering.cpp
@@ -62,18 +62,6 @@ FName GetStaticFilterName(FNetObjectFilterHandle Filter)
 	return NAME_None;
 }
 
-inline static ENetFilterStatus GetDependentObjectFilterStatus(const FNetRefHandleManager* NetRefHandleManager, const FNetBitArray& ObjectsInScope, FInternalNetRefIndex ObjectIndex)
-{
-	for (const FInternalNetRefIndex ParentObjectIndex : NetRefHandleManager->GetDependentObjectParents(ObjectIndex))
-	{
-		if (GetDependentObjectFilterStatus(NetRefHandleManager, ObjectsInScope, ParentObjectIndex) == ENetFilterStatus::Allow)
-		{
-			return ENetFilterStatus::Allow;
-		}
-	}
-	return ObjectsInScope.GetBit(ObjectIndex) ? ENetFilterStatus::Allow : ENetFilterStatus::Disallow;
-}
-
 class FNetObjectFilterHandleUtil
 {
 public:
@@ -1349,6 +1337,10 @@ void FReplicationFiltering::UpdateDynamicFiltering()
 	TArray<uint32> FilteredOutByHysteresisObjects;
 	FilteredOutByHysteresisObjects.Reserve(256);
 
+	// For circular dependency detection/avoidance when retrieving dependent object filter status.
+	TSet<FInternalNetRefIndex> VisitedDependentObjects;
+	VisitedDependentObjects.Reserve(32);
+
 	uint32* AllowedObjectsData = static_cast<uint32*>(FMemory_Alloca(WordCountForObjectBitArrays * sizeof(uint32)));
 	FNetBitArrayView AllowedObjects(AllowedObjectsData, MaxInternalNetRefIndex, FNetBitArrayView::NoResetNoValidate);
 
@@ -1625,7 +1617,8 @@ void FReplicationFiltering::UpdateDynamicFiltering()
 			const uint32 AdjustHysteresisForUpdateThrottling = (HysteresisState.ConnectionStartId + HysteresisState.ConnectionIdStride - ConnIdMod) % HysteresisState.ConnectionIdStride;
 			for (const uint32 DependentObjectIndex : FilteredOutDependentObjects)
 			{
-				const bool bAllowReplication = GetDependentObjectFilterStatus(NetRefHandleManager, ConnectionInfo.ObjectsInScope, DependentObjectIndex) == ENetFilterStatus::Allow;
+				VisitedDependentObjects.Reset();
+				const bool bAllowReplication = GetDependentObjectFilterStatus(ConnId, DependentObjectIndex, VisitedDependentObjects) == ENetFilterStatus::Allow;
 				if (bAllowReplication)
 				{
 					ConnectionInfo.HysteresisUpdater.RemoveHysteresis(DependentObjectIndex);
@@ -1691,7 +1684,8 @@ void FReplicationFiltering::UpdateDynamicFiltering()
 			// We didn't filter out any dependent objects yet so we enable/disable as needed.
 			for (const uint32 DependentObjectIndex : FilteredOutDependentObjects)
 			{
-				const bool bAllowReplication = GetDependentObjectFilterStatus(NetRefHandleManager, ConnectionInfo.ObjectsInScope, DependentObjectIndex) == ENetFilterStatus::Allow;
+				VisitedDependentObjects.Reset();
+				const bool bAllowReplication = GetDependentObjectFilterStatus(ConnId, DependentObjectIndex, VisitedDependentObjects) == ENetFilterStatus::Allow;
 				ConnectionInfo.ObjectsInScope.SetBitValue(DependentObjectIndex, bAllowReplication);
 				for (const FInternalNetRefIndex SubObjectIndex : NetRefHandleManager->GetSubObjects(DependentObjectIndex))
 				{
@@ -2560,6 +2554,59 @@ void FReplicationFiltering::NotifyRemovedDependentObject(FInternalNetRefIndex Ob
 	ObjectsRequiringDynamicFilterUpdate.SetBit(ObjectIndex);
 }
 
+// Determine whether a dynamically filtered out dependent object should be in scope or not. Note that this function is expected to only be called for dependent objects NOT in scope. The Visited set is expected to be empty when calling this function.
+ENetFilterStatus FReplicationFiltering::GetDependentObjectFilterStatus(uint32 ConnectionId, FInternalNetRefIndex ObjectIndex, TSet<FInternalNetRefIndex>& Visited)
+{
+	const FPerConnectionInfo& ConnectionInfo = ConnectionInfos[ConnectionId];
+
+	// If the dependent object is filtered out prior to dynamic filtering we shouldn't replicate it.
+	if (!ConnectionInfo.ObjectsInScopeBeforeDynamicFiltering.GetBit(ObjectIndex))
+	{
+		return ENetFilterStatus::Disallow;
+	}
+
+	// Setup the initial objects to check whether they're in scope or not.
+	TArray<FInternalNetRefIndex, TInlineAllocator<64>> DependentObjects;
+	DependentObjects.Append(NetRefHandleManager->GetDependentObjectParents(ObjectIndex));
+	
+	// Iteratively visit parents until we find one in scope.
+	bool bIsAlreadyVisited = false;
+	while (!DependentObjects.IsEmpty()) 
+	{
+		if (const FInternalNetRefIndex ParentObjectIndex = DependentObjects.Pop(EAllowShrinking::No))
+		{
+			// If we have a cyclic dependency chain or already visited the object we continue to the next parent as we already know this one nor any of its parents are in scope.
+			Visited.Add(ParentObjectIndex, &bIsAlreadyVisited);
+			if (bIsAlreadyVisited)
+			{
+				constexpr uint32 ArbitraryMaxLogCount = 3U;
+				UE_CLOG((ParentObjectIndex == ObjectIndex) && (LoggedCyclicDependencyCount++ < ArbitraryMaxLogCount), LogIrisFiltering, Error, TEXT("Found circular reference while retrieving dependent object filter status for %s"), ToCStr(NetRefHandleManager->PrintObjectFromIndex(ObjectIndex)));
+				ensureMsgf(ParentObjectIndex != ObjectIndex, TEXT("Found circular reference while retrieving dependent object filter status for %s"), ToCStr(NetRefHandleManager->PrintObjectFromIndex(ObjectIndex)));
+				continue;
+			}
+
+			// If the parent is filtered out prior to dynamic filtering we should continue looking for other parents which may allow replication.
+			if (!ConnectionInfo.ObjectsInScopeBeforeDynamicFiltering.GetBit(ParentObjectIndex))
+			{
+				// There may be other parents in scope so we should continue processing those.
+				continue;
+			}
+
+			// If the parent is in scope so should the dependent object be.
+			if (ConnectionInfo.ObjectsInScope.GetBit(ParentObjectIndex))
+			{
+				return ENetFilterStatus::Allow;
+			}
+
+			// If it's filtered out dynamically and not a dependent object this parent can't allow replication. In that case it won't have any parents anyway so it's safe to call GetDependentObjectParents.
+			DependentObjects.Append(NetRefHandleManager->GetDependentObjectParents(ParentObjectIndex));
+		}
+	}
+	
+	// If no parent is in scope this object shouldn't be either. 
+	return ENetFilterStatus::Disallow;
+}
+
 ENetFilterStatus FReplicationFiltering::GetConnectionFilterStatus(const FPerObjectInfo& ObjectInfo, uint32 ConnectionId) const
 {
 	return (ObjectInfo.ConnectionIds[ConnectionId >> 5U] & (1U << (ConnectionId & 31U))) != 0U ? ENetFilterStatus::Allow : ENetFilterStatus::Disallow;
diff --git a/Engine/Source/Runtime/Net/Iris/Private/Iris/ReplicationSystem/Filtering/ReplicationFiltering.h b/Engine/Source/Runtime/Net/Iris/Private/Iris/ReplicationSystem/Filtering/ReplicationFiltering.h
--- a/Engine/Source/Runtime/Net/Iris/Private/Iris/ReplicationSystem/Filtering/ReplicationFiltering.h
+++ b/Engine/Source/Runtime/Net/Iris/Private/Iris/ReplicationSystem/Filtering/ReplicationFiltering.h
@@ -263,6 +263,7 @@ private:
 	FPerSubObjectFilterGroupInfo* GetPerSubObjectFilterGroupInfo(FNetObjectGroupHandle::FGroupIndexType GroupIndex);
 	const FPerSubObjectFilterGroupInfo* GetPerSubObjectFilterGroupInfo(FNetObjectGroupHandle::FGroupIndexType GroupIndex) const;
 
+	ENetFilterStatus GetDependentObjectFilterStatus(uint32 ConnectionId, FInternalNetRefIndex ObjectIndex, TSet<FInternalNetRefIndex>& Visited);
 	ENetFilterStatus GetConnectionFilterStatus(const FPerObjectInfo& ObjectInfo, uint32 ConnectionId) const;
 	bool IsAnyConnectionFilterStatusAllowed(const FPerObjectInfo& ObjectInfo) const;
 	bool IsAnyConnectionFilterStatusDisallowed(const FPerObjectInfo& ObjectInfo) const;
@@ -415,6 +416,8 @@ private:
 	uint32 bHasDirtyInclusionFilterGroup : 1;
 	// Is true if any initialized DynamicFilter has the NeedsUpdate trait
 	uint32 bHasDynamicFiltersWithUpdateTrait : 1;
+
+	uint32 LoggedCyclicDependencyCount = 0U;
 };
 
 inline bool FReplicationFiltering::HasDynamicFilters() const
