Iris - Fix issue with newly created dependent object entering hysteresis despite parent being filtered out already. Simplified dependent object filter status code.

diff --git a/Engine/Plugins/Runtime/ReplicationSystemTestPlugin/Source/Private/Tests/ReplicationSystem/Filtering/TestFiltering.cpp b/Engine/Plugins/Runtime/ReplicationSystemTestPlugin/Source/Private/Tests/ReplicationSystem/Filtering/TestFiltering.cpp
--- a/Engine/Plugins/Runtime/ReplicationSystemTestPlugin/Source/Private/Tests/ReplicationSystem/Filtering/TestFiltering.cpp
+++ b/Engine/Plugins/Runtime/ReplicationSystemTestPlugin/Source/Private/Tests/ReplicationSystem/Filtering/TestFiltering.cpp
@@ -22,7 +22,7 @@ protected:
 		InitNetObjectFilterDefinitions();
 		InitObjectScopeHysteresisProfiles();
 		FReplicationSystemServerClientTestFixture::SetUp();
-		InitMockNetObjectFilter();
+		InitFilterHandles();
 	}
 
 	virtual void TearDown() override
@@ -145,6 +145,12 @@ private:
 			MockDefinition.ClassName = GetMockFilterClassName();
 			MockDefinition.ConfigClassName = "/Script/ReplicationSystemTestPlugin.MockNetObjectFilterConfig";
 		}
+		{
+			FNetObjectFilterDefinition& NotRoutedDefinition = NewFilterDefinitions.Emplace_GetRef();
+			NotRoutedDefinition.FilterName = "NotRouted";
+			NotRoutedDefinition.ClassName = "/Script/IrisCore.FilterOutNetObjectFilter";
+			NotRoutedDefinition.ConfigClassName = "/Script/IrisCore.NetObjectGridFilterConfig";
+		}
 
 		DefinitionsProperty->CopyCompleteValue((void*)(UPTRINT(FilterDefinitions) + DefinitionsProperty->GetOffset_ForInternal()), &NewFilterDefinitions);
 	}
@@ -158,8 +164,9 @@ private:
 		DefinitionsProperty->CopyCompleteValue((void*)(UPTRINT(FilterDefinitions) + DefinitionsProperty->GetOffset_ForInternal()), &OriginalFilterDefinitions);
 		OriginalFilterDefinitions.Empty();
 
-		MockFilterHandle = InvalidNetObjectFilterHandle;
 		MockNetObjectFilter = nullptr;
+		MockFilterHandle = InvalidNetObjectFilterHandle;
+		NotRoutedFilterHandle = InvalidNetObjectFilterHandle;
 	}
 
 	void InitObjectScopeHysteresisProfiles()
@@ -194,6 +201,12 @@ private:
 				Profile.FilterProfileName = "ZeroFrames";
 				Profile.HysteresisFrameCount = 0;
 			}
+
+			{
+				FObjectScopeHysteresisProfile& Profile = NewProfiles.Emplace_GetRef();
+				Profile.FilterProfileName = "VeryManyFrames";
+				Profile.HysteresisFrameCount = 255;
+			}
 		}
 
 		ProfilesProperty->CopyCompleteValue((void*)(UPTRINT(FilteringConfig) + ProfilesProperty->GetOffset_ForInternal()), &NewProfiles);
@@ -216,14 +229,17 @@ private:
 		EnableObjectScopeHysteresisProperty->CopyCompleteValue((void*)(UPTRINT(FilteringConfig) + EnableObjectScopeHysteresisProperty->GetOffset_ForInternal()), &bOriginalIsObjectScopeHysteresisEnabled);
 	}
 
-	void InitMockNetObjectFilter()
+	void InitFilterHandles()
 	{
 		MockNetObjectFilter = CastChecked<UMockNetObjectFilter>(Server->GetReplicationSystem()->GetFilter(GetMockFilterName()));
 		MockFilterHandle = Server->GetReplicationSystem()->GetFilterHandle(GetMockFilterName());
+		NotRoutedFilterHandle = Server->GetReplicationSystem()->GetFilterHandle("NotRouted");
 	}
+
 protected:
-	UMockNetObjectFilter* MockNetObjectFilter;
+	UMockNetObjectFilter* MockNetObjectFilter = nullptr;
 	FNetObjectFilterHandle MockFilterHandle;
+	FNetObjectFilterHandle NotRoutedFilterHandle;
 	FName ObjectHysteresisProfileName;
 
 private:
@@ -2706,6 +2722,182 @@ UE_NET_TEST_FIXTURE(FTestFilteringFixture, LotsOfObjectsCanBeFilteredOutViaHyste
 	}
 }
 
+UE_NET_TEST_FIXTURE(FTestFilteringFixture, TestDependentObjectIsFilteredOutTogetherWithExclusionFilteredOutParent)
+{
+	constexpr uint32 DefaultHysteresisFrameCount = 2;
+	FScopedDefaultHysteresisFrameCount ScopedDefaultHysteresisFrameCount(DefaultHysteresisFrameCount);
+
+	UReplicationSystem* ReplicationSystem = Server->ReplicationSystem;
+	UReplicatedTestObjectBridge* Bridge = Server->GetReplicationBridge();
+
+	// Add a client
+	FReplicationSystemTestClient* Client = CreateClient();
+
+	// Spawn objects on server
+	UTestReplicatedIrisObject* ServerObject = Server->CreateObject(UTestReplicatedIrisObject::FComponents{});
+	UTestReplicatedIrisObject* ServerDependentObject = Server->CreateObject(UTestReplicatedIrisObject::FComponents{});
+
+	// Add parent object to MockFilter so we can filter in/out as desired.
+	SetDynamicFilterStatus(ENetFilterStatus::Allow);
+	ReplicationSystem->SetFilter(ServerObject->NetRefHandle, MockFilterHandle);
+
+	// Dynamically filter out dependent object
+	ReplicationSystem->SetFilter(ServerDependentObject->NetRefHandle, NotRoutedFilterHandle);
+
+	Bridge->AddDependentObject(ServerObject->NetRefHandle, ServerDependentObject->NetRefHandle);
+
+	// Send data
+	Server->UpdateAndSend({Client});
+
+	// Verify that objects have been created on the client
+	UE_NET_ASSERT_TRUE(Client->IsResolvableNetRefHandle(ServerObject->NetRefHandle));
+	UE_NET_ASSERT_TRUE(Client->IsResolvableNetRefHandle(ServerDependentObject->NetRefHandle));
+
+	// Throw in some filtering groups into the mix
+	FNetObjectGroupHandle ExclusionFilterGroup = ReplicationSystem->CreateGroup(NAME_None);
+	ReplicationSystem->AddExclusionFilterGroup(ExclusionFilterGroup);
+	FNetObjectGroupHandle InclusionFilterGroup = ReplicationSystem->CreateGroup(NAME_None);
+	ReplicationSystem->AddInclusionFilterGroup(InclusionFilterGroup);
+
+	// Allow replication of the InclusionGroup to the client
+	ReplicationSystem->SetGroupFilterStatus(InclusionFilterGroup, Client->ConnectionIdOnServer, ENetFilterStatus::Allow);
+
+	// Add parent to groups
+	ReplicationSystem->AddToGroup(ExclusionFilterGroup, ServerObject->NetRefHandle);
+	ReplicationSystem->AddToGroup(InclusionFilterGroup, ServerObject->NetRefHandle);
+
+	// Send data
+	Server->UpdateAndSend({Client});
+
+	// Allow replication of the ExclusionGroup to the client
+	ReplicationSystem->SetGroupFilterStatus(ExclusionFilterGroup, Client->ConnectionIdOnServer, ENetFilterStatus::Allow);
+
+	// Send data
+	Server->UpdateAndSend({Client});
+
+	// Verify that the objects exist on the client
+	UE_NET_ASSERT_TRUE(Client->IsResolvableNetRefHandle(ServerObject->NetRefHandle));
+	UE_NET_ASSERT_TRUE(Client->IsResolvableNetRefHandle(ServerDependentObject->NetRefHandle));
+
+	// Disallow replication of the ExclusionGroup to the client
+	ReplicationSystem->SetGroupFilterStatus(ExclusionFilterGroup, Client->ConnectionIdOnServer, ENetFilterStatus::Disallow);
+
+	// Perform send update for the hysteresis duration
+	for (uint32 It = 0; It != DefaultHysteresisFrameCount + 1; ++It)
+	{
+		Server->UpdateAndSend({ Client });
+	}
+
+	// Verify that the objects are no longer valid on the client
+	UE_NET_ASSERT_FALSE(Client->IsValidNetRefHandle(ServerObject->NetRefHandle));
+	UE_NET_ASSERT_FALSE(Client->IsValidNetRefHandle(ServerDependentObject->NetRefHandle));
+}
+
+UE_NET_TEST_FIXTURE(FTestFilteringFixture, DependentObjectAddedToExclusionFilteredOutObjectIsNotInScope)
+{
+	FReplicationSystemTestClient* ClientArray[] = { CreateClient(), CreateClient() };
+
+	UReplicatedTestObject* ServerObject = Server->CreateObject();
+
+	UReplicatedTestObject* ServerDependentObject1 = Server->CreateObject();
+	Server->ReplicationSystem->SetFilter(ServerDependentObject1->NetRefHandle, NotRoutedFilterHandle, /* .FilterConfigProfile */ "VeryManyFrames");
+
+	Server->ReplicationBridge->AddDependentObject(ServerObject->NetRefHandle, ServerDependentObject1->NetRefHandle);
+
+	Server->UpdateAndSend(ClientArray);
+
+	UE_NET_ASSERT_TRUE(ClientArray[0]->IsResolvableNetRefHandle(ServerObject->NetRefHandle));
+	UE_NET_ASSERT_TRUE(ClientArray[0]->IsResolvableNetRefHandle(ServerDependentObject1->NetRefHandle));
+
+	UE_NET_ASSERT_TRUE(ClientArray[1]->IsResolvableNetRefHandle(ServerObject->NetRefHandle));
+	UE_NET_ASSERT_TRUE(ClientArray[1]->IsResolvableNetRefHandle(ServerDependentObject1->NetRefHandle));
+
+	FNetObjectGroupHandle GroupHandle = Server->ReplicationSystem->CreateGroup(NAME_None);
+	Server->ReplicationSystem->AddToGroup(GroupHandle, ServerObject->NetRefHandle);
+	Server->ReplicationSystem->AddExclusionFilterGroup(GroupHandle);
+
+	Server->ReplicationSystem->SetGroupFilterStatus(GroupHandle, ENetFilterStatus::Allow);
+	Server->ReplicationSystem->SetGroupFilterStatus(GroupHandle, ClientArray[1]->ConnectionIdOnServer, ENetFilterStatus::Disallow);
+
+	Server->UpdateAndSend(ClientArray);
+
+	UE_NET_ASSERT_TRUE(ClientArray[0]->IsResolvableNetRefHandle(ServerObject->NetRefHandle));
+	UE_NET_ASSERT_TRUE(ClientArray[0]->IsResolvableNetRefHandle(ServerDependentObject1->NetRefHandle));
+
+	UE_NET_ASSERT_FALSE(ClientArray[1]->IsValidNetRefHandle(ServerObject->NetRefHandle));
+	// The dependent object was in scope for a while and should be subject to hysteresis, thus would be resolvable even for the client for which the dependant was just filtered out
+	UE_NET_ASSERT_TRUE(ClientArray[1]->IsValidNetRefHandle(ServerDependentObject1->NetRefHandle));
+
+	UReplicatedTestObject* ServerDependentObject2 = Server->CreateObject();
+	Server->ReplicationSystem->SetFilter(ServerDependentObject2->NetRefHandle, NotRoutedFilterHandle, /* .FilterConfigProfile */ "VeryManyFrames");
+
+	Server->ReplicationBridge->RemoveDependentObject(ServerObject->NetRefHandle, ServerDependentObject1->NetRefHandle);
+	Server->ReplicationBridge->AddDependentObject(ServerObject->NetRefHandle, ServerDependentObject2->NetRefHandle);
+
+	Server->UpdateAndSend(ClientArray);
+
+	UE_NET_ASSERT_TRUE(ClientArray[0]->IsResolvableNetRefHandle(ServerObject->NetRefHandle));
+	UE_NET_ASSERT_TRUE(ClientArray[0]->IsResolvableNetRefHandle(ServerDependentObject1->NetRefHandle));
+	UE_NET_ASSERT_TRUE(ClientArray[0]->IsResolvableNetRefHandle(ServerDependentObject2->NetRefHandle));
+
+	UE_NET_ASSERT_FALSE(ClientArray[1]->IsValidNetRefHandle(ServerObject->NetRefHandle));
+	UE_NET_ASSERT_TRUE(ClientArray[1]->IsResolvableNetRefHandle(ServerDependentObject1->NetRefHandle));
+	// The ServerDependentObject2 was never in scope for ClientArray[1] and isn't expected to be now that is was just added.
+	UE_NET_ASSERT_FALSE(ClientArray[1]->IsValidNetRefHandle(ServerDependentObject2->NetRefHandle));
+}
+
+UE_NET_TEST_FIXTURE(FTestFilteringFixture, AlreadyCreatedObjectAddedAsDependencyToExclusionFilteredOutObjectIsNotInScope)
+{
+	FReplicationSystemTestClient* ClientArray[] = { CreateClient(), CreateClient() };
+
+	UReplicatedTestObject* ServerObject = Server->CreateObject();
+
+	UReplicatedTestObject* ServerDependentObject1 = Server->CreateObject();
+	UReplicatedTestObject* ServerDependentObject2 = Server->CreateObject();
+	Server->ReplicationSystem->SetFilter(ServerDependentObject1->NetRefHandle, NotRoutedFilterHandle, /* .FilterConfigProfile */ "VeryManyFrames");
+	Server->ReplicationSystem->SetFilter(ServerDependentObject2->NetRefHandle, NotRoutedFilterHandle, /* .FilterConfigProfile */ "VeryManyFrames");
+
+	Server->ReplicationBridge->AddDependentObject(ServerObject->NetRefHandle, ServerDependentObject1->NetRefHandle);
+
+	Server->UpdateAndSend(ClientArray);
+
+	UE_NET_ASSERT_TRUE(ClientArray[0]->IsResolvableNetRefHandle(ServerObject->NetRefHandle));
+	UE_NET_ASSERT_TRUE(ClientArray[0]->IsResolvableNetRefHandle(ServerDependentObject1->NetRefHandle));
+
+	UE_NET_ASSERT_TRUE(ClientArray[1]->IsResolvableNetRefHandle(ServerObject->NetRefHandle));
+	UE_NET_ASSERT_TRUE(ClientArray[1]->IsResolvableNetRefHandle(ServerDependentObject1->NetRefHandle));
+
+	FNetObjectGroupHandle GroupHandle = Server->ReplicationSystem->CreateGroup(NAME_None);
+	Server->ReplicationSystem->AddToGroup(GroupHandle, ServerObject->NetRefHandle);
+	Server->ReplicationSystem->AddExclusionFilterGroup(GroupHandle);
+
+	Server->ReplicationSystem->SetGroupFilterStatus(GroupHandle, ENetFilterStatus::Allow);
+	Server->ReplicationSystem->SetGroupFilterStatus(GroupHandle, ClientArray[1]->ConnectionIdOnServer, ENetFilterStatus::Disallow);
+
+	Server->UpdateAndSend(ClientArray);
+
+	UE_NET_ASSERT_TRUE(ClientArray[0]->IsResolvableNetRefHandle(ServerObject->NetRefHandle));
+	UE_NET_ASSERT_TRUE(ClientArray[0]->IsResolvableNetRefHandle(ServerDependentObject1->NetRefHandle));
+
+	UE_NET_ASSERT_FALSE(ClientArray[1]->IsValidNetRefHandle(ServerObject->NetRefHandle));
+	// The dependent object was in scope for a while and should be subject to hysteresis, thus would be resolvable even for the client for which the dependant was just filtered out
+	UE_NET_ASSERT_TRUE(ClientArray[1]->IsValidNetRefHandle(ServerDependentObject1->NetRefHandle));
+
+	Server->ReplicationBridge->RemoveDependentObject(ServerObject->NetRefHandle, ServerDependentObject1->NetRefHandle);
+	Server->ReplicationBridge->AddDependentObject(ServerObject->NetRefHandle, ServerDependentObject2->NetRefHandle);
+
+	Server->UpdateAndSend(ClientArray);
+
+	UE_NET_ASSERT_TRUE(ClientArray[0]->IsResolvableNetRefHandle(ServerObject->NetRefHandle));
+	UE_NET_ASSERT_TRUE(ClientArray[0]->IsResolvableNetRefHandle(ServerDependentObject1->NetRefHandle));
+	UE_NET_ASSERT_TRUE(ClientArray[0]->IsResolvableNetRefHandle(ServerDependentObject2->NetRefHandle));
+
+	UE_NET_ASSERT_FALSE(ClientArray[1]->IsValidNetRefHandle(ServerObject->NetRefHandle));
+	UE_NET_ASSERT_TRUE(ClientArray[1]->IsResolvableNetRefHandle(ServerDependentObject1->NetRefHandle));
+	// The ServerDependentObject2 was never in scope for ClientArray[1] and isn't expected to be now that is was just added.
+	UE_NET_ASSERT_FALSE(ClientArray[1]->IsValidNetRefHandle(ServerDependentObject2->NetRefHandle));
+}
+
 UE_NET_TEST_FIXTURE(FTestFilteringFixture, OwningConnectionIsSetProperly)
 {
 	FReplicationSystemTestClient* Client = CreateClient();
diff --git a/Engine/Source/Runtime/Net/Iris/Private/Iris/ReplicationSystem/Filtering/ReplicationFiltering.cpp b/Engine/Source/Runtime/Net/Iris/Private/Iris/ReplicationSystem/Filtering/ReplicationFiltering.cpp
--- a/Engine/Source/Runtime/Net/Iris/Private/Iris/ReplicationSystem/Filtering/ReplicationFiltering.cpp
+++ b/Engine/Source/Runtime/Net/Iris/Private/Iris/ReplicationSystem/Filtering/ReplicationFiltering.cpp
@@ -1639,17 +1639,19 @@ void FReplicationFiltering::UpdateDynamicFiltering()
 					bool bIsFilteredOut = DynamicFilteredOutObjectsHysteresisAdjusted.GetBit(DependentObjectIndex);
 					if (!bIsFilteredOut)
 					{
-						// If we're updated by hysteresis at this point we should not start again but otherwise let's start!
+						// If we're updated by hysteresis at this point we should not start again, but if other conditions allow it we should start!
 						if (!ConnectionInfo.HysteresisUpdater.IsObjectUpdated(DependentObjectIndex))
 						{
-							if (const uint32 HysteresisFrameCount = ObjectScopeHysteresisFrameCounts[DependentObjectIndex])
+							const uint32 HysteresisFrameCount = ObjectScopeHysteresisFrameCounts[DependentObjectIndex];
+							const bool bIsExemptFromHysteresis = HysteresisState.ObjectsExemptFromHysteresis.GetBit(DependentObjectIndex);
+							if (HysteresisFrameCount && !bIsExemptFromHysteresis)
 							{
 								const uint16 TotalHysteresisFrameCount = static_cast<uint16>(HysteresisFrameCount - 1 + AdjustHysteresisForUpdateThrottling);
 								ConnectionInfo.HysteresisUpdater.SetHysteresisFrameCount(DependentObjectIndex, TotalHysteresisFrameCount);
 							}
 							else
 							{
-								// Object doesn't support hysteresis so we need to filter it out immediately.
+								// Object is exempt from or doesn't support hysteresis so we need to filter it out immediately.
 								bIsFilteredOut = true;
 							}
 						}
@@ -2555,53 +2557,46 @@ void FReplicationFiltering::NotifyRemovedDependentObject(FInternalNetRefIndex Ob
 }
 
 // Determine whether a dynamically filtered out dependent object should be in scope or not. Note that this function is expected to only be called for dependent objects NOT in scope. The Visited set is expected to be empty when calling this function.
-ENetFilterStatus FReplicationFiltering::GetDependentObjectFilterStatus(uint32 ConnectionId, FInternalNetRefIndex ObjectIndex, TSet<FInternalNetRefIndex>& Visited)
+ENetFilterStatus FReplicationFiltering::GetDependentObjectFilterStatus(uint32 ConnectionId, FInternalNetRefIndex DependentObjectIndex, TSet<FInternalNetRefIndex>& Visited)
 {
 	const FPerConnectionInfo& ConnectionInfo = ConnectionInfos[ConnectionId];
 
-	// If the dependent object is filtered out prior to dynamic filtering we shouldn't replicate it.
-	if (!ConnectionInfo.ObjectsInScopeBeforeDynamicFiltering.GetBit(ObjectIndex))
-	{
-		return ENetFilterStatus::Disallow;
-	}
-
-	// Setup the initial objects to check whether they're in scope or not.
-	TArray<FInternalNetRefIndex, TInlineAllocator<64>> DependentObjects;
-	DependentObjects.Append(NetRefHandleManager->GetDependentObjectParents(ObjectIndex));
-	
-	// Iteratively visit parents until we find one in scope.
+	TArray<FInternalNetRefIndex, TInlineAllocator<64>> DependentObjects;	
 	bool bIsAlreadyVisited = false;
-	while (!DependentObjects.IsEmpty()) 
+
+	// Iteratively visit objects and their parents until we find one in scope.
+	DependentObjects.Add(DependentObjectIndex);
+	do
 	{
-		if (const FInternalNetRefIndex ParentObjectIndex = DependentObjects.Pop(EAllowShrinking::No))
-		{
-			// If we have a cyclic dependency chain or already visited the object we continue to the next parent as we already know this one nor any of its parents are in scope.
-			Visited.Add(ParentObjectIndex, &bIsAlreadyVisited);
-			if (bIsAlreadyVisited)
-			{
-				constexpr uint32 ArbitraryMaxLogCount = 3U;
-				UE_CLOG((ParentObjectIndex == ObjectIndex) && (LoggedCyclicDependencyCount++ < ArbitraryMaxLogCount), LogIrisFiltering, Error, TEXT("Found circular reference while retrieving dependent object filter status for %s"), ToCStr(NetRefHandleManager->PrintObjectFromIndex(ObjectIndex)));
-				ensureMsgf(ParentObjectIndex != ObjectIndex, TEXT("Found circular reference while retrieving dependent object filter status for %s"), ToCStr(NetRefHandleManager->PrintObjectFromIndex(ObjectIndex)));
-				continue;
-			}
+		const FInternalNetRefIndex ObjectIndex = DependentObjects.Pop(EAllowShrinking::No);
 
-			// If the parent is filtered out prior to dynamic filtering we should continue looking for other parents which may allow replication.
-			if (!ConnectionInfo.ObjectsInScopeBeforeDynamicFiltering.GetBit(ParentObjectIndex))
-			{
-				// There may be other parents in scope so we should continue processing those.
-				continue;
-			}
+		// If we have a cyclic dependency chain or already visited the object we continue to the next parent as we already know this one nor any of its parents are in scope.
+		Visited.Add(ObjectIndex, &bIsAlreadyVisited);
+		if (bIsAlreadyVisited)
+		{
+			constexpr uint32 ArbitraryMaxLogCount = 3U;
+			UE_CLOG((ObjectIndex == DependentObjectIndex) && (LoggedCyclicDependencyCount++ < ArbitraryMaxLogCount), LogIrisFiltering, Error, TEXT("Found circular reference while retrieving dependent object filter status for %s"), ToCStr(NetRefHandleManager->PrintObjectFromIndex(DependentObjectIndex)));
+			ensureMsgf(ObjectIndex != DependentObjectIndex, TEXT("Found circular reference while retrieving dependent object filter status for %s"), ToCStr(NetRefHandleManager->PrintObjectFromIndex(DependentObjectIndex)));
+			continue;
+		}
 
-			// If the parent is in scope so should the dependent object be.
-			if (ConnectionInfo.ObjectsInScope.GetBit(ParentObjectIndex))
-			{
-				return ENetFilterStatus::Allow;
-			}
+		// If the parent is filtered out prior to dynamic filtering we should continue looking for other parents which may allow replication.
+		if (!ConnectionInfo.ObjectsInScopeBeforeDynamicFiltering.GetBit(ObjectIndex))
+		{
+			// There may be other parents in scope so we should continue processing those.
+			continue;
+		}
 
-			// If it's filtered out dynamically and not a dependent object this parent can't allow replication. In that case it won't have any parents anyway so it's safe to call GetDependentObjectParents.
-			DependentObjects.Append(NetRefHandleManager->GetDependentObjectParents(ParentObjectIndex));
+		// If the parent is in scope so should the dependent object be.
+		if (ConnectionInfo.ObjectsInScope.GetBit(ObjectIndex))
+		{
+			return ENetFilterStatus::Allow;
 		}
+
+		// If it's filtered out dynamically and not a dependent object this parent can't allow replication. In that case it won't have any parents anyway so it's safe to call GetDependentObjectParents.
+		DependentObjects.Append(NetRefHandleManager->GetDependentObjectParents(ObjectIndex));
 	}
+	while (!DependentObjects.IsEmpty());
 	
 	// If no parent is in scope this object shouldn't be either. 
 	return ENetFilterStatus::Disallow;
