Fixed horde fail

#b Peter.Engstrom

diff --git a/Engine/Plugins/Runtime/ReplicationSystemTestPlugin/Source/Private/Tests/ReplicationSystem/TestDependentObjects.cpp b/Engine/Plugins/Runtime/ReplicationSystemTestPlugin/Source/Private/Tests/ReplicationSystem/TestDependentObjects.cpp
--- a/Engine/Plugins/Runtime/ReplicationSystemTestPlugin/Source/Private/Tests/ReplicationSystem/TestDependentObjects.cpp
+++ b/Engine/Plugins/Runtime/ReplicationSystemTestPlugin/Source/Private/Tests/ReplicationSystem/TestDependentObjects.cpp
@@ -1036,7 +1036,7 @@ UE_NET_TEST_FIXTURE(FReplicationSystemServerClientTestFixture, TestNestedCreatio
 		DependentObject->IntA = ObjectIndex;
 
 		// Increasing prio
-		float Prio = ObjectIndex * 1.0f;
+		const float Prio = (float)ObjectIndex * 1.0f;
 		ReplicationSystem->SetStaticPriority(DependentObject->NetRefHandle, Prio);
 
 		// Set up a nested creation dependency chain
