Iris
- Fixed bug with recent changes to NetPacketNotify if sending a packet before accepting or rejecting the packet

diff --git a/Engine/Plugins/Runtime/ReplicationSystemTestPlugin/Source/Private/Tests/Core/TestNetPacketNotify.cpp b/Engine/Plugins/Runtime/ReplicationSystemTestPlugin/Source/Private/Tests/Core/TestNetPacketNotify.cpp
--- a/Engine/Plugins/Runtime/ReplicationSystemTestPlugin/Source/Private/Tests/Core/TestNetPacketNotify.cpp
+++ b/Engine/Plugins/Runtime/ReplicationSystemTestPlugin/Source/Private/Tests/Core/TestNetPacketNotify.cpp
@@ -302,6 +302,66 @@ struct FNetPacketNotifyTestUtil : public FNetworkAutomationTestSuiteFixture
 			return bSuccess;
 		}
 
+		// Read IncomingData and (optional) compare received acks with state from the from node.
+		// Returns false if an error is detected
+		template<typename T>
+		bool ReceiveWithCallback(const FBitWriter& IncomingData, T&& Functor, const FTestNode* From = nullptr)
+		{
+			bool bSuccess = true;
+
+			auto HandlePacketNotification = [&](FNetPacketNotify::SequenceNumberT Seq, bool bDelivered)
+			{
+				++LastNotifiedPacketId;
+
+				// Sanity check
+				if (!ensure(FNetPacketNotify::SequenceNumberT(FNetPacketNotify::SequenceNumberT::SequenceT(LastNotifiedPacketId)) == Seq))
+				{
+					bSuccess = false;
+				}
+
+				if (bDelivered)
+				{
+					AcknowledgedPackets.Add(LastNotifiedPacketId);
+					if (!ensure(!From || From->AcceptedInPackets.Contains(LastNotifiedPacketId)))
+					{
+						bSuccess = false;
+					}
+				}
+				else
+				{
+					if (!ensure(!From || !From->AcceptedInPackets.Contains(LastNotifiedPacketId)))
+					{
+						bSuccess = false;
+					}
+				}
+			};
+
+			FBitReader Reader(IncomingData.GetData(), IncomingData.GetNumBits());
+
+			FNetPacketNotify::FNotificationHeader Header;
+			PacketNotify.ReadHeader(Header, Reader);
+			
+			const uint32 SequenceDelta = PacketNotify.Update(Header, HandlePacketNotification);
+			if (SequenceDelta > 0U)
+			{
+				InPacketId += SequenceDelta;
+
+				// Sanity check
+				check(FNetPacketNotify::SequenceNumberT(FNetPacketNotify::SequenceNumberT::SequenceT(InPacketId)) == PacketNotify.InSeq);
+
+				if (!PacketNotify.IsWaitingForSequenceHistoryFlush())
+				{
+					// Do some things before accepting the packet.
+					Functor();
+
+					PacketNotify.AckSeq(FNetPacketNotify::SequenceNumberT::SequenceT(InPacketId));
+					AcceptedInPackets.Add(InPacketId);
+				}
+			}
+
+			return bSuccess;
+		}
+
 		void SendAndDeliverTo(FTestNode& Target)
 		{
 			Target.Receive(Send());
@@ -938,6 +998,20 @@ UE_NET_TEST_FIXTURE(FNetPacketNotifyTestUtil, TestGetBothEndsWaitForSequenceWind
 	UE_NET_ASSERT_TRUE(Dst.ValidateExpectedAcceptedPackets());
 }
 
+UE_NET_TEST_FIXTURE(FNetPacketNotifyTestUtil, TestSendWithDeferredAck)
+{
+	FTestNode Src;
+	FTestNode Dst;
+
+	Src.SendBurst(1);
+	Dst.ReceiveWithCallback(Src.Writer,
+	[&Dst]() {
+		Dst.SendBurst(1);
+	}, &Src);
+
+	Src.Receive(Dst.Writer, &Dst);
+}
+
 // Stress test to find random cases.
 #if 0
 UE_NET_TEST_FIXTURE(FNetPacketNotifyTestUtil, TestGetBothEndsWaitForSequenceWindowFlushRandomized)
diff --git a/Engine/Source/Runtime/Engine/Private/Net/NetPacketNotify.cpp b/Engine/Source/Runtime/Engine/Private/Net/NetPacketNotify.cpp
--- a/Engine/Source/Runtime/Engine/Private/Net/NetPacketNotify.cpp
+++ b/Engine/Source/Runtime/Engine/Private/Net/NetPacketNotify.cpp
@@ -218,7 +218,9 @@ bool FNetPacketNotify::WriteHeader(FBitWriter& Writer, bool bRefresh)
 	WrittenInAckSeq = InAckSeq;
 
 	SequenceNumberT::SequenceT Seq = OutSeq.Get();
-	SequenceNumberT::SequenceT AckedSeq = InSeq.Get();
+
+	// If we are waiting for sequencehistory flush we send the latest recvd InSeq, otherwise we send the latest accepted or rejected sequence.
+	SequenceNumberT::SequenceT AckedSeq = bIsWaitingForSequenceHistoryFlush ? InSeq.Get() : InAckSeq.Get();
 
 	// Pack data into a uint
 	uint32 PackedHeader = FPackedHeader::Pack(Seq, AckedSeq, WrittenHistoryWordCount - 1, bIsWaitingForSequenceHistoryFlush);
