Adding a fixture and a couple simple tests for proxy modes in ReplicationSystemTestPlugin.

diff --git a/Engine/Plugins/Runtime/ReplicationSystemTestPlugin/Source/Private/Tests/ReplicationSystem/ReplicationSystemProxyTestFixture.cpp b/Engine/Plugins/Runtime/ReplicationSystemTestPlugin/Source/Private/Tests/ReplicationSystem/ReplicationSystemProxyTestFixture.cpp
new file mode 100644
--- /dev/null
+++ b/Engine/Plugins/Runtime/ReplicationSystemTestPlugin/Source/Private/Tests/ReplicationSystem/ReplicationSystemProxyTestFixture.cpp
@@ -0,0 +1,170 @@
+// Copyright Epic Games, Inc. All Rights Reserved.
+
+#include "ReplicationSystemProxyTestFixture.h"
+
+namespace UE::Net
+{
+
+FReplicationSystemTestProxy::FReplicationSystemTestProxy(const TCHAR* Name)
+{
+	FReplicationSystemTestNode::FReplicationSystemParamsOverride ProxyFrontendParams;
+	ProxyFrontendParams.ProxyType = EProxyType::Frontend;
+
+	FrontendServer.Reset(new FReplicationSystemTestServer(Name, &ProxyFrontendParams));
+}
+
+FReplicationSystemTestClient* FReplicationSystemTestProxy::CreateAndConnectBackendClient(FReplicationSystemTestServer* BackendServer)
+{
+	FReplicationSystemTestNode::FReplicationSystemParamsOverride ProxyBackendParams;
+	ProxyBackendParams.ProxyType = EProxyType::Backend;
+
+	FReplicationSystemTestClient* BackendClient = new FReplicationSystemTestClient(TEXT("ProxyBackendClient"), &ProxyBackendParams);
+	BackendClients.Emplace(BackendClient);
+
+	// The client needs a connection
+	BackendClient->LocalConnectionId = BackendClient->AddConnection();
+
+	// Auto connect to server
+	BackendClient->ConnectionIdOnServer = BackendServer->AddConnection();
+
+	return BackendClient;
+}
+
+void FReplicationSystemProxyTestFixture::SetUp()
+{
+	FNetworkAutomationTestSuiteFixture::SetUp();
+
+	// Fake what we normally get from config
+	DataStreamUtil.SetUp();
+	DataStreamUtil.AddDataStreamDefinition(TEXT("NetToken"), TEXT("/Script/IrisCore.NetTokenDataStream"));
+	DataStreamUtil.AddDataStreamDefinition(TEXT("Replication"), TEXT("/Script/IrisCore.ReplicationDataStream"));
+	DataStreamUtil.FixupDefinitions();
+
+	NetTokenDataStoreUtil.SetUp();
+}
+
+void FReplicationSystemProxyTestFixture::TearDown()
+{
+	Clients.Reset();
+	ProxyServers.Reset();
+	BackendServers.Reset();
+
+	ServerToProxyClients.Reset();
+	ProxyFrontendToClients.Reset();
+
+	DataStreamUtil.TearDown();
+	NetTokenDataStoreUtil.TearDown();
+
+	FNetworkAutomationTestSuiteFixture::TearDown();
+}
+
+FReplicationSystemTestServer* FReplicationSystemProxyTestFixture::CreateServer(uint32 InitialNetRefHandleIndex)
+{
+	FReplicationSystemTestNode::FReplicationSystemParamsOverride BackendServerParams;
+	BackendServerParams.InitialNetRefHandleIndex = InitialNetRefHandleIndex;
+
+	return BackendServers.Emplace_GetRef(new FReplicationSystemTestServer(GetName(), &BackendServerParams)).Get();
+}
+
+FReplicationSystemTestProxy* FReplicationSystemProxyTestFixture::CreateProxy()
+{
+	return ProxyServers.Emplace_GetRef(new FReplicationSystemTestProxy(GetName())).Get();;
+}
+
+FReplicationSystemTestClient* FReplicationSystemProxyTestFixture::CreateClient()
+{
+	return Clients.Emplace_GetRef(new FReplicationSystemTestClient(GetName())).Get();
+}
+
+void FReplicationSystemProxyTestFixture::ConnectProxyToBackendServer(FReplicationSystemTestProxy* Proxy, FReplicationSystemTestServer* BackendServer)
+{
+	FReplicationSystemTestClient* BackendClient = Proxy->CreateAndConnectBackendClient(BackendServer);
+	TArray<FReplicationSystemTestClient*>& ServerClients = ServerToProxyClients.FindOrAdd(BackendServer);
+	ServerClients.Add(BackendClient);
+}
+
+void FReplicationSystemProxyTestFixture::ConnectClientToProxy(FReplicationSystemTestClient* Client, FReplicationSystemTestProxy* Proxy)
+{
+	FReplicationSystemTestServer* ProxyFrontendServer = Proxy->GetFrontendServer();
+
+	// The client needs a connection
+	Client->LocalConnectionId = Client->AddConnection();
+
+	// Auto connect to server
+	Client->ConnectionIdOnServer = ProxyFrontendServer->AddConnection();
+
+	TArray<FReplicationSystemTestClient*>& ProxyFrontendClients = ProxyFrontendToClients.FindOrAdd(ProxyFrontendServer);
+	ProxyFrontendClients.Add(Client);
+}
+
+bool FReplicationSystemProxyTestFixture::UpdateAndSendServerToAllProxies(FReplicationSystemTestServer* Server)
+{
+	bool bSuccess = true;
+
+	Server->NetUpdate();
+
+	TArray<FReplicationSystemTestClient*>* ProxyClients = ServerToProxyClients.Find(Server);
+
+	if (ProxyClients)
+	{
+		for (FReplicationSystemTestClient* Client : *ProxyClients)
+		{
+			bSuccess &= Server->SendAndDeliverTo(Client, true);
+		}
+	}
+	else
+	{
+		bSuccess = false;
+	}
+
+	Server->PostSendUpdate();
+
+	return bSuccess;
+}
+
+bool FReplicationSystemProxyTestFixture::UpdateAndSendProxyToAllClients(FReplicationSystemTestProxy* Proxy)
+{
+	bool bSuccess = true;
+
+	FReplicationSystemTestServer* const ProxyFrontendServer = Proxy->GetFrontendServer();
+	ProxyFrontendServer->NetUpdate();
+
+	TArray<FReplicationSystemTestClient*>* ProxyClients = ProxyFrontendToClients.Find(ProxyFrontendServer);
+
+	if (ProxyClients)
+	{
+		for (FReplicationSystemTestClient* Client : *ProxyClients)
+		{
+			bSuccess &= ProxyFrontendServer->SendAndDeliverTo(Client, true);
+		}
+	}
+	else
+	{
+		bSuccess = false;
+	}
+
+	ProxyFrontendServer->PostSendUpdate();
+
+	return bSuccess;
+}
+
+FReplicationSystemTestClient* FReplicationSystemProxyTestFixture::GetProxyClientForBackendServer(FReplicationSystemTestServer* Server, FReplicationSystemTestProxy* Proxy)
+{
+	TArray<FReplicationSystemTestClient*>* ProxyClientsForServer = ServerToProxyClients.Find(Server);
+		
+	if (ProxyClientsForServer)
+	{
+		for (const TUniquePtr<FReplicationSystemTestClient>& ProxyClient : Proxy->BackendClients)
+		{
+			int32 FoundIndex = ProxyClientsForServer->Find(ProxyClient.Get());
+			if (FoundIndex != INDEX_NONE)
+			{
+				return (*ProxyClientsForServer)[FoundIndex];
+			}
+		}
+	}
+
+	return nullptr;
+}
+
+}
diff --git a/Engine/Plugins/Runtime/ReplicationSystemTestPlugin/Source/Private/Tests/ReplicationSystem/ReplicationSystemProxyTestFixture.h b/Engine/Plugins/Runtime/ReplicationSystemTestPlugin/Source/Private/Tests/ReplicationSystem/ReplicationSystemProxyTestFixture.h
new file mode 100644
--- /dev/null
+++ b/Engine/Plugins/Runtime/ReplicationSystemTestPlugin/Source/Private/Tests/ReplicationSystem/ReplicationSystemProxyTestFixture.h
@@ -0,0 +1,75 @@
+// Copyright Epic Games, Inc. All Rights Reserved.
+
+#pragma once
+
+#include "Containers/Array.h"
+#include "Containers/Map.h"
+#include "NetworkAutomationTest.h"
+#include "ReplicationSystemServerClientTestFixture.h"
+#include "Templates/UniquePtr.h"
+
+namespace UE::Net
+{
+
+class FReplicationSystemTestProxy
+{
+public:
+	// Creates the frontend server
+	explicit FReplicationSystemTestProxy(const TCHAR* Name);
+
+	// Returns non-owning pointer, do not delete it!
+	FReplicationSystemTestClient* CreateAndConnectBackendClient(FReplicationSystemTestServer* BackendServer);
+
+	// Returns non-owning pointer, do not delete it!
+	FReplicationSystemTestServer* GetFrontendServer()
+	{
+		return FrontendServer.Get();
+	}
+
+	TArray<TUniquePtr<FReplicationSystemTestClient>> BackendClients;
+	TUniquePtr<FReplicationSystemTestServer> FrontendServer;
+};
+
+class FReplicationSystemProxyTestFixture : public FNetworkAutomationTestSuiteFixture
+{
+public:
+	FReplicationSystemProxyTestFixture() : FNetworkAutomationTestSuiteFixture() {}
+
+protected:
+	virtual void SetUp() override;
+	virtual void TearDown() override;
+
+	// Returns non-owning pointer, do not delete it!
+	FReplicationSystemTestServer* CreateServer(uint32 InitialNetRefHandleIndex = 1);
+
+	// Returns non-owning pointer, do not delete it!
+	FReplicationSystemTestProxy* CreateProxy();
+
+	// Returns non-owning pointer, do not delete it!
+	FReplicationSystemTestClient* CreateClient();
+
+	void ConnectProxyToBackendServer(FReplicationSystemTestProxy* Proxy, FReplicationSystemTestServer* BackendServer);
+	void ConnectClientToProxy(FReplicationSystemTestClient* Client, FReplicationSystemTestProxy* Proxy);
+
+	bool UpdateAndSendServerToAllProxies(FReplicationSystemTestServer* Server);
+	bool UpdateAndSendProxyToAllClients(FReplicationSystemTestProxy* Proxy);
+
+	FReplicationSystemTestClient* GetProxyClientForBackendServer(FReplicationSystemTestServer* Server, FReplicationSystemTestProxy* Proxy);
+
+	FDataStreamTestUtil DataStreamUtil;
+	FNetTokenDataStoreTestUtil NetTokenDataStoreUtil;
+
+	TArray<TUniquePtr<FReplicationSystemTestClient>> Clients;
+	TArray<TUniquePtr<FReplicationSystemTestProxy>> ProxyServers;
+	TArray<TUniquePtr<FReplicationSystemTestServer>> BackendServers;
+
+	// Map of backend server nodes to their corresponding client nodes on proxies.
+	// Stores non-owning pointers, do not delete them!
+	TMap<FReplicationSystemTestServer*, TArray<FReplicationSystemTestClient*>> ServerToProxyClients;
+
+	// Map of frontend proxy server nodes to their corresponding client nodes
+	// Stores non-owning pointers, do not delete them!
+	TMap<FReplicationSystemTestServer*, TArray<FReplicationSystemTestClient*>> ProxyFrontendToClients;
+};
+
+}
diff --git a/Engine/Plugins/Runtime/ReplicationSystemTestPlugin/Source/Private/Tests/ReplicationSystem/ReplicationSystemServerClientTestFixture.cpp b/Engine/Plugins/Runtime/ReplicationSystemTestPlugin/Source/Private/Tests/ReplicationSystem/ReplicationSystemServerClientTestFixture.cpp
--- a/Engine/Plugins/Runtime/ReplicationSystemTestPlugin/Source/Private/Tests/ReplicationSystem/ReplicationSystemServerClientTestFixture.cpp
+++ b/Engine/Plugins/Runtime/ReplicationSystemTestPlugin/Source/Private/Tests/ReplicationSystem/ReplicationSystemServerClientTestFixture.cpp
@@ -160,6 +160,8 @@ void FReplicationSystemTestNode::Setup(bool bIsServer, const TCHAR* Name, const
 		{
 			Params.bAllowMinimalUpdateIfNoConnections = ParamsOverride->bAllowMinimalUpdateIfNoConnections.GetValue();
 		}
+		Params.InitialNetRefHandleIndex = ParamsOverride->InitialNetRefHandleIndex;
+		Params.ProxyType = ParamsOverride->ProxyType;
 	}
 
 	LOG_SCOPE_VERBOSITY_OVERRIDE(LogIris, ELogVerbosity::Error);
@@ -496,8 +498,8 @@ float FReplicationSystemTestNode::ConvertPollPeriodIntoFrequency(uint32 PollPeri
 // Class FReplicationSystemTestClient
 //*****************************************************************************
 
-FReplicationSystemTestClient::FReplicationSystemTestClient(const TCHAR* Name)
-: FReplicationSystemTestNode(false, Name)
+FReplicationSystemTestClient::FReplicationSystemTestClient(const TCHAR* Name, const FReplicationSystemTestNode::FReplicationSystemParamsOverride* ParamsOverride)
+: FReplicationSystemTestNode(false, Name, ParamsOverride)
 , ConnectionIdOnServer(~0U)
 {
 }
diff --git a/Engine/Plugins/Runtime/ReplicationSystemTestPlugin/Source/Private/Tests/ReplicationSystem/ReplicationSystemServerClientTestFixture.h b/Engine/Plugins/Runtime/ReplicationSystemTestPlugin/Source/Private/Tests/ReplicationSystem/ReplicationSystemServerClientTestFixture.h
--- a/Engine/Plugins/Runtime/ReplicationSystemTestPlugin/Source/Private/Tests/ReplicationSystem/ReplicationSystemServerClientTestFixture.h
+++ b/Engine/Plugins/Runtime/ReplicationSystemTestPlugin/Source/Private/Tests/ReplicationSystem/ReplicationSystemServerClientTestFixture.h
@@ -103,6 +103,8 @@ public:
 		bool bUseRemoteObjectReferences = false;
 		bool bAllowParallelTasks = false;
 		TOptional<bool> bAllowMinimalUpdateIfNoConnections;
+		uint32 InitialNetRefHandleIndex = 1;
+		UE::Net::EProxyType ProxyType = UE::Net::EProxyType::None;
 	};
 
 	enum EDelaySetup { DelaySetup=0 };
@@ -229,7 +231,7 @@ private:
 class FReplicationSystemTestClient : public FReplicationSystemTestNode
 {
 public:
-	FReplicationSystemTestClient(const TCHAR* Name);
+	explicit FReplicationSystemTestClient(const TCHAR* Name, const FReplicationSystemTestNode::FReplicationSystemParamsOverride* ParamsOverride = nullptr);
 	explicit FReplicationSystemTestClient(FReplicationSystemTestNode::EDelaySetup Delay) : FReplicationSystemTestNode(Delay) {}
 
 	// Tick and send packets to the server
diff --git a/Engine/Plugins/Runtime/ReplicationSystemTestPlugin/Source/Private/Tests/ReplicationSystem/TestReplicationSystemProxy.cpp b/Engine/Plugins/Runtime/ReplicationSystemTestPlugin/Source/Private/Tests/ReplicationSystem/TestReplicationSystemProxy.cpp
new file mode 100644
--- /dev/null
+++ b/Engine/Plugins/Runtime/ReplicationSystemTestPlugin/Source/Private/Tests/ReplicationSystem/TestReplicationSystemProxy.cpp
@@ -0,0 +1,133 @@
+// Copyright Epic Games, Inc. All Rights Reserved.
+
+#include "ReplicationSystemProxyTestFixture.h"
+#include "Iris/ReplicationSystem/ReplicationFragmentUtil.h"
+#include "Iris/ReplicationSystem/ReplicationSystem.h"
+
+namespace UE::Net::Private
+{
+
+UE_NET_TEST_FIXTURE(FReplicationSystemProxyTestFixture, ReplicateAndDestroySingleObject)
+{
+	// Add a backend server
+	FReplicationSystemTestServer* Server = CreateServer();
+
+	// Add a proxy server
+	FReplicationSystemTestProxy* Proxy = CreateProxy();
+
+	// Add a client
+	FReplicationSystemTestClient* Client = CreateClient();
+
+	// Make connections
+	ConnectProxyToBackendServer(Proxy, Server);
+	ConnectClientToProxy(Client, Proxy);
+
+	// Spawn object on server
+	UTestReplicatedIrisObject* ServerObject = Server->CreateObject();
+
+	// Send and deliver packet, server -> proxy
+	UpdateAndSendServerToAllProxies(Server);
+
+	// Verify that created server handle now also exists on the proxy
+	FReplicationSystemTestClient* ProxyClient = GetProxyClientForBackendServer(Server, Proxy);
+	UTestReplicatedIrisObject* ProxyTestObject = ProxyClient->GetObjectAs<UTestReplicatedIrisObject>(ServerObject->NetRefHandle);
+	UE_NET_ASSERT_NE(ProxyTestObject, nullptr);
+
+	// Let the proxy replicate the object to its clients
+	Proxy->GetFrontendServer()->GetReplicationBridge()->BeginReplication(ProxyTestObject);
+
+	// Send and deliver packet, proxy -> client
+	UpdateAndSendProxyToAllClients(Proxy);
+
+	// Verify that created server handle now also exists on client
+	UE_NET_ASSERT_TRUE(Client->IsResolvableNetRefHandle(ServerObject->NetRefHandle));
+
+	// Destroy the spawned object on server
+	Server->DestroyObject(ServerObject);
+
+	// Send and deliver packet, server -> proxy
+	UpdateAndSendServerToAllProxies(Server);
+
+	// Verify that object now is destroyed on proxy as well
+	UE_NET_ASSERT_FALSE(ProxyClient->IsValidNetRefHandle(ServerObject->NetRefHandle));
+
+	// Have the proxy destroy the object on its clients
+	Proxy->GetFrontendServer()->GetReplicationBridge()->StopReplicatingNetRefHandle(ServerObject->NetRefHandle, EEndReplicationFlags::Destroy);
+
+	// Send and deliver packet, proxy -> client
+	UpdateAndSendProxyToAllClients(Proxy);
+
+	// Verify that object now is destroyed on client as well
+	UE_NET_ASSERT_FALSE(Client->IsValidNetRefHandle(ServerObject->NetRefHandle));
+}
+
+// Test objects and simple properties replicating from two backend servers through one proxy to one client.
+UE_NET_TEST_FIXTURE(FReplicationSystemProxyTestFixture, TestTwoServers)
+{
+	// Add two backend servers
+	constexpr uint32 InitialNetRefHandleIndexServer1 = 1000;
+	FReplicationSystemTestServer* Server1 = CreateServer(InitialNetRefHandleIndexServer1);
+
+	constexpr uint32 InitialNetRefHandleIndexServer2 = 2000;
+	FReplicationSystemTestServer* Server2 = CreateServer(InitialNetRefHandleIndexServer2);
+
+	// Add a proxy server
+	FReplicationSystemTestProxy* Proxy = CreateProxy();
+
+	// Add a client
+	FReplicationSystemTestClient* Client = CreateClient();
+
+	// Make connections
+	ConnectProxyToBackendServer(Proxy, Server1);
+	ConnectProxyToBackendServer(Proxy, Server2);
+	ConnectClientToProxy(Client, Proxy);
+
+	// Spawn one object on each server
+	UTestReplicatedIrisObject* ServerObject1 = Server1->CreateObject();
+	UTestReplicatedIrisObject* ServerObject2 = Server2->CreateObject();
+	
+	const int32 IntServer1 = 10;
+	const int32 IntServer2 = 20;
+
+	ServerObject1->IntA = IntServer1;
+	ServerObject2->IntA = IntServer2;
+	
+	// Send and deliver packet, servers -> proxy
+	UpdateAndSendServerToAllProxies(Server1);
+	UpdateAndSendServerToAllProxies(Server2);
+
+	// Verify that objects have been spawned on the proxy
+	FReplicationSystemTestClient* ProxyBackendClient1 = GetProxyClientForBackendServer(Server1, Proxy);
+	FReplicationSystemTestClient* ProxyBackendClient2 = GetProxyClientForBackendServer(Server2, Proxy);
+	UTestReplicatedIrisObject* ProxyObject1 = ProxyBackendClient1->GetObjectAs<UTestReplicatedIrisObject>(ServerObject1->NetRefHandle);
+	UTestReplicatedIrisObject* ProxyObject2 = ProxyBackendClient2->GetObjectAs<UTestReplicatedIrisObject>(ServerObject2->NetRefHandle);
+	UE_NET_ASSERT_NE(ProxyObject1, nullptr);
+	UE_NET_ASSERT_NE(ProxyObject2, nullptr);
+	UE_NET_ASSERT_NE(ProxyObject1, ProxyObject2);
+
+	// Verify that the values are correct on the proxy
+	UE_NET_ASSERT_EQ(ProxyObject1->IntA, IntServer1);
+	UE_NET_ASSERT_EQ(ProxyObject2->IntA, IntServer2);
+
+	// Let the proxy replicate the object to its clients
+	Proxy->GetFrontendServer()->GetReplicationBridge()->BeginReplication(ProxyObject1);
+	Proxy->GetFrontendServer()->GetReplicationBridge()->BeginReplication(ProxyObject2);
+
+	// Send and deliver packet, proxy -> client
+	UpdateAndSendProxyToAllClients(Proxy);
+
+	// Verify that objects have been spawned on client
+	UTestReplicatedIrisObject* ClientObject1 = Client->GetObjectAs<UTestReplicatedIrisObject>(ServerObject1->NetRefHandle);
+	UTestReplicatedIrisObject* ClientObject2 = Client->GetObjectAs<UTestReplicatedIrisObject>(ServerObject2->NetRefHandle);
+	UE_NET_ASSERT_NE(ClientObject1, nullptr);
+	UE_NET_ASSERT_NE(ClientObject2, nullptr);
+	UE_NET_ASSERT_NE(ClientObject1, ClientObject2);
+
+	// Verify that the values are correct on the client
+	UE_NET_ASSERT_EQ(ClientObject1->IntA, IntServer1);
+	UE_NET_ASSERT_EQ(ClientObject2->IntA, IntServer2);
+}
+
+} // end namespace UE::Net::Private
+
+
