Don't fail tests just because ensures are disabled.

diff --git a/Engine/Plugins/Runtime/ReplicationSystemTestPlugin/Source/Private/Tests/ReplicationSystem/TestDependentObjects.cpp b/Engine/Plugins/Runtime/ReplicationSystemTestPlugin/Source/Private/Tests/ReplicationSystem/TestDependentObjects.cpp
--- a/Engine/Plugins/Runtime/ReplicationSystemTestPlugin/Source/Private/Tests/ReplicationSystem/TestDependentObjects.cpp
+++ b/Engine/Plugins/Runtime/ReplicationSystemTestPlugin/Source/Private/Tests/ReplicationSystem/TestDependentObjects.cpp
@@ -1105,7 +1105,7 @@ UE_NET_TEST_FIXTURE(FTestFilteredDependentObjectFixture, TestCircularDependencyI
 		LOG_SCOPE_VERBOSITY_OVERRIDE(LogIris, ELogVerbosity::Fatal);
 		FTestEnsureScope SuppressEnsureScope;
 		Bridge->AddDependentObject(ServerObjects[ObjectCount - 1]->NetRefHandle, ServerObjects[0]->NetRefHandle);
-		UE_NET_ASSERT_EQ(SuppressEnsureScope.GetCount(), 1);
+		UE_NET_ASSERT_EQ(SuppressEnsureScope.GetCount(), (DO_ENSURE ? 1 : 0));
 	}
 }
 
