Add NetObjectBlobHandler to PartialNetBlobTestFixture

#Synced-CL 47893047

diff --git a/Engine/Plugins/Runtime/ReplicationSystemTestPlugin/Source/Private/Tests/ReplicationSystem/NetBlob/PartialNetBlobTestFixture.cpp b/Engine/Plugins/Runtime/ReplicationSystemTestPlugin/Source/Private/Tests/ReplicationSystem/NetBlob/PartialNetBlobTestFixture.cpp
--- a/Engine/Plugins/Runtime/ReplicationSystemTestPlugin/Source/Private/Tests/ReplicationSystem/NetBlob/PartialNetBlobTestFixture.cpp
+++ b/Engine/Plugins/Runtime/ReplicationSystemTestPlugin/Source/Private/Tests/ReplicationSystem/NetBlob/PartialNetBlobTestFixture.cpp
@@ -67,6 +67,7 @@ void FPartialNetBlobTestFixture::AddNetBlobHandlerDefinitions()
 		// Used for tests involving server and clients.
 		{TEXT("MockNetObjectAttachmentHandler"),}, 
 		{TEXT("PartialNetObjectAttachmentHandler"),}, 
+		{TEXT("NetObjectBlobHandler"),}, 
 	};
 	Super::AddNetBlobHandlerDefinitions(PartialNetBlobHandlerDefinitions, UE_ARRAY_COUNT(PartialNetBlobHandlerDefinitions));
 }
