Increase the size of FNetHandle to hold the full FRemoteObject server identifier and serial number.

#Synced-CL 47902096

diff --git a/Engine/Plugins/Runtime/ReplicationSystemTestPlugin/Source/Private/Tests/ReplicationState/TestGeneratedReplicationState.cpp b/Engine/Plugins/Runtime/ReplicationSystemTestPlugin/Source/Private/Tests/ReplicationState/TestGeneratedReplicationState.cpp
--- a/Engine/Plugins/Runtime/ReplicationSystemTestPlugin/Source/Private/Tests/ReplicationState/TestGeneratedReplicationState.cpp
+++ b/Engine/Plugins/Runtime/ReplicationSystemTestPlugin/Source/Private/Tests/ReplicationState/TestGeneratedReplicationState.cpp
@@ -75,12 +75,12 @@ UE_NET_TEST(FTestTypedReplicationState, ValidateDescriptor)
 	UE_NET_ASSERT_EQ(0u, (uint32)Descriptor->TagCount);
 	UE_NET_ASSERT_EQ(Descriptor->MemberTagDescriptors, nullptr);
 	UE_NET_ASSERT_EQ(Descriptor->MemberProperties, nullptr);
-	UE_NET_ASSERT_EQ(48u, Descriptor->ExternalSize);
+	UE_NET_ASSERT_EQ(Align(40u + static_cast<uint32>(sizeof(FReplicationStateHeader)), 8U), Descriptor->ExternalSize);
 	UE_NET_ASSERT_EQ(40u, Descriptor->InternalSize);
 	UE_NET_ASSERT_EQ(8u, (uint32)Descriptor->ExternalAlignment);
 	UE_NET_ASSERT_EQ(8u, (uint32)Descriptor->InternalAlignment);
 	UE_NET_ASSERT_EQ(4u, (uint32)Descriptor->ChangeMaskBitCount);
-	UE_NET_ASSERT_EQ(4u, Descriptor->ChangeMasksExternalOffset);
+	UE_NET_ASSERT_EQ(static_cast<uint32>(sizeof(FReplicationStateHeader)), Descriptor->ChangeMasksExternalOffset);
 }
 
 UE_NET_TEST(FTestTypedReplicationState, ValidateState)
@@ -93,7 +93,7 @@ UE_NET_TEST(FTestTypedReplicationState, ValidateState)
 	// so test by casting as it always should be the first member
 	const UE::Net::FReplicationStateHeader& Header = UE::Net::Private::GetReplicationStateHeader(reinterpret_cast<uint8*>(&State), State.GetReplicationStateDescriptor());
 
-	UE_NET_ASSERT_EQ(FReplicationStateHeaderAccessor::GetNetHandleId(Header), 0U);
+	UE_NET_ASSERT_EQ(FReplicationStateHeaderAccessor::GetNetHandleId(Header), 0ULL);
 	UE_NET_ASSERT_FALSE(FReplicationStateHeaderAccessor::GetIsInitStateDirty(Header));
 	UE_NET_ASSERT_FALSE(FReplicationStateHeaderAccessor::GetIsStateDirty(Header));
 
diff --git a/Engine/Plugins/Runtime/ReplicationSystemTestPlugin/Source/Private/Tests/ReplicationState/TestPropertyReplicationState.cpp b/Engine/Plugins/Runtime/ReplicationSystemTestPlugin/Source/Private/Tests/ReplicationState/TestPropertyReplicationState.cpp
--- a/Engine/Plugins/Runtime/ReplicationSystemTestPlugin/Source/Private/Tests/ReplicationState/TestPropertyReplicationState.cpp
+++ b/Engine/Plugins/Runtime/ReplicationSystemTestPlugin/Source/Private/Tests/ReplicationState/TestPropertyReplicationState.cpp
@@ -310,7 +310,7 @@ UE_NET_TEST_FIXTURE(FTestPropertyReplicationStateContext, PropertyReplicationSta
 	OtherState = BoundState;
 
 	// Verify that we did not overwrite the NetHandle
-	UE_NET_ASSERT_EQ(FReplicationStateHeaderAccessor::GetNetHandleId(GetReplicationStateHeader(OtherState.GetStateBuffer(), Descriptor)), 0U);
+	UE_NET_ASSERT_EQ(FReplicationStateHeaderAccessor::GetNetHandleId(GetReplicationStateHeader(OtherState.GetStateBuffer(), Descriptor)), 0ULL);
 
 	// check that original states are still dirty and that we also have dirtied the value that was modified in OtherState
 	UE_NET_ASSERT_TRUE(OtherState.IsDirty(0));
diff --git a/Engine/Plugins/Runtime/ReplicationSystemTestPlugin/Source/Private/Tests/ReplicationState/TestReplicationStateDescriptorBuilder.cpp b/Engine/Plugins/Runtime/ReplicationSystemTestPlugin/Source/Private/Tests/ReplicationState/TestReplicationStateDescriptorBuilder.cpp
--- a/Engine/Plugins/Runtime/ReplicationSystemTestPlugin/Source/Private/Tests/ReplicationState/TestReplicationStateDescriptorBuilder.cpp
+++ b/Engine/Plugins/Runtime/ReplicationSystemTestPlugin/Source/Private/Tests/ReplicationState/TestReplicationStateDescriptorBuilder.cpp
@@ -7,6 +7,7 @@
 #include "Iris/ReplicationState/ReplicationStateDescriptorBuilder.h"
 #include "Iris/ReplicationState/InternalPropertyReplicationState.h"
 #include "Iris/ReplicationState/PropertyNetSerializerInfoRegistry.h"
+#include "Iris/ReplicationState/ReplicationStateFwd.h"
 #include "Iris/Serialization/InternalNetSerializers.h"
 #include "Iris/Serialization/NetSerializers.h"
 #include "Iris/Serialization/RemoteObjectReferenceNetSerializer.h"
@@ -236,16 +237,16 @@ static FReplicationStateMemberTraitsDescriptor TestClassMemberTraitsDescriptors[
 };
 
 // 4 byte header + changemask + pollmask
-constexpr SIZE_T TestClassMemberDescriptors_MemberOffset = 12;
+constexpr SIZE_T TestClassMemberDescriptors_MemberOffset = 8 + sizeof(FReplicationStateHeader);
 static FReplicationStateMemberDescriptor TestClassMemberDescriptors[] = 
 {
 	{TestClassMemberDescriptors_MemberOffset, 0},
 	{TestClassMemberDescriptors_MemberOffset + 2, 2},
 	{TestClassMemberDescriptors_MemberOffset + 4, 4},
-	{TestClassMemberDescriptors_MemberOffset + 12, 8},
-	{TestClassMemberDescriptors_MemberOffset + 20, 16},
-	{TestClassMemberDescriptors_MemberOffset + 21, 17},
-	{TestClassMemberDescriptors_MemberOffset + 22, 18}
+	{TestClassMemberDescriptors_MemberOffset + 8, 8},
+	{TestClassMemberDescriptors_MemberOffset + 16, 16},
+	{TestClassMemberDescriptors_MemberOffset + 17, 17},
+	{TestClassMemberDescriptors_MemberOffset + 18, 18}
 };
 
 static FReplicationStateMemberChangeMaskDescriptor TestClassMemberChangeMaskDescriptors[] =
@@ -276,7 +277,7 @@ static const FReplicationStateDescriptor TestClassReplictionStateDescriptor
 	static_cast<const UScriptStruct*>(nullptr), // BaseStruct
 	static_cast<const FNetDebugName*>(nullptr), // DebugName
 	static_cast<const FReplicationStateMemberDebugDescriptor*>(nullptr), // MemberDebugDescriptors
-	40,			//uint32 ExternalSize;
+	Align<uint32>(32U + sizeof(FReplicationStateHeader), 8U),	//uint32 ExternalSize;
 	24,			//uint32 InternalSize;
 	8,			//uint16 ExternalAlignment;
 	8,			//uint16 InternalAlignment;
@@ -286,7 +287,7 @@ static const FReplicationStateDescriptor TestClassReplictionStateDescriptor
 	0,			//uint16 ReferenceCount
 	0,			//uint16 RepIndexCount
 	7,			//uint16 ChangeMaskBitCount;			// How many bits do we need for our tracking of dirty states
-	4,			//uint32 ChangeMaskExternalOffset;		// This is the offset to where we store our state mask data in the external state.
+	sizeof(FReplicationStateHeader), //uint32 ChangeMaskExternalOffset;		// This is the offset to where we store our state mask data in the external state.
 	{0},			//FReplicationStateIdentifier DescriptorIdentifier;
 	ConstructPropertyReplicationState,	//ConstructReplicationStateFunc CreateReplicationState;
 	DestructPropertyReplicationState,	//DestructReplicationStateFunc DestroyReplicationState;
@@ -312,7 +313,7 @@ static const FReplicationStateDescriptor TestClassReplictionStateDescriptorNotRe
 	static_cast<const UScriptStruct*>(nullptr), // BaseStruct
 	static_cast<const FNetDebugName*>(nullptr), // DebugName
 	static_cast<const FReplicationStateMemberDebugDescriptor*>(nullptr), // MemberDebugDescriptors
-	40,			//uint32 ExternalSize;
+	Align<uint32>(32U + sizeof(FReplicationStateHeader), 8U),	//uint32 ExternalSize;
 	24,			//uint32 InternalSize;
 	8,			//uint16 ExternalAlignment;
 	8,			//uint16 InternalAlignment;
@@ -322,7 +323,7 @@ static const FReplicationStateDescriptor TestClassReplictionStateDescriptorNotRe
 	0,			//uint16 ReferenceCount;
 	0,			//uint16 RepIndexCount
 	7,			//uint16 ChangeMaskBitCount;			// How many bits do we need for our tracking of dirty states
-	4,			//uint32 ChangeMaskExternalOffset;		// This is the offset to where we store our state mask data in the external state.
+	sizeof(FReplicationStateHeader), //uint32 ChangeMaskExternalOffset;		// This is the offset to where we store our state mask data in the external state.
 	{0},			//FReplicationStateIdentifier DescriptorIdentifier;
 	ConstructPropertyReplicationState,	//ConstructReplicationStateFunc CreateReplicationState;
 	DestructPropertyReplicationState,	//DestructReplicationStateFunc DestroyReplicationState;
@@ -348,7 +349,7 @@ static const FReplicationStateDescriptor TestClassWithNonReplicatedDataReplicati
 	static_cast<const UScriptStruct*>(nullptr), // BaseStruct
 	static_cast<const FNetDebugName*>(nullptr), // DebugName
 	static_cast<const FReplicationStateMemberDebugDescriptor*>(nullptr), // MemberDebugDescriptors
-	40,			//uint32 ExternalSize;
+	Align<uint32>(32U + sizeof(FReplicationStateHeader), 8U),	//uint32 ExternalSize;
 	24,			//uint32 InternalSize;
 	8,			//uint16 ExternalAlignment;
 	8,			//uint16 InternalAlignment;
@@ -358,7 +359,7 @@ static const FReplicationStateDescriptor TestClassWithNonReplicatedDataReplicati
 	0,			//uint16 ReferenceCount;
 	0,			//uint16 RepIndexCount
 	7,			//uint16 ChangeMaskBitCount;			// How many bits do we need for our tracking of dirty states
-	4,			//uint32 ChangeMaskExternalOffset;		// This is the offset to where we store our state mask data in the external state.
+	sizeof(FReplicationStateHeader), //uint32 ChangeMaskExternalOffset;		// This is the offset to where we store our state mask data in the external state.
 	{0},			//FReplicationStateIdentifier DescriptorIdentifier;
 	ConstructPropertyReplicationState,	//ConstructReplicationStateFunc CreateReplicationState;
 	DestructPropertyReplicationState,	//DestructReplicationStateFunc DestroyReplicationState;
@@ -400,18 +401,18 @@ namespace TestReplicationStateDescriptor_FakeDescriptorData
 {
 
 // 4 byte header + changemask + pollmask
-constexpr SIZE_T TestClassWithInheritanceMemberDescriptors_MemberOffset = 12;
+constexpr SIZE_T TestClassWithInheritanceMemberDescriptors_MemberOffset = 8 + sizeof(FReplicationStateHeader);
 static FReplicationStateMemberDescriptor TestClassWithInheritanceMemberDescriptors[] = 
 {
 	{TestClassWithInheritanceMemberDescriptors_MemberOffset, 0},
 	{TestClassWithInheritanceMemberDescriptors_MemberOffset + 2, 2},
 	{TestClassWithInheritanceMemberDescriptors_MemberOffset + 4, 4},
-	{TestClassWithInheritanceMemberDescriptors_MemberOffset + 12, 8},
-	{TestClassWithInheritanceMemberDescriptors_MemberOffset + 20, 16},
-	{TestClassWithInheritanceMemberDescriptors_MemberOffset + 21, 17},
-	{TestClassWithInheritanceMemberDescriptors_MemberOffset + 22, 18},
-	{TestClassWithInheritanceMemberDescriptors_MemberOffset + 23, 19},
-	{TestClassWithInheritanceMemberDescriptors_MemberOffset + 24, 20}
+	{TestClassWithInheritanceMemberDescriptors_MemberOffset + 8, 8},
+	{TestClassWithInheritanceMemberDescriptors_MemberOffset + 16, 16},
+	{TestClassWithInheritanceMemberDescriptors_MemberOffset + 17, 17},
+	{TestClassWithInheritanceMemberDescriptors_MemberOffset + 18, 18},
+	{TestClassWithInheritanceMemberDescriptors_MemberOffset + 19, 19},
+	{TestClassWithInheritanceMemberDescriptors_MemberOffset + 20, 20}
 };
 
 // Fake member serializer, config is complicated to setup
@@ -470,7 +471,7 @@ static FReplicationStateDescriptor TestClassWithInheritanceReplictionStateDescri
 	static_cast<const UScriptStruct*>(nullptr), // BaseStruct
 	static_cast<const FNetDebugName*>(nullptr), // DebugName
 	static_cast<const FReplicationStateMemberDebugDescriptor*>(nullptr), // MemberDebugDescriptors
-	40, // ExternalSize (rounded up since it also contains changemask)
+	Align<uint32>(32U + sizeof(FReplicationStateHeader), 8U),	//uint32 ExternalSize;
 	24,	// InternalSize
 	8,	// ExternalAlignment
 	8,	// InternalAlignment
@@ -480,7 +481,7 @@ static FReplicationStateDescriptor TestClassWithInheritanceReplictionStateDescri
 	0,	// ReferenceCount;
 	0,	// RepIndexCount
 	9,	// ChangeMaskBitCount
-	4,	// ChangeMaskExternalOffset
+	sizeof(FReplicationStateHeader),	// ChangeMaskExternalOffset
 	{0}, // FReplicationStateIdentifier DescriptorIdentifier
 	ConstructPropertyReplicationState,	// ConstructReplicationStateFunc CreateReplicationState;
 	DestructPropertyReplicationState,	// DestructReplicationStateFunc DestroyReplicationState;
@@ -503,7 +504,7 @@ namespace TestReplicationStateDescriptor_FakeDescriptorData
 
 // Fake desc for FTestReplicationStateDescriptor_TestStruct
 // 4 byte header + changemask + pollmask
-constexpr SIZE_T TestClassWithInheritanceNoSuperMemberDescriptors_MemberOffset = 12;
+constexpr SIZE_T TestClassWithInheritanceNoSuperMemberDescriptors_MemberOffset = 8 + sizeof(FReplicationStateHeader);
 static FReplicationStateMemberDescriptor TestClassWithInheritanceNoSuperMemberDescriptors[] = 
 {
 	{TestClassWithInheritanceNoSuperMemberDescriptors_MemberOffset, 0},
@@ -545,9 +546,9 @@ static FReplicationStateDescriptor TestClassWithInheritanceNoSuperReplictionStat
 	static_cast<const UScriptStruct*>(nullptr), // BaseStruct
 	static_cast<const FNetDebugName*>(nullptr), // DebugName
 	static_cast<const FReplicationStateMemberDebugDescriptor*>(nullptr), // MemberDebugDescriptors
-	16, // ExternalSize (rounded up since it also contains changemask)
+	Align<uint32>(16U + sizeof(FReplicationStateHeader), 8U),	//uint32 ExternalSize;
 	4,	// InternalSize
-	4,	// ExternalAlignment
+	8,	// ExternalAlignment
 	2,	// InternalAlignment
 	2,	// MemberCount
 	0,	// FunctionCount
@@ -555,7 +556,7 @@ static FReplicationStateDescriptor TestClassWithInheritanceNoSuperReplictionStat
 	0,	// ReferenceCount;
 	0,	// RepIndexCount
 	2,	// ChangeMaskBitCount
-	4,	// ChangeMaskExternalOffset
+	sizeof(FReplicationStateHeader),	// ChangeMaskExternalOffset
 	{0}, // FReplicationStateIdentifier DescriptorIdentifier
 	ConstructPropertyReplicationState,	// ConstructReplicationStateFunc CreateReplicationState;
 	DestructPropertyReplicationState,	// DestructReplicationStateFunc DestroyReplicationState;
@@ -577,7 +578,7 @@ namespace TestReplicationStateDescriptor_FakeDescriptorData
 
 // Fake desc for FTestReplicationStateDescriptor_TestStruct
 // 4 byte header + changemask + pollmask
-constexpr SIZE_T TestClassWithReplicatedStructMemberDescriptors_MemberOffset = 16;
+constexpr SIZE_T TestClassWithReplicatedStructMemberDescriptors_MemberOffset = 8 + sizeof(FReplicationStateHeader);
 static FReplicationStateMemberDescriptor TestClassWithReplicatedStructMemberDescriptors[] = 
 {
 	{TestClassWithReplicatedStructMemberDescriptors_MemberOffset, 0},
@@ -615,7 +616,7 @@ static FReplicationStateDescriptor TestClassWithReplicatedStructReplictionStateD
 	static_cast<const UScriptStruct*>(nullptr), // BaseStruct
 	static_cast<const FNetDebugName*>(nullptr), // DebugName
 	static_cast<const FReplicationStateMemberDebugDescriptor*>(nullptr), // MemberDebugDescriptors
-	40, // ExternalSize (Header + ChangeMasks)
+	Align<uint32>(32U + sizeof(FReplicationStateHeader), 8U),	//uint32 ExternalSize;
 	24,	// InternalSize
 	8,	// ExternalAlignment
 	8,	// InternalAlignment
@@ -625,7 +626,7 @@ static FReplicationStateDescriptor TestClassWithReplicatedStructReplictionStateD
 	0,	// ReferenceCount;
 	0,	// RepIndexCount
 	1,	// ChangeMaskBitCount
-	4,	// ChangeMaskExternalOffset
+	sizeof(FReplicationStateHeader),	// ChangeMaskExternalOffset
 	{0}, // FReplicationStateIdentifier DescriptorIdentifier
 	ConstructPropertyReplicationState,	// ConstructReplicationStateFunc CreateReplicationState;
 	DestructPropertyReplicationState,	// DestructReplicationStateFunc DestroyReplicationState;
diff --git a/Engine/Plugins/Runtime/ReplicationSystemTestPlugin/Source/Private/Tests/ReplicationSystem/TestReplicationBridge.cpp b/Engine/Plugins/Runtime/ReplicationSystemTestPlugin/Source/Private/Tests/ReplicationSystem/TestReplicationBridge.cpp
--- a/Engine/Plugins/Runtime/ReplicationSystemTestPlugin/Source/Private/Tests/ReplicationSystem/TestReplicationBridge.cpp
+++ b/Engine/Plugins/Runtime/ReplicationSystemTestPlugin/Source/Private/Tests/ReplicationSystem/TestReplicationBridge.cpp
@@ -35,7 +35,7 @@ UE_NET_TEST_FIXTURE(FReplicationBridgeTestFixture, CreateObjectWithSingleFragmen
 	UE_NET_ASSERT_EQ(1u, Protocol->ReplicationStateCount);
 	UE_NET_ASSERT_EQ(4u, Protocol->InternalTotalAlignment);
 	UE_NET_ASSERT_EQ(20u, Protocol->InternalTotalSize);
-	UE_NET_ASSERT_EQ(36u, Protocol->MaxExternalStateSize);
+	UE_NET_ASSERT_EQ(Align(32u + static_cast<uint32>(sizeof(FReplicationStateHeader)), 8U), Protocol->MaxExternalStateSize);
 	UE_NET_ASSERT_EQ(5u, Protocol->ChangeMaskBitCount);
 	UE_NET_ASSERT_EQ(EReplicationProtocolTraits::SupportsDeltaCompression, Protocol->ProtocolTraits);
 
@@ -63,7 +63,7 @@ UE_NET_TEST_FIXTURE(FReplicationBridgeTestFixture, CreateObjectWithMultipleFragm
 	UE_NET_ASSERT_EQ(25u, Protocol->ReplicationStateCount);
 	UE_NET_ASSERT_EQ(8u, Protocol->InternalTotalAlignment);
 	UE_NET_ASSERT_EQ(16u + PropertyComponentCount*PropertyComponentSize, Protocol->InternalTotalSize);
-	UE_NET_ASSERT_EQ(40u, Protocol->MaxExternalStateSize);
+	UE_NET_ASSERT_EQ(Align(36u + static_cast<uint32>(sizeof(FReplicationStateHeader)), 8U), Protocol->MaxExternalStateSize);
 	UE_NET_ASSERT_EQ(5U + PropertyComponentCount*RegularStatePropertyCountInPropertyComponent, Protocol->ChangeMaskBitCount);
 	UE_NET_ASSERT_EQ(EReplicationProtocolTraits::SupportsDeltaCompression, Protocol->ProtocolTraits);
 
@@ -91,7 +91,7 @@ UE_NET_TEST_FIXTURE(FReplicationBridgeTestFixture, CreateObjectWithMixedFragment
 	UE_NET_ASSERT_EQ(28u, Protocol->ReplicationStateCount);
 	UE_NET_ASSERT_EQ(8u, Protocol->InternalTotalAlignment);
 	UE_NET_ASSERT_EQ(16u + PropertyComponentCount*PropertyComponentSize + IrisComponentCount*12U, Protocol->InternalTotalSize);
-	UE_NET_ASSERT_EQ(40u, Protocol->MaxExternalStateSize);
+	UE_NET_ASSERT_EQ(Align(40u + static_cast<uint32>(sizeof(FReplicationStateHeader)), 8U), Protocol->MaxExternalStateSize);
 	UE_NET_ASSERT_EQ(5U + PropertyComponentCount*RegularStatePropertyCountInPropertyComponent + IrisComponentCount*3U, Protocol->ChangeMaskBitCount);
 	UE_NET_ASSERT_EQ(EReplicationProtocolTraits::SupportsDeltaCompression, Protocol->ProtocolTraits);
 
diff --git a/Engine/Source/Runtime/Net/Core/Public/Net/Core/NetHandle/NetHandle.h b/Engine/Source/Runtime/Net/Core/Public/Net/Core/NetHandle/NetHandle.h
--- a/Engine/Source/Runtime/Net/Core/Public/Net/Core/NetHandle/NetHandle.h
+++ b/Engine/Source/Runtime/Net/Core/Public/Net/Core/NetHandle/NetHandle.h
@@ -31,6 +31,14 @@ namespace UE::Net
 class FNetHandle
 {
 public:
+#if UE_WITH_REMOTE_OBJECT_HANDLE
+	static constexpr uint64 IdBits = (REMOTE_OBJECT_SERVER_ID_BIT_SIZE + REMOTE_OBJECT_SERIAL_NUMBER_BIT_SIZE);
+#else
+	static constexpr uint64 IdBits = 30;
+#endif
+
+	static_assert(IdBits <= 64, "The number of bits representing the FNetHandle id cannot be greater than 64-bits");
+
 	FNetHandle();
 	FNetHandle(const FNetHandle& Other);
 	FNetHandle(FNetHandle&& Other);
@@ -38,7 +46,7 @@ public:
 	FNetHandle& operator=(const FNetHandle& Other);
 	FNetHandle& operator=(FNetHandle&& Other);
 
-	uint32 GetId() const;
+	uint64 GetId() const;
 	bool IsValid() const;
 
 	bool IsCompleteHandle() const;
@@ -137,9 +145,13 @@ inline FNetHandle& FNetHandle::operator=(FNetHandle&& Other)
 	return *this;
 }
 
-inline uint32 FNetHandle::GetId() const
+inline uint64 FNetHandle::GetId() const
 {
+#if UE_WITH_REMOTE_OBJECT_HANDLE
+	return (Value.GetIdNumber() >> 1);
+#else
 	return BitCast<FInternalValue>(Value).Id;
+#endif
 }
 
 inline bool FNetHandle::IsValid() const
diff --git a/Engine/Source/Runtime/Net/Iris/Public/Iris/ReplicationState/ReplicationStateFwd.h b/Engine/Source/Runtime/Net/Iris/Public/Iris/ReplicationState/ReplicationStateFwd.h
--- a/Engine/Source/Runtime/Net/Iris/Public/Iris/ReplicationState/ReplicationStateFwd.h
+++ b/Engine/Source/Runtime/Net/Iris/Public/Iris/ReplicationState/ReplicationStateFwd.h
@@ -4,6 +4,7 @@
 
 #include "HAL/Platform.h"
 #include "Net/Core/NetHandle/NetHandle.h"
+#include "Iris/Core/BitTwiddling.h"
 
 namespace UE::Net
 {
@@ -27,15 +28,13 @@ private:
 	friend Private::FReplicationStateHeaderAccessor;
 
 	// All replication states that are bound by an instance protocol is assigned a NetHandle for dirty state tracking
-	uint32 NetHandleId : 30;
+	uint64 NetHandleId : FNetHandle::IdBits;
 	// Init state doesn't use changemasks, instead we have a reserved bit here
-	uint32 bInitStateIsDirty : 1;
+	uint64 bInitStateIsDirty : 1;
 	// Track whether any state is dirty.
-	uint32 bStateIsDirty : 1;
+	uint64 bStateIsDirty : 1;
 };
 
-static_assert(sizeof(FReplicationStateHeader) == sizeof(uint32) && alignof(FReplicationStateHeader) == alignof(uint32), "FReplicationStateHeader must currently have the same size and alignment as uint32");
-
 namespace Private
 {
 
@@ -44,7 +43,7 @@ namespace Private
  */
 struct FReplicationStateHeaderAccessor
 {
-	static uint32 GetNetHandleId(const FReplicationStateHeader& Header) { return Header.NetHandleId; }
+	static uint64 GetNetHandleId(const FReplicationStateHeader& Header) { return Header.NetHandleId; }
 	static bool GetIsInitStateDirty(const FReplicationStateHeader& Header) { return Header.bInitStateIsDirty; }
 	static bool GetIsStateDirty(const FReplicationStateHeader& Header) { return Header.bStateIsDirty; }
 
@@ -58,9 +57,19 @@ struct FReplicationStateHeaderAccessor
 		Header.bStateIsDirty = false;
 	}
 
-	static void SetNetHandleId(FReplicationStateHeader& Header, FNetHandle NetHandle) { Header.NetHandleId = NetHandle.GetId(); }
+	static void SetNetHandleId(FReplicationStateHeader& Header, FNetHandle NetHandle);
 };
 
+inline void FReplicationStateHeaderAccessor::SetNetHandleId(FReplicationStateHeader& Header, FNetHandle NetHandle)
+{
+	const uint64 Id = NetHandle.GetId();
+
+	static constexpr uint64 LargestId = UE::Net::LargestBitValueForBits(FNetHandle::IdBits);
+	ensureMsgf(static_cast<uint64>(Id) < LargestId, TEXT("Setting a NetHandle identifier %llu that cannot be packed into %d bits (largest value of %llu)."), Id, FNetHandle::IdBits, LargestId);
+
+	Header.NetHandleId = Id;
+}
+
 }
 
 }
