Iris
* Deleted Bridge::IsAllowedToDestroyInstance since it's obsolete now that NetObjectFactories are in charge of object lifetimes
* NetActorFactory now tests NetDriver->ShouldClientDestroyActor before destroying a detached object

#rnx
#Synced-CL 48438982

diff --git a/Engine/Plugins/Runtime/ReplicationSystemTestPlugin/Source/Private/Tests/ReplicationSystem/ReplicatedTestObject.cpp b/Engine/Plugins/Runtime/ReplicationSystemTestPlugin/Source/Private/Tests/ReplicationSystem/ReplicatedTestObject.cpp
--- a/Engine/Plugins/Runtime/ReplicationSystemTestPlugin/Source/Private/Tests/ReplicationSystem/ReplicatedTestObject.cpp
+++ b/Engine/Plugins/Runtime/ReplicationSystemTestPlugin/Source/Private/Tests/ReplicationSystem/ReplicatedTestObject.cpp
@@ -107,12 +107,6 @@ void UReplicatedTestObjectBridge::EndReplication(UReplicatedTestObject* Instance
 	StopReplicatingNetObject(Instance, Flags);
 }
 
-
-bool UReplicatedTestObjectBridge::IsAllowedToDestroyInstance(const UObject* Instance) const
-{
-	return true;
-}
-
 void UReplicatedTestObjectBridge::SetExternalWorldLocationUpdateFunctor(TFunction<void(FNetRefHandle NetHandle, const UObject* ReplicatedObject, FVector& OutLocation, float& OutCullDistance)> LocUpdateFunctor)
 {
 	UReplicatedTestObjectFactory* TestFactory = CastChecked<UReplicatedTestObjectFactory>(GetNetFactory(ReplicatedObjectFactoryId));
diff --git a/Engine/Plugins/Runtime/ReplicationSystemTestPlugin/Source/Private/Tests/ReplicationSystem/ReplicatedTestObject.h b/Engine/Plugins/Runtime/ReplicationSystemTestPlugin/Source/Private/Tests/ReplicationSystem/ReplicatedTestObject.h
--- a/Engine/Plugins/Runtime/ReplicationSystemTestPlugin/Source/Private/Tests/ReplicationSystem/ReplicatedTestObject.h
+++ b/Engine/Plugins/Runtime/ReplicationSystemTestPlugin/Source/Private/Tests/ReplicationSystem/ReplicatedTestObject.h
@@ -637,14 +637,14 @@ public:
 	float GetMaxTickRate() const { return Super::GetMaxTickRate(); }
 
 	// Scope to suppress ensure when testing instantiation failure on clients
-	class FSupressCreateInstanceFailedEnsureScope
+	class FSuppressCreateInstanceFailedEnsureScope
 	{
 	public:
-		explicit FSupressCreateInstanceFailedEnsureScope(UReplicatedTestObjectBridge& BridgeIn) : Bridge(BridgeIn), bSuppressCreateInstanceFailedEnsure(Bridge.bSuppressCreateInstanceFailedEnsure)
+		explicit FSuppressCreateInstanceFailedEnsureScope(UReplicatedTestObjectBridge& BridgeIn) : Bridge(BridgeIn), bSuppressCreateInstanceFailedEnsure(Bridge.bSuppressCreateInstanceFailedEnsure)
 		{
 			Bridge.bSuppressCreateInstanceFailedEnsure = true;
 		}
-		~FSupressCreateInstanceFailedEnsureScope()
+		~FSuppressCreateInstanceFailedEnsureScope()
 		{
 			// Restore
 			Bridge.bSuppressCreateInstanceFailedEnsure = bSuppressCreateInstanceFailedEnsure;			
@@ -657,9 +657,7 @@ public:
 
 protected:
 
-	friend FSupressCreateInstanceFailedEnsureScope;
-
-	virtual bool IsAllowedToDestroyInstance(const UObject* Instance) const override;
+	friend FSuppressCreateInstanceFailedEnsureScope;
 
 	TFunction<void(FNetRefHandle NetHandle, const UObject* ReplicatedObject, FVector& OutLocation, float& OutCullDistance)> WorldLocationUpdateFunc;
 	bool bForceFailCreateRemoteInstance = false;
diff --git a/Engine/Source/Runtime/Engine/Private/Net/Iris/ReplicationSystem/EngineReplicationBridge.cpp b/Engine/Source/Runtime/Engine/Private/Net/Iris/ReplicationSystem/EngineReplicationBridge.cpp
--- a/Engine/Source/Runtime/Engine/Private/Net/Iris/ReplicationSystem/EngineReplicationBridge.cpp
+++ b/Engine/Source/Runtime/Engine/Private/Net/Iris/ReplicationSystem/EngineReplicationBridge.cpp
@@ -818,16 +818,6 @@ bool UEngineReplicationBridge::ObjectLevelHasFinishedLoading(UObject* Object) co
 	return true;
 }
 
-bool UEngineReplicationBridge::IsAllowedToDestroyInstance(const UObject* Instance) const
-{
-	if (AActor* Actor = const_cast<AActor*>(Cast<AActor>(Instance)))
-	{
-		return !NetDriver || NetDriver->ShouldClientDestroyActor(Actor);
-	}
-
-	return true;
-}
-
 UEngineReplicationBridge* UEngineReplicationBridge::Create(UNetDriver* NetDriver)
 {
 	UEngineReplicationBridge* Bridge = NewObject<UEngineReplicationBridge>(GetTransientPackage(), UEngineReplicationBridge::StaticClass());
diff --git a/Engine/Source/Runtime/Engine/Private/Net/Iris/ReplicationSystem/NetActorFactory.cpp b/Engine/Source/Runtime/Engine/Private/Net/Iris/ReplicationSystem/NetActorFactory.cpp
--- a/Engine/Source/Runtime/Engine/Private/Net/Iris/ReplicationSystem/NetActorFactory.cpp
+++ b/Engine/Source/Runtime/Engine/Private/Net/Iris/ReplicationSystem/NetActorFactory.cpp
@@ -505,7 +505,7 @@ void UNetActorFactory::DetachedFromReplication(const FDestroyedContext& Context)
 		{
 			NetDriver->ClientSetActorTornOff(Actor);
 		}
-		else if (EnumHasAnyFlags(Context.DestroyFlags, EReplicationBridgeDestroyInstanceFlags::AllowDestroyInstanceFromRemote))
+		else if (EnumHasAnyFlags(Context.DestroyFlags, EReplicationBridgeDestroyInstanceFlags::AllowDestroyInstanceFromRemote) && NetDriver->ShouldClientDestroyActor(Actor))
 		{
 			// Note that subobjects have already been detached by the ReplicationBridge
 			Actor->PreDestroyFromReplication();
diff --git a/Engine/Source/Runtime/Engine/Public/Net/Iris/ReplicationSystem/EngineReplicationBridge.h b/Engine/Source/Runtime/Engine/Public/Net/Iris/ReplicationSystem/EngineReplicationBridge.h
--- a/Engine/Source/Runtime/Engine/Public/Net/Iris/ReplicationSystem/EngineReplicationBridge.h
+++ b/Engine/Source/Runtime/Engine/Public/Net/Iris/ReplicationSystem/EngineReplicationBridge.h
@@ -165,7 +165,6 @@ protected:
 	virtual void GetInitialDependencies(FNetRefHandle Handle, FNetDependencyInfoArray& OutDependencies) const override;
 	virtual bool RemapPathForPIE(uint32 ConnectionId, FString& Path, bool bReading) const override;
 	virtual bool ObjectLevelHasFinishedLoading(UObject* Object) const override;
-	virtual bool IsAllowedToDestroyInstance(const UObject* Instance) const override;
 	virtual void OnProtocolMismatchDetected(FNetRefHandle ObjectHandle) override;
 	virtual void OnProtocolMismatchReported(FNetRefHandle RefHandle, uint32 ConnectionId) override;
 	virtual bool CanCreateDestructionInfo() const override;
diff --git a/Engine/Source/Runtime/Net/Iris/Private/Iris/ReplicationSystem/ObjectReplicationBridge.cpp b/Engine/Source/Runtime/Net/Iris/Private/Iris/ReplicationSystem/ObjectReplicationBridge.cpp
--- a/Engine/Source/Runtime/Net/Iris/Private/Iris/ReplicationSystem/ObjectReplicationBridge.cpp
+++ b/Engine/Source/Runtime/Net/Iris/Private/Iris/ReplicationSystem/ObjectReplicationBridge.cpp
@@ -506,7 +506,7 @@ void UObjectReplicationBridge::ReadAndExecuteDestructionInfoFromRemote(FReplicat
 		UE_LOG(LogIris, VeryVerbose, TEXT("ReadAndExecuteDestructionInfoFromRemote on %s | %s | NetFactoryId: %u"), *ReferenceToDestroy.ToString(), *ReferenceToDestroy.GetRefHandle().ToString(), NetFactoryId);
 
 		constexpr EReplicationBridgeDestroyInstanceReason DestroyReason = EReplicationBridgeDestroyInstanceReason::Destroy;
-		const EReplicationBridgeDestroyInstanceFlags DestroyFlags = IsAllowedToDestroyInstance(Instance) ? EReplicationBridgeDestroyInstanceFlags::AllowDestroyInstanceFromRemote : EReplicationBridgeDestroyInstanceFlags::None;
+		const EReplicationBridgeDestroyInstanceFlags DestroyFlags = EReplicationBridgeDestroyInstanceFlags::AllowDestroyInstanceFromRemote;
 
 		// if a subobject owner is to be destroyed we want to detach all subobjects before doing so to ensure we execute expected callbacks
 		// We keep tracking them internally
diff --git a/Engine/Source/Runtime/Net/Iris/Private/Iris/ReplicationSystem/ReplicationBridge.cpp b/Engine/Source/Runtime/Net/Iris/Private/Iris/ReplicationSystem/ReplicationBridge.cpp
--- a/Engine/Source/Runtime/Net/Iris/Private/Iris/ReplicationSystem/ReplicationBridge.cpp
+++ b/Engine/Source/Runtime/Net/Iris/Private/Iris/ReplicationSystem/ReplicationBridge.cpp
@@ -139,11 +139,6 @@ bool UReplicationBridge::CallCacheNetRefHandleCreationInfo(FNetRefHandle Handle)
 	return CacheNetRefHandleCreationInfo(Handle);
 }
 
-bool UReplicationBridge::IsAllowedToDestroyInstance(const UObject* Instance) const
-{
-	return true;
-}
-
 UReplicationBridge::~UReplicationBridge()
 {
 }
diff --git a/Engine/Source/Runtime/Net/Iris/Public/Iris/ReplicationSystem/ReplicationBridge.h b/Engine/Source/Runtime/Net/Iris/Public/Iris/ReplicationSystem/ReplicationBridge.h
--- a/Engine/Source/Runtime/Net/Iris/Public/Iris/ReplicationSystem/ReplicationBridge.h
+++ b/Engine/Source/Runtime/Net/Iris/Public/Iris/ReplicationSystem/ReplicationBridge.h
@@ -238,9 +238,6 @@ protected:
 	/** Destroys the group associated with the level. */
 	IRISCORE_API void DestroyLevelGroup(const UObject* Level);
 
-	/** Called when destruction info is received to determine whether the instance may be destroyed. */
-	IRISCORE_API virtual bool IsAllowedToDestroyInstance(const UObject* Instance) const;
-
 	/** Called when a remote connection detected a protocol mismatch when trying to instantiate the NetRefHandle replicated object. */
 	virtual void OnProtocolMismatchReported(FNetRefHandle RefHandle, uint32 ConnectionId) {}
 
