Unshelved from pending changelist '48465432':

Missing file from cl

#Synced-CL 48438982

diff --git a/Engine/Plugins/Runtime/ReplicationSystemTestPlugin/Source/Private/Tests/ReplicationSystem/TestReplicationSystem.cpp b/Engine/Plugins/Runtime/ReplicationSystemTestPlugin/Source/Private/Tests/ReplicationSystem/TestReplicationSystem.cpp
--- a/Engine/Plugins/Runtime/ReplicationSystemTestPlugin/Source/Private/Tests/ReplicationSystem/TestReplicationSystem.cpp
+++ b/Engine/Plugins/Runtime/ReplicationSystemTestPlugin/Source/Private/Tests/ReplicationSystem/TestReplicationSystem.cpp
@@ -2124,7 +2124,7 @@ UE_NET_TEST_FIXTURE(FReplicationSystemServerClientTestFixture, TestClientCanSkip
 		ServerObjectA->bForceFailToInstantiateOnRemote = true;
 
 		// Suppress ensure that will occur due to failing to instantiate the object
-		UReplicatedTestObjectBridge::FSupressCreateInstanceFailedEnsureScope SuppressEnsureScope(*Client->GetReplicationBridge());
+		UReplicatedTestObjectBridge::FSuppressCreateInstanceFailedEnsureScope SuppressEnsureScope(*Client->GetReplicationBridge());
 
 		// Disable error logging as we know we will fail.
 		auto IrisLogVerbosity = UE_GET_LOG_VERBOSITY(LogIris);
