Fix  issue with dependent object remaining in scope even if filtered out by an exclusion group.

#Synced-CL 48600773

diff --git a/Engine/Plugins/Runtime/ReplicationSystemTestPlugin/Source/Private/Tests/ReplicationSystem/Filtering/TestFiltering.cpp b/Engine/Plugins/Runtime/ReplicationSystemTestPlugin/Source/Private/Tests/ReplicationSystem/Filtering/TestFiltering.cpp
--- a/Engine/Plugins/Runtime/ReplicationSystemTestPlugin/Source/Private/Tests/ReplicationSystem/Filtering/TestFiltering.cpp
+++ b/Engine/Plugins/Runtime/ReplicationSystemTestPlugin/Source/Private/Tests/ReplicationSystem/Filtering/TestFiltering.cpp
@@ -2898,6 +2898,83 @@ UE_NET_TEST_FIXTURE(FTestFilteringFixture, AlreadyCreatedObjectAddedAsDependency
 	UE_NET_ASSERT_FALSE(ClientArray[1]->IsValidNetRefHandle(ServerDependentObject2->NetRefHandle));
 }
 
+UE_NET_TEST_FIXTURE(FTestFilteringFixture, DependentObjectWithHysteresisAddedToExclusionGroupObjectIsNotInScope)
+{
+	FReplicationSystemTestClient* ClientArray[] = { CreateClient(), CreateClient() };
+
+	UReplicatedTestObject* ServerObject = Server->CreateObject();
+
+	Server->UpdateAndSend(ClientArray);
+
+	UE_NET_ASSERT_TRUE(ClientArray[0]->IsResolvableNetRefHandle(ServerObject->NetRefHandle));
+	UE_NET_ASSERT_TRUE(ClientArray[1]->IsResolvableNetRefHandle(ServerObject->NetRefHandle));
+
+	UReplicatedTestObject* ServerDependentObject = Server->CreateObject();
+	Server->ReplicationSystem->SetFilter(ServerDependentObject->NetRefHandle, NotRoutedFilterHandle, /* .FilterConfigProfile */ "VeryManyFrames");
+
+	Server->ReplicationBridge->AddDependentObject(ServerObject->NetRefHandle, ServerDependentObject->NetRefHandle);
+
+	Server->UpdateAndSend(ClientArray);
+
+	FNetObjectGroupHandle GroupHandle = Server->ReplicationSystem->CreateGroup(NAME_None);
+	Server->ReplicationSystem->AddToGroup(GroupHandle, ServerDependentObject->NetRefHandle);
+	Server->ReplicationSystem->AddExclusionFilterGroup(GroupHandle);
+
+	Server->ReplicationSystem->SetGroupFilterStatus(GroupHandle, ENetFilterStatus::Allow);
+	Server->ReplicationSystem->SetGroupFilterStatus(GroupHandle, ClientArray[1]->ConnectionIdOnServer, ENetFilterStatus::Disallow);
+
+	Server->UpdateAndSend(ClientArray);
+
+	UE_NET_ASSERT_TRUE(ClientArray[0]->IsResolvableNetRefHandle(ServerDependentObject->NetRefHandle));
+	UE_NET_ASSERT_FALSE(ClientArray[1]->IsValidNetRefHandle(ServerDependentObject->NetRefHandle));
+}
+
+UE_NET_TEST_FIXTURE(FTestFilteringFixture, DependentObjectWithHysteresisAddedToExclusionGroupThenAllowedIsNotFilteredOutByHysteresis)
+{
+	FReplicationSystemTestClient* ClientArray[] = { CreateClient(), CreateClient() };
+
+	UReplicatedTestObject* ServerObject = Server->CreateObject();
+
+	Server->UpdateAndSend(ClientArray);
+
+	UE_NET_ASSERT_TRUE(ClientArray[0]->IsResolvableNetRefHandle(ServerObject->NetRefHandle));
+	UE_NET_ASSERT_TRUE(ClientArray[1]->IsResolvableNetRefHandle(ServerObject->NetRefHandle));
+
+	UReplicatedTestObject* ServerDependentObject = Server->CreateObject();
+	Server->ReplicationSystem->SetFilter(ServerDependentObject->NetRefHandle, NotRoutedFilterHandle, /* .FilterConfigProfile */ "FiveFrames");
+
+	Server->ReplicationBridge->AddDependentObject(ServerObject->NetRefHandle, ServerDependentObject->NetRefHandle);
+
+	Server->UpdateAndSend(ClientArray);
+
+	FNetObjectGroupHandle GroupHandle = Server->ReplicationSystem->CreateGroup(NAME_None);
+	Server->ReplicationSystem->AddToGroup(GroupHandle, ServerDependentObject->NetRefHandle);
+	Server->ReplicationSystem->AddExclusionFilterGroup(GroupHandle);
+
+	Server->ReplicationSystem->SetGroupFilterStatus(GroupHandle, ENetFilterStatus::Allow);
+	Server->ReplicationSystem->SetGroupFilterStatus(GroupHandle, ClientArray[1]->ConnectionIdOnServer, ENetFilterStatus::Disallow);
+
+	Server->UpdateAndSend(ClientArray);
+
+	UE_NET_ASSERT_TRUE(ClientArray[0]->IsResolvableNetRefHandle(ServerDependentObject->NetRefHandle));
+	UE_NET_ASSERT_FALSE(ClientArray[1]->IsValidNetRefHandle(ServerDependentObject->NetRefHandle));
+
+	// Allow dependent object again
+	Server->ReplicationSystem->SetGroupFilterStatus(GroupHandle, ClientArray[1]->ConnectionIdOnServer, ENetFilterStatus::Allow);
+
+	Server->UpdateAndSend(ClientArray);
+
+	// Dependent object should now have been replicated again
+	UE_NET_ASSERT_TRUE(ClientArray[1]->IsResolvableNetRefHandle(ServerDependentObject->NetRefHandle));
+
+	// Update for the expected hysteresis duration and make sure the dependent object never goes out of scope.
+	for (unsigned FrameIt = 0; FrameIt < 6U; ++FrameIt)
+	{
+		Server->UpdateAndSend(ClientArray);
+		UE_NET_ASSERT_TRUE(ClientArray[1]->IsResolvableNetRefHandle(ServerDependentObject->NetRefHandle));
+	}
+}
+
 UE_NET_TEST_FIXTURE(FTestFilteringFixture, OwningConnectionIsSetProperly)
 {
 	FReplicationSystemTestClient* Client = CreateClient();
diff --git a/Engine/Source/Runtime/Net/Iris/Private/Iris/ReplicationSystem/Filtering/ReplicationFiltering.cpp b/Engine/Source/Runtime/Net/Iris/Private/Iris/ReplicationSystem/Filtering/ReplicationFiltering.cpp
--- a/Engine/Source/Runtime/Net/Iris/Private/Iris/ReplicationSystem/Filtering/ReplicationFiltering.cpp
+++ b/Engine/Source/Runtime/Net/Iris/Private/Iris/ReplicationSystem/Filtering/ReplicationFiltering.cpp
@@ -1636,7 +1636,7 @@ void FReplicationFiltering::UpdateDynamicFiltering()
 				else
 				{
 					// Start/continue hysteresis or time to prevent replication?
-					bool bIsFilteredOut = DynamicFilteredOutObjectsHysteresisAdjusted.GetBit(DependentObjectIndex);
+					bool bIsFilteredOut = DynamicFilteredOutObjectsHysteresisAdjusted.GetBit(DependentObjectIndex) || !ConnectionInfo.ObjectsInScopeBeforeDynamicFiltering.GetBit(DependentObjectIndex);
 					if (!bIsFilteredOut)
 					{
 						// If we're updated by hysteresis at this point we should not start again, but if other conditions allow it we should start!
@@ -1655,10 +1655,11 @@ void FReplicationFiltering::UpdateDynamicFiltering()
 								bIsFilteredOut = true;
 							}
 						}
-						// Else let hysteresis updating continue. Once fully processed the DynamicFilteredOutObjectsHysteresisAdjusted bitarray will be adjsuted accordingly and filter out the object.
+						// Else let hysteresis updating continue. Once fully processed the DynamicFilteredOutObjectsHysteresisAdjusted bitarray will be adjusted accordingly and filter out the object.
 
 						// Explicitly enable dependent object and its subobjects.
 						ConnectionInfo.ObjectsInScope.SetBitValue(DependentObjectIndex, !bIsFilteredOut);
+						DynamicFilteredOutObjectsHysteresisAdjusted.SetBitValue(DependentObjectIndex, bIsFilteredOut);
 						for (const FInternalNetRefIndex SubObjectIndex : NetRefHandleManager->GetSubObjects(DependentObjectIndex))
 						{
 							DynamicFilteredOutObjectsHysteresisAdjusted.SetBitValue(SubObjectIndex, bIsFilteredOut);
