IrisTests
- Added support to fake stablynamed objects from ReplicationSystemTest
- Now ticks test client from DeliverTo to drive required logic

#rnx
#Synced-CL 48600773

diff --git a/Engine/Plugins/Runtime/ReplicationSystemTestPlugin/Source/Private/Tests/ReplicationSystem/ReplicatedTestObject.h b/Engine/Plugins/Runtime/ReplicationSystemTestPlugin/Source/Private/Tests/ReplicationSystem/ReplicatedTestObject.h
--- a/Engine/Plugins/Runtime/ReplicationSystemTestPlugin/Source/Private/Tests/ReplicationSystem/ReplicatedTestObject.h
+++ b/Engine/Plugins/Runtime/ReplicationSystemTestPlugin/Source/Private/Tests/ReplicationSystem/ReplicatedTestObject.h
@@ -35,6 +35,8 @@ public:
 	// Register the fragments for this object
 	virtual void RegisterReplicationFragments(UE::Net::FFragmentRegistrationContext& Fragments, UE::Net::EFragmentRegistrationFlags RegistrationFlags) override {}
 	virtual bool IsSupportedForNetworking() const override { return true; }
+	virtual bool IsNameStableForNetworking() const override { return Super::IsNameStableForNetworking() || bForceStablyNamed; }
+	virtual bool IsFullNameStableForNetworking() const override { return Super::IsFullNameStableForNetworking() || bForceStablyNamed; }
 
 	virtual void OnSubObjectCreated(UObject* SubObject) {}
 	virtual void OnSubObjectDestroyed(UObject* SubObject) {}
@@ -47,6 +49,9 @@ public:
 
 	// To determine if this object is a root object or a subobject
 	bool bIsSubObject = false;
+
+	// Force the object to be treated as stably named
+	bool bForceStablyNamed = false;
 };
 
 USTRUCT()
diff --git a/Engine/Plugins/Runtime/ReplicationSystemTestPlugin/Source/Private/Tests/ReplicationSystem/ReplicatedTestObjectFactory.cpp b/Engine/Plugins/Runtime/ReplicationSystemTestPlugin/Source/Private/Tests/ReplicationSystem/ReplicatedTestObjectFactory.cpp
--- a/Engine/Plugins/Runtime/ReplicationSystemTestPlugin/Source/Private/Tests/ReplicationSystem/ReplicatedTestObjectFactory.cpp
+++ b/Engine/Plugins/Runtime/ReplicationSystemTestPlugin/Source/Private/Tests/ReplicationSystem/ReplicatedTestObjectFactory.cpp
@@ -6,7 +6,7 @@
 #include "Iris/Serialization/NetBitStreamWriter.h"
 #include "Iris/Serialization/NetSerializationContext.h"
 #include "Iris/Serialization/NetBitStreamUtil.h"
-
+#include "Iris/ReplicationSystem/NetObjectFactory.h"
 
 namespace UE::Net
 {
@@ -17,36 +17,51 @@ namespace UE::Net
 
 FString FReplicationTestCreationHeader::ToString() const
 {
-	return FString::Printf(TEXT("FReplicationTestCreationHeader (ProtocolId:0x%x)\n\t"
-								"ArchetypeName=%s"
-								"NumComponentsToSpawn=%u"
-								"NumIrisComponentsToSpawn=%u"
-								"NumDynamicComponentsToSpawn=%u"
-								"NumConnectionFilteredComponentsToSpawn=%u"
-								"NumObjectReferenceComponentsToSpawn=%u"
-								"bForceFailCreationRemoteInstance=%u"),
-						   		GetProtocolId(),
-						   		*ArchetypeName,
-						   		NumComponentsToSpawn,
-						   		NumIrisComponentsToSpawn,
-						   		NumDynamicComponentsToSpawn,
-						   		NumConnectionFilteredComponentsToSpawn,
-						   		NumObjectReferenceComponentsToSpawn,
-						   		bForceFailCreateRemoteInstance);
+	if (bIsStable)
+	{
+		return FString::Printf(TEXT("FReplicationTestCreationHeader (ProtocolId:0x%x)\n\t"
+									"StableName=%s"
+									"bForceFailCreationRemoteInstance=%u"),
+						   			GetProtocolId(),
+						   			*ArchetypeOrStableName,
+						   			bForceFailCreateRemoteInstance);	
+	}
+	else
+	{
+		return FString::Printf(TEXT("FReplicationTestCreationHeader (ProtocolId:0x%x)\n\t"
+									"ArchetypeName=%s"
+									"NumComponentsToSpawn=%u"
+									"NumIrisComponentsToSpawn=%u"
+									"NumDynamicComponentsToSpawn=%u"
+									"NumConnectionFilteredComponentsToSpawn=%u"
+									"NumObjectReferenceComponentsToSpawn=%u"
+									"bForceFailCreationRemoteInstance=%u"),
+						   			GetProtocolId(),
+						   			*ArchetypeOrStableName,
+						   			NumComponentsToSpawn,
+						   			NumIrisComponentsToSpawn,
+						   			NumDynamicComponentsToSpawn,
+						   			NumConnectionFilteredComponentsToSpawn,
+						   			NumObjectReferenceComponentsToSpawn,
+						   			bForceFailCreateRemoteInstance);
+	}
 }
 
 bool FReplicationTestCreationHeader::Serialize(const FCreationHeaderContext& Context) const
 {
 	FNetBitStreamWriter& Writer = *Context.Serialization.GetBitStreamWriter();
 
-	WriteString(&Writer, ArchetypeName);
-	Writer.WriteBits(NumComponentsToSpawn, 16);
-	Writer.WriteBits(NumIrisComponentsToSpawn, 16);
-	Writer.WriteBits(NumDynamicComponentsToSpawn, 16);
-	Writer.WriteBits(NumConnectionFilteredComponentsToSpawn, 16);
-	Writer.WriteBits(NumObjectReferenceComponentsToSpawn, 16);
+	WriteString(&Writer, ArchetypeOrStableName);
 	Writer.WriteBool(bForceFailCreateRemoteInstance);
-
+	if (!Writer.WriteBool(bIsStable))
+	{
+		Writer.WriteBits(NumComponentsToSpawn, 16);
+		Writer.WriteBits(NumIrisComponentsToSpawn, 16);
+		Writer.WriteBits(NumDynamicComponentsToSpawn, 16);
+		Writer.WriteBits(NumConnectionFilteredComponentsToSpawn, 16);
+		Writer.WriteBits(NumObjectReferenceComponentsToSpawn, 16);
+	}
+	
 	return !Writer.IsOverflown();
 }
 
@@ -54,14 +69,17 @@ bool FReplicationTestCreationHeader::Deserialize(const FCreationHeaderContext& C
 {
 	FNetBitStreamReader& Reader = *Context.Serialization.GetBitStreamReader();
 
-	ReadString(&Reader, this->ArchetypeName);
-
-	this->NumComponentsToSpawn = Reader.ReadBits(16);
-	this->NumIrisComponentsToSpawn = Reader.ReadBits(16);
-	this->NumDynamicComponentsToSpawn = Reader.ReadBits(16);
-	this->NumConnectionFilteredComponentsToSpawn = Reader.ReadBits(16);
-	this->NumObjectReferenceComponentsToSpawn = Reader.ReadBits(16);
+	ReadString(&Reader, this->ArchetypeOrStableName);
 	this->bForceFailCreateRemoteInstance = Reader.ReadBool();
+	this->bIsStable = Reader.ReadBool();
+	if (!this->bIsStable)
+	{
+		this->NumComponentsToSpawn = Reader.ReadBits(16);
+		this->NumIrisComponentsToSpawn = Reader.ReadBits(16);
+		this->NumDynamicComponentsToSpawn = Reader.ReadBits(16);
+		this->NumConnectionFilteredComponentsToSpawn = Reader.ReadBits(16);
+		this->NumObjectReferenceComponentsToSpawn = Reader.ReadBits(16);
+	}
 
 	return !Reader.IsOverflown();
 }
@@ -88,6 +106,23 @@ TUniquePtr<UE::Net::FNetObjectCreationHeader> UReplicatedTestObjectFactory::Crea
 		return nullptr;
 	}
 
+	if (const bool bIsStable = Handle.IsStatic())
+	{
+		TUniquePtr<FReplicationTestCreationHeader> HeaderRef(new FReplicationTestCreationHeader);
+		FReplicationTestCreationHeader* Header = HeaderRef.Get();
+
+		Header->bIsStable = true;
+
+		Header->ArchetypeOrStableName = Object->GetPathName(Bridge);
+
+		if (const UReplicatedTestObject* ReplicatedTestObject = Cast<UReplicatedTestObject>(Object))
+		{
+			Header->bForceFailCreateRemoteInstance = ReplicatedTestObject->bForceFailToInstantiateOnRemote;
+		}
+
+		return HeaderRef;
+	}
+
 	UObject* Archetype = Object->GetArchetype();
 	if (!Archetype)
 	{
@@ -112,7 +147,7 @@ TUniquePtr<UE::Net::FNetObjectCreationHeader> UReplicatedTestObjectFactory::Crea
 		Header->NumObjectReferenceComponentsToSpawn = IntCastChecked<uint16>(TestReplicatedIrisObject->ObjectReferenceComponents.Num());
 	}
 
-	Header->ArchetypeName = Archetype->GetPathName();
+	Header->ArchetypeOrStableName = Archetype->GetPathName();
 
 	return HeaderRef;
 }
@@ -142,49 +177,73 @@ UNetObjectFactory::FInstantiateResult UReplicatedTestObjectFactory::InstantiateR
 		return FInstantiateResult();
 	}
 
-	UObject* ArcheType = StaticFindObject(UObject::StaticClass(), nullptr, *Header->ArchetypeName, EFindObjectFlags::None);
-	check(ArcheType);
+	if (Header->bIsStable)
+	{
+		UObject* Object = FindObjectFast<UReplicatedTestObject>(Bridge, FName(Header->ArchetypeOrStableName));
 
-	FStaticConstructObjectParameters ConstructObjectParameters(ArcheType->GetClass());
-	UObject* CreatedObject = StaticConstructObject_Internal(ConstructObjectParameters);
+		const bool bIsSubObject = Context.RootObjectOfSubObject.IsValid();
+		if (UReplicatedTestObject* BaseTestObject = Cast<UReplicatedTestObject>(Object))
+		{
+			BaseTestObject->bIsSubObject = bIsSubObject;
+		}
 
-	const bool bIsSubObject = Context.RootObjectOfSubObject.IsValid();
+		FInstantiateResult InstantiateResult;
+		InstantiateResult.Instance = Object;
 
-	if (UReplicatedTestObject* BaseTestObject = Cast<UReplicatedTestObject>(CreatedObject))
-	{
-		BaseTestObject->bIsSubObject = bIsSubObject;
+		//TBD: allow for now
+		InstantiateResult.Flags |= EReplicationBridgeCreateNetRefHandleResultFlags::AllowDestroyInstanceFromRemote;
+		InstantiateResult.Flags |= EReplicationBridgeCreateNetRefHandleResultFlags::BindStaticObjectInReferenceCache;
+		if (bIsSubObject)
+		{
+			InstantiateResult.Flags |= EReplicationBridgeCreateNetRefHandleResultFlags::ShouldCallSubObjectCreatedFromReplication;
+		}
+		return InstantiateResult;
 	}
-	
-	if (UTestReplicatedIrisObject* CreatedTestObject = Cast<UTestReplicatedIrisObject>(CreatedObject))
+	else
 	{
-		UTestReplicatedIrisObject::FComponents Components;
-		Components.PropertyComponentCount = Header->NumComponentsToSpawn;
-		Components.IrisComponentCount = Header->NumIrisComponentsToSpawn;
-		Components.DynamicStateComponentCount = Header->NumDynamicComponentsToSpawn;
-		Components.ConnectionFilteredComponentCount = Header->NumConnectionFilteredComponentsToSpawn;
-		Components.ObjectReferenceComponentCount = Header->NumObjectReferenceComponentsToSpawn;
+		UObject* ArcheType = StaticFindObject(UObject::StaticClass(), nullptr, *Header->ArchetypeOrStableName, EFindObjectFlags::None);
+		check(ArcheType);
 
-		CreatedTestObject->AddComponents(Components);
-	}
+		FStaticConstructObjectParameters ConstructObjectParameters(ArcheType->GetClass());
+		UObject* CreatedObject = StaticConstructObject_Internal(ConstructObjectParameters);
 
-	UReplicatedTestObjectBridge* TestBridge = Cast<UReplicatedTestObjectBridge>(Bridge);
+		const bool bIsSubObject = Context.RootObjectOfSubObject.IsValid();
 
-	// Store the object so that we can find detached/torn off instances from tests
-	if (CreatedObjectsOnNode)
-	{
-		CreatedObjectsOnNode->Add(TStrongObjectPtr<UObject>(CreatedObject));
-	}
+		if (UReplicatedTestObject* BaseTestObject = Cast<UReplicatedTestObject>(CreatedObject))
+		{
+			BaseTestObject->bIsSubObject = bIsSubObject;
+		}
 	
-	FInstantiateResult InstantiateResult;
-	InstantiateResult.Instance = CreatedObject;
-	InstantiateResult.Flags |= EReplicationBridgeCreateNetRefHandleResultFlags::AllowDestroyInstanceFromRemote;
+		if (UTestReplicatedIrisObject* CreatedTestObject = Cast<UTestReplicatedIrisObject>(CreatedObject))
+		{
+			UTestReplicatedIrisObject::FComponents Components;
+			Components.PropertyComponentCount = Header->NumComponentsToSpawn;
+			Components.IrisComponentCount = Header->NumIrisComponentsToSpawn;
+			Components.DynamicStateComponentCount = Header->NumDynamicComponentsToSpawn;
+			Components.ConnectionFilteredComponentCount = Header->NumConnectionFilteredComponentsToSpawn;
+			Components.ObjectReferenceComponentCount = Header->NumObjectReferenceComponentsToSpawn;
+
+			CreatedTestObject->AddComponents(Components);
+		}
 
-	if (bIsSubObject)
-	{
-		InstantiateResult.Flags |= EReplicationBridgeCreateNetRefHandleResultFlags::ShouldCallSubObjectCreatedFromReplication;
-	}
+		UReplicatedTestObjectBridge* TestBridge = Cast<UReplicatedTestObjectBridge>(Bridge);
 
-	return InstantiateResult;
+		// Store the object so that we can find detached/torn off instances from tests
+		if (CreatedObjectsOnNode)
+		{
+			CreatedObjectsOnNode->Add(TStrongObjectPtr<UObject>(CreatedObject));
+		}
+	
+		FInstantiateResult InstantiateResult;
+		InstantiateResult.Instance = CreatedObject;
+		InstantiateResult.Flags |= EReplicationBridgeCreateNetRefHandleResultFlags::AllowDestroyInstanceFromRemote;
+
+		if (bIsSubObject)
+		{
+			InstantiateResult.Flags |= EReplicationBridgeCreateNetRefHandleResultFlags::ShouldCallSubObjectCreatedFromReplication;
+		}
+		return InstantiateResult;
+	}
 }
 
 void UReplicatedTestObjectFactory::PostInit(const FPostInitContext& Context)
diff --git a/Engine/Plugins/Runtime/ReplicationSystemTestPlugin/Source/Private/Tests/ReplicationSystem/ReplicatedTestObjectFactory.h b/Engine/Plugins/Runtime/ReplicationSystemTestPlugin/Source/Private/Tests/ReplicationSystem/ReplicatedTestObjectFactory.h
--- a/Engine/Plugins/Runtime/ReplicationSystemTestPlugin/Source/Private/Tests/ReplicationSystem/ReplicatedTestObjectFactory.h
+++ b/Engine/Plugins/Runtime/ReplicationSystemTestPlugin/Source/Private/Tests/ReplicationSystem/ReplicatedTestObjectFactory.h
@@ -12,13 +12,15 @@ namespace UE::Net
 class FReplicationTestCreationHeader : public FNetObjectCreationHeader
 {
 public:
-	FString ArchetypeName;
+	FString ArchetypeOrStableName;
 	uint32 NumComponentsToSpawn = 0;
 	uint32 NumIrisComponentsToSpawn = 0;
 	uint32 NumDynamicComponentsToSpawn = 0;
 	uint32 NumConnectionFilteredComponentsToSpawn = 0;
 	uint32 NumObjectReferenceComponentsToSpawn = 0;
+
 	bool bForceFailCreateRemoteInstance = false;
+	bool bIsStable = false;
 
 	virtual FString ToString() const override;
 
diff --git a/Engine/Plugins/Runtime/ReplicationSystemTestPlugin/Source/Private/Tests/ReplicationSystem/ReplicationSystemServerClientTestFixture.cpp b/Engine/Plugins/Runtime/ReplicationSystemTestPlugin/Source/Private/Tests/ReplicationSystem/ReplicationSystemServerClientTestFixture.cpp
--- a/Engine/Plugins/Runtime/ReplicationSystemTestPlugin/Source/Private/Tests/ReplicationSystem/ReplicationSystemServerClientTestFixture.cpp
+++ b/Engine/Plugins/Runtime/ReplicationSystemTestPlugin/Source/Private/Tests/ReplicationSystem/ReplicationSystemServerClientTestFixture.cpp
@@ -554,6 +554,11 @@ bool FReplicationSystemTestServer::SendTo(FReplicationSystemTestClient* Client,
 void FReplicationSystemTestServer::DeliverTo(FReplicationSystemTestClient* Client, bool bDeliver)
 {
 	FReplicationSystemTestNode::DeliverTo(*Client, Client->ConnectionIdOnServer, Client->LocalConnectionId, bDeliver);
+	if (bUpdateClientOnDelivery)
+	{
+		Client->NetUpdate();
+		Client->PostSendUpdate();
+	}
 }
 
 bool FReplicationSystemTestServer::UpdateAndSend(const TArrayView<FReplicationSystemTestClient*const>& Clients, bool bDeliver /*= true*/, const TCHAR* Desc)
diff --git a/Engine/Plugins/Runtime/ReplicationSystemTestPlugin/Source/Private/Tests/ReplicationSystem/ReplicationSystemServerClientTestFixture.h b/Engine/Plugins/Runtime/ReplicationSystemTestPlugin/Source/Private/Tests/ReplicationSystem/ReplicationSystemServerClientTestFixture.h
--- a/Engine/Plugins/Runtime/ReplicationSystemTestPlugin/Source/Private/Tests/ReplicationSystem/ReplicationSystemServerClientTestFixture.h
+++ b/Engine/Plugins/Runtime/ReplicationSystemTestPlugin/Source/Private/Tests/ReplicationSystem/ReplicationSystemServerClientTestFixture.h
@@ -132,6 +132,33 @@ public:
 		return nullptr;
 	}
 
+	template<typename T>
+	T* CreateStablyNamedObject(const TCHAR* ObjectName)
+	{
+		check(ReplicationSystem);
+		if (!ReplicationSystem)
+		{
+			return nullptr;
+		}
+
+		T* CreatedObject = NewObject<T>(ReplicationSystem->GetReplicationBridge(), *FString::Printf(TEXT("%s"), ObjectName));
+
+		if (UReplicatedTestObject* ReplicatedTestObject = Cast<UReplicatedTestObject>(CreatedObject))
+		{
+			ReplicatedTestObject->bForceStablyNamed = true;
+			CreatedObjects.Add(TStrongObjectPtr<UObject>(CreatedObject));
+
+			if (ReplicationSystem && ReplicationSystem->IsServer())
+			{
+				// Add it to the bridge for replication
+				ReplicationBridge->BeginReplication(CreatedObject);
+			}
+	
+			return CreatedObject;
+		}
+		return nullptr;
+	}
+
 	template<typename T>
 	T* CreateObjectWithParams(const UObjectReplicationBridge::FRootObjectReplicationParams& Params)
 	{
@@ -258,6 +285,9 @@ public:
 
 	// Tick and send packets to one or all clients
 	bool UpdateAndSend(const TArrayView<FReplicationSystemTestClient*const>& Clients, bool bDeliver=true, const TCHAR* Desc = nullptr);
+
+	// If set, automatically call NetUpdate and PostSendUpdate after each deliverd packet to simulate a client tick.
+	bool bUpdateClientOnDelivery = true;
 };
 
 class FReplicationSystemServerClientTestFixture : public FNetworkAutomationTestSuiteFixture
