Iris - Don't store instanced struct type in default state due to its NetObjectReference is going to be invalid. This causes issues when the default state contains a valid instance due to the client not being able to resolve the type from its default state.

#Synced-CL 48600773

diff --git a/Engine/Plugins/Runtime/ReplicationSystemTestPlugin/Source/Private/Tests/Serialization/TestInstancedStructNetSerializer.cpp b/Engine/Plugins/Runtime/ReplicationSystemTestPlugin/Source/Private/Tests/Serialization/TestInstancedStructNetSerializer.cpp
--- a/Engine/Plugins/Runtime/ReplicationSystemTestPlugin/Source/Private/Tests/Serialization/TestInstancedStructNetSerializer.cpp
+++ b/Engine/Plugins/Runtime/ReplicationSystemTestPlugin/Source/Private/Tests/Serialization/TestInstancedStructNetSerializer.cpp
@@ -479,6 +479,42 @@ UE_NET_TEST_FIXTURE(FTestInstancedStructNetSerializerFixture, TestDequantizeSeri
 }
 
 // End-to-end tests
+UE_NET_TEST_FIXTURE(FTestInstancedStructNetSerializerFixture, CanInitializeInstancedStructInConstructor)
+{
+	// Add a client
+	FReplicationSystemTestClient* Client = CreateClient();
+
+	// Spawn object on server
+	UInstancedStructNetSerializerTestObject* ServerObject = Server->CreateObject<UInstancedStructNetSerializerTestObject>();
+	
+	// Replicate
+	Server->UpdateAndSend({ Client });
+
+	UInstancedStructNetSerializerTestObject* ClientObject = Client->GetObjectAs<UInstancedStructNetSerializerTestObject>(ServerObject->NetRefHandle);
+	UE_NET_ASSERT_NE(ClientObject, nullptr);
+	CA_ASSUME(ClientObject != nullptr);
+	UE_NET_ASSERT_EQ(ClientObject->ConstructorInitializedInstancedStruct.GetScriptStruct(), static_cast<const UScriptStruct*>(FStructForInstancedStructTestA::StaticStruct()));
+}
+
+UE_NET_TEST_FIXTURE(FTestInstancedStructNetSerializerFixture, CanModifyInstancedStructInitializedInConstructor)
+{
+	// Add a client
+	FReplicationSystemTestClient* Client = CreateClient();
+
+	// Spawn object on server
+	UInstancedStructNetSerializerTestObject* ServerObject = Server->CreateObject<UInstancedStructNetSerializerTestObject>();
+	ServerObject->ConstructorInitializedInstancedStruct.GetMutable<FStructForInstancedStructTestA>().SomeUint16 += 1;
+	
+	// Replicate
+	Server->UpdateAndSend({ Client });
+
+	UInstancedStructNetSerializerTestObject* ClientObject = Client->GetObjectAs<UInstancedStructNetSerializerTestObject>(ServerObject->NetRefHandle);
+	UE_NET_ASSERT_NE(ClientObject, nullptr);
+	CA_ASSUME(ClientObject != nullptr);
+	UE_NET_ASSERT_EQ(ClientObject->ConstructorInitializedInstancedStruct.GetScriptStruct(), static_cast<const UScriptStruct*>(FStructForInstancedStructTestA::StaticStruct()));
+	UE_NET_ASSERT_EQ(ClientObject->ConstructorInitializedInstancedStruct.Get<FStructForInstancedStructTestA>().SomeUint16, ServerObject->ConstructorInitializedInstancedStruct.Get<FStructForInstancedStructTestA>().SomeUint16);
+}
+
 UE_NET_TEST_FIXTURE(FTestInstancedStructNetSerializerFixture, ModifyInstance)
 {
 	// Add a client
@@ -1231,11 +1267,17 @@ void FStructForInstancedStructTestWithCustomSerializerNetSerializer::FNetSeriali
 }
 
 //
+UInstancedStructNetSerializerTestObject::UInstancedStructNetSerializerTestObject()
+{
+	ConstructorInitializedInstancedStruct.InitializeAs<FStructForInstancedStructTestA>();
+}
+
 void UInstancedStructNetSerializerTestObject::GetLifetimeReplicatedProps(TArray<class FLifetimeProperty>& OutLifetimeProps) const
 {
 	FDoRepLifetimeParams Params;
 	Params.bIsPushBased = false;
 
+	DOREPLIFETIME_WITH_PARAMS_FAST(ThisClass, ConstructorInitializedInstancedStruct, Params);
 	DOREPLIFETIME_WITH_PARAMS_FAST(ThisClass, InstancedStruct, Params);
 	DOREPLIFETIME_WITH_PARAMS_FAST(ThisClass, InstancedStructArray, Params);
 	DOREPLIFETIME_WITH_PARAMS_FAST(ThisClass, StructWithInstancedStructArray, Params);
diff --git a/Engine/Plugins/Runtime/ReplicationSystemTestPlugin/Source/Private/Tests/Serialization/TestInstancedStructNetSerializer.h b/Engine/Plugins/Runtime/ReplicationSystemTestPlugin/Source/Private/Tests/Serialization/TestInstancedStructNetSerializer.h
--- a/Engine/Plugins/Runtime/ReplicationSystemTestPlugin/Source/Private/Tests/Serialization/TestInstancedStructNetSerializer.h
+++ b/Engine/Plugins/Runtime/ReplicationSystemTestPlugin/Source/Private/Tests/Serialization/TestInstancedStructNetSerializer.h
@@ -135,6 +135,11 @@ protected:
 	virtual void RegisterReplicationFragments(UE::Net::FFragmentRegistrationContext& Fragments, UE::Net::EFragmentRegistrationFlags RegistrationFlags) override;
 
 public:
+	UInstancedStructNetSerializerTestObject();
+
+	UPROPERTY(Replicated, Transient)
+	FInstancedStruct ConstructorInitializedInstancedStruct;
+
 	UPROPERTY(Replicated, Transient)
 	FInstancedStruct InstancedStruct;
 
diff --git a/Engine/Source/Runtime/Net/Iris/Private/Iris/Serialization/InstancedStructNetSerializer.cpp b/Engine/Source/Runtime/Net/Iris/Private/Iris/Serialization/InstancedStructNetSerializer.cpp
--- a/Engine/Source/Runtime/Net/Iris/Private/Iris/Serialization/InstancedStructNetSerializer.cpp
+++ b/Engine/Source/Runtime/Net/Iris/Private/Iris/Serialization/InstancedStructNetSerializer.cpp
@@ -352,7 +352,7 @@ void FInstancedStructNetSerializer::Quantize(FNetSerializationContext& Context,
 
 	FInstancedStructNetSerializerConfig* Config = const_cast<FInstancedStructNetSerializerConfig*>(static_cast<const FInstancedStructNetSerializerConfig*>(Args.NetSerializerConfig));
 
-	if (Source.IsValid())
+	if (Source.IsValid() && !Context.IsInitializingDefaultState())
 	{
 		const UScriptStruct* Struct = Source.GetScriptStruct();
 
@@ -405,7 +405,6 @@ void FInstancedStructNetSerializer::Quantize(FNetSerializationContext& Context,
 	Reset(Context, Config, Target);
 }
 
-// $IRIS TODO : Consider implementing Apply to avoid unnecessary memory operations.
 void FInstancedStructNetSerializer::Dequantize(FNetSerializationContext& Context, const FNetDequantizeArgs& Args)
 {
 	const QuantizedType& Source = *reinterpret_cast<QuantizedType*>(Args.Source);
