Iris
- Implemented support for async stop replication allowing cancel if start replication is issued again before stop replication has completed.

#Synced-CL 48784178

diff --git a/Engine/Plugins/Runtime/ReplicationSystemTestPlugin/Source/Private/Tests/ReplicationSystem/ReplicatedTestObject.cpp b/Engine/Plugins/Runtime/ReplicationSystemTestPlugin/Source/Private/Tests/ReplicationSystem/ReplicatedTestObject.cpp
--- a/Engine/Plugins/Runtime/ReplicationSystemTestPlugin/Source/Private/Tests/ReplicationSystem/ReplicatedTestObject.cpp
+++ b/Engine/Plugins/Runtime/ReplicationSystemTestPlugin/Source/Private/Tests/ReplicationSystem/ReplicatedTestObject.cpp
@@ -59,7 +59,8 @@ UE::Net::FNetRefHandle UReplicatedTestObjectBridge::BeginReplication(UReplicated
 {
 	// Create NetRefHandle for the registered fragments
 	FRootObjectReplicationParams Params;
-	FNetRefHandle Handle = StartReplicatingRootObject(Instance, Params, ReplicatedObjectFactoryId);
+
+	FNetRefHandle Handle = Super::StartReplicatingRootObject(Instance, Params, ReplicatedObjectFactoryId);
 
 	// This is optional but typically we want to cache at least the NetRefHandle in the game instance to avoid doing map lookups to find it
 	if (Handle.IsValid())
@@ -76,10 +77,7 @@ UE::Net::FNetRefHandle UReplicatedTestObjectBridge::BeginReplication(UReplicated
 	FNetRefHandle Handle = Super::StartReplicatingRootObject(Instance, Params, ReplicatedObjectFactoryId);
 
 	// This is optional but typically we want to cache at least the NetRefHandle in the game instance to avoid doing map lookups to find it
-	if (Handle.IsValid())
-	{
-		Instance->NetRefHandle = Handle;
-	}
+	Instance->NetRefHandle = Handle;
 
 	return Handle;
 }
diff --git a/Engine/Plugins/Runtime/ReplicationSystemTestPlugin/Source/Private/Tests/ReplicationSystem/TestFlushBeforeDestroy.cpp b/Engine/Plugins/Runtime/ReplicationSystemTestPlugin/Source/Private/Tests/ReplicationSystem/TestFlushBeforeDestroy.cpp
--- a/Engine/Plugins/Runtime/ReplicationSystemTestPlugin/Source/Private/Tests/ReplicationSystem/TestFlushBeforeDestroy.cpp
+++ b/Engine/Plugins/Runtime/ReplicationSystemTestPlugin/Source/Private/Tests/ReplicationSystem/TestFlushBeforeDestroy.cpp
@@ -684,118 +684,6 @@ UE_NET_TEST_FIXTURE(FTestFlushBeforeDestroyFixture, TestStateInFlightFlushedBefo
 	UE_NET_ASSERT_EQ(Client->GetReplicationBridge()->GetReplicatedObject(ObjectHandle), nullptr);
 }
 
-UE_NET_TEST_FIXTURE(FTestFlushBeforeDestroyFixture, TestDroppedPendingTearOffIsCancelledByEndReplication)
-{
-	// As we are testing old behavior, we need to make sure to allow double endreplication so we hit the path we want to test.
-	IConsoleVariable* CVarAllowDestroyToCancelFlushAndTearOff = IConsoleManager::Get().FindConsoleVariable(TEXT("net.Iris.AllowDestroyToCancelFlushAndTearOff"), false);
-	UE_NET_ASSERT_NE(CVarAllowDestroyToCancelFlushAndTearOff, nullptr);
-	UE_NET_ASSERT_TRUE(CVarAllowDestroyToCancelFlushAndTearOff->IsVariableBool());
-
-	const bool bOldAllowDestroyToCancelFlushAndTearOff = CVarAllowDestroyToCancelFlushAndTearOff->GetBool();
-	ON_SCOPE_EXIT { CVarAllowDestroyToCancelFlushAndTearOff->Set(bOldAllowDestroyToCancelFlushAndTearOff, ECVF_SetByCode); };
-
-	CVarAllowDestroyToCancelFlushAndTearOff->Set(true, ECVF_SetByCode);
-
-	FReplicationSystemTestClient* Client = CreateClient();
-	RegisterNetBlobHandlers(Client);
-
-	Server->NetUpdate();
-	Server->SendAndDeliverTo(Client, DeliverPacket);
-	Server->PostSendUpdate();
-
-	// Setup case where we have a new object for which we have an attachment which should execute a tearoff after we have confirmed creation
-	UReplicatedTestObject* ServerObject = Server->CreateObject(0, 0);
-	FNetRefHandle ObjectHandle = ServerObject->NetRefHandle;
-
-	// Create attachment
-	{
-		constexpr uint32 PayloadBitCount = 24;
-		const TRefCountPtr<FNetObjectAttachment>& Attachment = MockNetObjectAttachmentHandler->CreateReliableNetObjectAttachment(PayloadBitCount);
-		FNetObjectReference AttachmentTarget = FObjectReferenceCache::MakeNetObjectReference(ServerObject->NetRefHandle);
-		Server->GetReplicationSystem()->QueueNetObjectAttachment(Client->ConnectionIdOnServer, AttachmentTarget, Attachment);
-	}
-
-	// Request tearoff
-	Server->ReplicationBridge->EndReplication(ServerObject, EEndReplicationFlags::TearOff);
-
-	// Send packet so that we have creationdata in flight
-	Server->NetUpdate();
-	Server->SendTo(Client);
-	Server->PostSendUpdate();
-
-	// Force destroy object already pending tearoff/flush. DestroyLocalNetHandle will invalidate cached creationinfo.
-	Server->ReplicationBridge->EndReplication(ServerObject, EEndReplicationFlags::Destroy);
-
-	// Drop and notify that the packet while object still is the state waitoncreateconfirmation as we have not yet updated scope.
-	// When this failed it did put the state of the object back in PendingCreate even though we no longer had any cached creationinfo.
-	Server->DeliverTo(Client, false);
-
-	// Deliver a packet, this should flush the object and deliver the attachment
-	Server->NetUpdate();
-	Server->SendAndDeliverTo(Client, DeliverPacket);
-	Server->PostSendUpdate();
-
-	// Verify that the attachment has not been received
-	UE_NET_ASSERT_NE(ClientMockNetObjectAttachmentHandler->GetFunctionCallCounts().OnNetBlobReceived, 1U);
-
-	// Deliver a packet. Should destroy the object on the client unless that was already done
-	Server->NetUpdate();
-	Server->SendAndDeliverTo(Client, DeliverPacket);
-	Server->PostSendUpdate();
-
-	// Verify that object is destroyed
-	UE_NET_ASSERT_EQ(Client->GetReplicationBridge()->GetReplicatedObject(ObjectHandle), nullptr);
-}
-
-UE_NET_TEST_FIXTURE(FTestFlushBeforeDestroyFixture, TestPendingCreateTearOffIsCancelledByEndReplication)
-{
-	// As we are testing old behavior, we need to make sure to allow double endreplication so we hit the path we want to test.
-	IConsoleVariable* CVarAllowDestroyToCancelFlushAndTearOff = IConsoleManager::Get().FindConsoleVariable(TEXT("net.Iris.AllowDestroyToCancelFlushAndTearOff"), false);
-	UE_NET_ASSERT_NE(CVarAllowDestroyToCancelFlushAndTearOff, nullptr);
-	UE_NET_ASSERT_TRUE(CVarAllowDestroyToCancelFlushAndTearOff->IsVariableBool());
-
-	const bool bOldAllowDestroyToCancelFlushAndTearOff = CVarAllowDestroyToCancelFlushAndTearOff->GetBool();
-	ON_SCOPE_EXIT { CVarAllowDestroyToCancelFlushAndTearOff->Set(bOldAllowDestroyToCancelFlushAndTearOff, ECVF_SetByCode); };
-
-	CVarAllowDestroyToCancelFlushAndTearOff->Set(true, ECVF_SetByCode);
-
-	FReplicationSystemTestClient* Client = CreateClient();
-	RegisterNetBlobHandlers(Client);
-
-	// Setup case where we have a new object for which we have an attachment which should execute a tearoff after we have confirmed creation
-	UReplicatedTestObject* ServerObject = Server->CreateObject(0, 0);
-	FNetRefHandle ObjectHandle = ServerObject->NetRefHandle;
-
-	// Create attachment
-	{
-		constexpr uint32 PayloadBitCount = 24;
-		const TRefCountPtr<FNetObjectAttachment>& Attachment = MockNetObjectAttachmentHandler->CreateReliableNetObjectAttachment(PayloadBitCount);
-		FNetObjectReference AttachmentTarget = FObjectReferenceCache::MakeNetObjectReference(ServerObject->NetRefHandle);
-		Server->GetReplicationSystem()->QueueNetObjectAttachment(Client->ConnectionIdOnServer, AttachmentTarget, Attachment);
-	}
-
-	// Request tearoff
-	Server->ReplicationBridge->EndReplication(ServerObject, EEndReplicationFlags::TearOff);
-
-	// PreUpdate to update scoping to get the object into the PendingCreate state.
-	Server->NetUpdate();
-	Server->PostSendUpdate();
-
-	// Force destroy object already pending tearoff/flush. DestroyLocalNetHandle will invalidate cached creationinfo.
-	Server->ReplicationBridge->EndReplication(ServerObject, EEndReplicationFlags::Destroy);
-
-	// Send a packet.
-	Server->NetUpdate();
-	Server->SendAndDeliverTo(Client, DeliverPacket);
-	Server->PostSendUpdate();
-
-	// Verify that the attachment has not been received.
-	UE_NET_ASSERT_NE(ClientMockNetObjectAttachmentHandler->GetFunctionCallCounts().OnNetBlobReceived, 1U);
-
-	// Verify that the object is not created.
-	UE_NET_ASSERT_EQ(Client->GetReplicationBridge()->GetReplicatedObject(ObjectHandle), nullptr);
-}
-
 UE_NET_TEST_FIXTURE(FTestFlushBeforeDestroyFixture, TestPendingCreateTearOffIsNotCancelledByEndReplication)
 {
 	FReplicationSystemTestClient* Client = CreateClient();
diff --git a/Engine/Plugins/Runtime/ReplicationSystemTestPlugin/Source/Private/Tests/ReplicationSystem/TestReplicationSystem.cpp b/Engine/Plugins/Runtime/ReplicationSystemTestPlugin/Source/Private/Tests/ReplicationSystem/TestReplicationSystem.cpp
--- a/Engine/Plugins/Runtime/ReplicationSystemTestPlugin/Source/Private/Tests/ReplicationSystem/TestReplicationSystem.cpp
+++ b/Engine/Plugins/Runtime/ReplicationSystemTestPlugin/Source/Private/Tests/ReplicationSystem/TestReplicationSystem.cpp
@@ -2453,6 +2453,371 @@ UE_NET_TEST_FIXTURE(FReplicationSystemServerClientTestFixture, TestSubObjectRepl
 	UE_NET_ASSERT_FALSE(ClientSubObject2->ObjectToWatch->IsValid());
 }
 
+UE_NET_TEST_FIXTURE(FReplicationSystemServerClientTestFixture, TestStablyNamedObject)
+{
+	const TCHAR* StableObjectName = TEXT("StableObjectName");
+
+	UReplicationSystem* ReplicationSystem = Server->ReplicationSystem;
+
+	// Add a client
+	FReplicationSystemTestClient* Client = CreateClient();
+
+	// Create stably named object on both server and client
+	UTestReplicatedIrisObject* ServerObject = Server->CreateStablyNamedObject<UTestReplicatedIrisObject>(StableObjectName);
+	Client->CreateStablyNamedObject<UTestReplicatedIrisObject>(StableObjectName);
+
+	UE_NET_ASSERT_TRUE(ServerObject->NetRefHandle.IsValid());
+	UE_NET_ASSERT_TRUE(ServerObject->NetRefHandle.IsStatic());
+
+	// Send and deliver packet
+	Server->UpdateAndSend({Client});
+
+	// We expect object to exist on both server and client
+	UE_NET_ASSERT_TRUE(Cast<UTestReplicatedIrisObject>(Client->GetReplicationBridge()->GetReplicatedObject(ServerObject->NetRefHandle)) != nullptr);
+	UE_NET_ASSERT_TRUE(Cast<UTestReplicatedIrisObject>(Server->GetReplicationBridge()->GetReplicatedObject(ServerObject->NetRefHandle)) != nullptr);
+ 
+	// Stop/end replication using async path
+	Server->ReplicationBridge->EndReplication(ServerObject);
+
+	// Send and deliver packet
+	Server->UpdateAndSend({Client});
+
+	// We expect object to no longer exist on client and server
+	UE_NET_ASSERT_TRUE(Cast<UTestReplicatedIrisObject>(Client->GetReplicationBridge()->GetReplicatedObject(ServerObject->NetRefHandle)) == nullptr);
+	UE_NET_ASSERT_TRUE(Cast<UTestReplicatedIrisObject>(Server->GetReplicationBridge()->GetReplicatedObject(ServerObject->NetRefHandle)) == nullptr);
+}
+
+// expect fail here as we try to replicate object before it is allowed
+#if 0
+UE_NET_TEST_FIXTURE(FReplicationSystemServerClientTestFixture, TestStablyNamedObjectCannotReplicateWhilePendingDestroy)
+{
+	const TCHAR* StableObjectName = TEXT("StableObjectName");
+
+	UReplicationSystem* ReplicationSystem = Server->ReplicationSystem;
+
+	// Add a client
+	FReplicationSystemTestClient* Client = CreateClient();
+
+	// Create stably named object on both server and client
+	UTestReplicatedIrisObject* ServerObject = Server->CreateStablyNamedObject<UTestReplicatedIrisObject>(StableObjectName);
+	Client->CreateStablyNamedObject<UTestReplicatedIrisObject>(StableObjectName);
+
+	UE_NET_ASSERT_TRUE(ServerObject->NetRefHandle.IsValid());
+	UE_NET_ASSERT_TRUE(ServerObject->NetRefHandle.IsStatic());
+
+	// Send and deliver packet
+	Server->UpdateAndSend({Client});
+
+	// We expect object to exist on both server and client
+	UE_NET_ASSERT_TRUE(Cast<UTestReplicatedIrisObject>(Client->GetReplicationBridge()->GetReplicatedObject(ServerObject->NetRefHandle)) != nullptr);
+	UE_NET_ASSERT_TRUE(Cast<UTestReplicatedIrisObject>(Server->GetReplicationBridge()->GetReplicatedObject(ServerObject->NetRefHandle)) != nullptr);
+ 
+	// Stop replication
+	Server->ReplicationBridge->EndReplication(ServerObject);
+
+	// Try to restart replication	
+	Server->ReplicationBridge->BeginReplication(ServerObject);
+
+	// Send and deliver packet
+	Server->UpdateAndSend({Client});
+
+	// We expect object to no longer exist on client and server
+	UE_NET_ASSERT_TRUE(Cast<UTestReplicatedIrisObject>(Client->GetReplicationBridge()->GetReplicatedObject(ServerObject->NetRefHandle)) == nullptr);
+	UE_NET_ASSERT_TRUE(Cast<UTestReplicatedIrisObject>(Server->GetReplicationBridge()->GetReplicatedObject(ServerObject->NetRefHandle)) == nullptr);
+}
+#endif
+
+UE_NET_TEST_FIXTURE(FReplicationSystemServerClientTestFixture, TestStablyNamedObjectAsyncStopReplication)
+{
+	const TCHAR* StableObjectName = TEXT("StableObjectName");
+
+	UReplicationSystem* ReplicationSystem = Server->ReplicationSystem;
+
+	// Add a client
+	FReplicationSystemTestClient* Client = CreateClient();
+
+	// Create stably named object on both server and client
+	UTestReplicatedIrisObject* ServerObject = Server->CreateStablyNamedObject<UTestReplicatedIrisObject>(StableObjectName);
+	Client->CreateStablyNamedObject<UTestReplicatedIrisObject>(StableObjectName);
+
+	UE_NET_ASSERT_TRUE(ServerObject->NetRefHandle.IsValid());
+	UE_NET_ASSERT_TRUE(ServerObject->NetRefHandle.IsStatic());
+
+	// Send and deliver packet
+	Server->UpdateAndSend({Client});
+
+	// We expect object to exist on both server and client
+	const UTestReplicatedIrisObject* ClientObject = Cast<UTestReplicatedIrisObject>(Client->GetReplicationBridge()->GetReplicatedObject(ServerObject->NetRefHandle));
+	UE_NET_ASSERT_TRUE(ClientObject != nullptr);
+	UE_NET_ASSERT_TRUE(Cast<UTestReplicatedIrisObject>(Server->GetReplicationBridge()->GetReplicatedObject(ServerObject->NetRefHandle)) != nullptr);
+
+	// Stop replication using async path
+	Server->ReplicationBridge->EndReplication(ServerObject, EEndReplicationFlags::AllowAsyncStopReplication | EEndReplicationFlags::Destroy);
+	
+	// Should be pending stop replication
+	UE_NET_ASSERT_TRUE(Server->ReplicationBridge->IsPendingStopReplication(ServerObject->NetRefHandle));
+	// Still replicated until the async stop is completed.
+	UE_NET_ASSERT_TRUE(Server->ReplicationBridge->IsReplicatedHandle(ServerObject->NetRefHandle));
+	
+	// This should cancel the async stop replication and keep replicating the object
+	Server->ReplicationBridge->BeginReplication(ServerObject);
+
+	// Should be pending stop replication
+	UE_NET_ASSERT_FALSE(Server->ReplicationBridge->IsPendingStopReplication(ServerObject->NetRefHandle));
+
+	// Send and deliver packet
+	Server->UpdateAndSend({Client});
+		
+	// We expect object to exist on both server and client and that it should be the same instance as we had before.
+	UE_NET_ASSERT_TRUE(Cast<UTestReplicatedIrisObject>(Client->GetReplicationBridge()->GetReplicatedObject(ServerObject->NetRefHandle)) == ClientObject);
+	UE_NET_ASSERT_TRUE(Cast<UTestReplicatedIrisObject>(Server->GetReplicationBridge()->GetReplicatedObject(ServerObject->NetRefHandle)) != nullptr);
+ 
+	// Stop replication using async path
+	Server->ReplicationBridge->EndReplication(ServerObject, EEndReplicationFlags::AllowAsyncStopReplication | EEndReplicationFlags::Destroy);
+	UE_NET_ASSERT_TRUE(Server->ReplicationBridge->IsPendingStopReplication(ServerObject->NetRefHandle));
+
+	// Send and deliver packet
+	Server->UpdateAndSend({Client});
+
+	// We expect object to no longer exist on client and server
+	UE_NET_ASSERT_TRUE(Cast<UTestReplicatedIrisObject>(Client->GetReplicationBridge()->GetReplicatedObject(ServerObject->NetRefHandle)) == nullptr);
+	UE_NET_ASSERT_TRUE(Cast<UTestReplicatedIrisObject>(Server->GetReplicationBridge()->GetReplicatedObject(ServerObject->NetRefHandle)) == nullptr);
+}
+
+UE_NET_TEST_FIXTURE(FReplicationSystemServerClientTestFixture, TestStablyNamedObjectAsyncStopReplicationWithTwoClients)
+{
+	const TCHAR* StableObjectName = TEXT("StableObjectName");
+
+	UReplicationSystem* ReplicationSystem = Server->ReplicationSystem;
+
+	// Add clients
+	CreateClient();
+	CreateClient();
+
+	// Create stably named object on both server and clients
+	UTestReplicatedIrisObject* ServerObject = Server->CreateStablyNamedObject<UTestReplicatedIrisObject>(StableObjectName);
+	TArray<UTestReplicatedIrisObject*> ClientObjects;
+	ClientObjects.Add(Clients[0]->CreateStablyNamedObject<UTestReplicatedIrisObject>(StableObjectName));
+	ClientObjects.Add(Clients[1]->CreateStablyNamedObject<UTestReplicatedIrisObject>(StableObjectName));
+
+	UE_NET_ASSERT_TRUE(ServerObject->NetRefHandle.IsValid());
+	UE_NET_ASSERT_TRUE(ServerObject->NetRefHandle.IsStatic());
+
+	// Send and deliver packet to first client but not to the other
+	Server->UpdateAndSend({Clients[0]});
+
+	// We expect object to exist on server and clients[0]
+	UE_NET_ASSERT_TRUE(Cast<UTestReplicatedIrisObject>(Clients[0]->GetReplicationBridge()->GetReplicatedObject(ServerObject->NetRefHandle)) == ClientObjects[0]);
+	UE_NET_ASSERT_TRUE(Cast<UTestReplicatedIrisObject>(Server->GetReplicationBridge()->GetReplicatedObject(ServerObject->NetRefHandle)) != nullptr);
+
+	// Stop replication using async path, intentionally not passing the destroy flag to keep instance
+	Server->ReplicationBridge->EndReplication(ServerObject, EEndReplicationFlags::AllowAsyncStopReplication);
+
+	// Should be pending stop replication
+	UE_NET_ASSERT_TRUE(Server->ReplicationBridge->IsPendingStopReplication(ServerObject->NetRefHandle));
+	// Still replicated until the async stop is completed.
+	UE_NET_ASSERT_TRUE(Server->ReplicationBridge->IsReplicatedHandle(ServerObject->NetRefHandle));
+	
+	// This should cancel the async stop replication and keep replicating the object
+	Server->ReplicationBridge->BeginReplication(ServerObject);
+
+	// Should be pending stop replication
+	UE_NET_ASSERT_FALSE(Server->ReplicationBridge->IsPendingStopReplication(ServerObject->NetRefHandle));
+
+	// Send and deliver packet to both clients
+	Server->UpdateAndSend(Clients);
+		
+	// We expect object to exist on server and both clients
+	UE_NET_ASSERT_TRUE(Cast<UTestReplicatedIrisObject>(Clients[0]->GetReplicationBridge()->GetReplicatedObject(ServerObject->NetRefHandle)) == ClientObjects[0]);
+	UE_NET_ASSERT_TRUE(Cast<UTestReplicatedIrisObject>(Clients[1]->GetReplicationBridge()->GetReplicatedObject(ServerObject->NetRefHandle)) == ClientObjects[1]);
+	UE_NET_ASSERT_TRUE(Cast<UTestReplicatedIrisObject>(Server->GetReplicationBridge()->GetReplicatedObject(ServerObject->NetRefHandle)) != nullptr);
+ 
+	// Stop replication using async path, intentionally not passing the destroy flag to keep instance
+	Server->ReplicationBridge->EndReplication(ServerObject, EEndReplicationFlags::AllowAsyncStopReplication);
+
+	// Send and deliver packet to only client 0
+	Server->UpdateAndSend({Clients[0]});
+
+	// We expect object to no longer exist on clients[0] and server should be pending async stop replication
+	UE_NET_ASSERT_TRUE(Cast<UTestReplicatedIrisObject>(Clients[0]->GetReplicationBridge()->GetReplicatedObject(ServerObject->NetRefHandle)) == nullptr);
+	UE_NET_ASSERT_TRUE(Cast<UTestReplicatedIrisObject>(Clients[1]->GetReplicationBridge()->GetReplicatedObject(ServerObject->NetRefHandle)) ==  ClientObjects[1]);
+	UE_NET_ASSERT_TRUE(Server->ReplicationBridge->IsPendingStopReplication(ServerObject->NetRefHandle));
+
+	// This should cancel the async stop replication and resume replicating the object, 
+	Server->ReplicationBridge->BeginReplication(ServerObject);
+	UE_NET_ASSERT_FALSE(Server->ReplicationBridge->IsPendingStopReplication(ServerObject->NetRefHandle));
+
+	// Update and send to all 
+	Server->UpdateAndSend(Clients);
+
+	// Object should now be restored on both clients
+	UE_NET_ASSERT_TRUE(Cast<UTestReplicatedIrisObject>(Clients[0]->GetReplicationBridge()->GetReplicatedObject(ServerObject->NetRefHandle)) == ClientObjects[0]);
+	UE_NET_ASSERT_TRUE(Cast<UTestReplicatedIrisObject>(Clients[1]->GetReplicationBridge()->GetReplicatedObject(ServerObject->NetRefHandle)) == ClientObjects[1]);
+	UE_NET_ASSERT_TRUE(Cast<UTestReplicatedIrisObject>(Server->GetReplicationBridge()->GetReplicatedObject(ServerObject->NetRefHandle)) != nullptr);
+
+	// Stop replication using async path, intentionally not passing the destroy flag to keep instance
+	Server->ReplicationBridge->EndReplication(ServerObject, EEndReplicationFlags::AllowAsyncStopReplication);
+
+	// Update and send to all 
+	Server->UpdateAndSend(Clients);
+
+	// We now expect object to be considered fully destroyed.
+	UE_NET_ASSERT_TRUE(Cast<UTestReplicatedIrisObject>(Clients[0]->GetReplicationBridge()->GetReplicatedObject(ServerObject->NetRefHandle)) == nullptr);
+	UE_NET_ASSERT_TRUE(Cast<UTestReplicatedIrisObject>(Clients[1]->GetReplicationBridge()->GetReplicatedObject(ServerObject->NetRefHandle)) == nullptr);
+	UE_NET_ASSERT_TRUE(Cast<UTestReplicatedIrisObject>(Server->GetReplicationBridge()->GetReplicatedObject(ServerObject->NetRefHandle)) == nullptr);
+}
+
+UE_NET_TEST_FIXTURE(FReplicationSystemServerClientTestFixture, TestDynamicObjectAsyncStopReplication)
+{
+	UReplicationSystem* ReplicationSystem = Server->ReplicationSystem;
+
+	// Add a client
+	FReplicationSystemTestClient* Client = CreateClient();
+
+	// Create dynamic object
+	UTestReplicatedIrisObject* ServerObject = Server->CreateObject<UTestReplicatedIrisObject>();
+
+	UE_NET_ASSERT_TRUE(ServerObject->NetRefHandle.IsValid());
+	UE_NET_ASSERT_TRUE(ServerObject->NetRefHandle.IsDynamic());
+
+	const FNetRefHandle OriginalServerHandle = ServerObject->NetRefHandle;
+
+	// Send and deliver packet
+	Server->UpdateAndSend({Client});
+
+	// Stop replication using async path
+	Server->ReplicationBridge->EndReplication(ServerObject, EEndReplicationFlags::AllowAsyncStopReplication | EEndReplicationFlags::Destroy);
+
+	// Should be pending stop replication
+	UE_NET_ASSERT_TRUE(Server->ReplicationBridge->IsPendingStopReplication(ServerObject->NetRefHandle));
+	UE_NET_ASSERT_TRUE(Server->ReplicationBridge->IsReplicatedHandle(ServerObject->NetRefHandle));
+	
+	// This should cancel the async stop replication and keep replicating the object using the same handle
+	Server->ReplicationBridge->BeginReplication(ServerObject);
+
+	// Verify that we did not get a new handle
+	UE_NET_ASSERT_EQ(OriginalServerHandle, ServerObject->NetRefHandle);
+	UE_NET_ASSERT_TRUE(Server->ReplicationBridge->IsReplicatedHandle(ServerObject->NetRefHandle));
+
+	// Should be pending stop replication
+	UE_NET_ASSERT_FALSE(Server->ReplicationBridge->IsPendingStopReplication(ServerObject->NetRefHandle));
+
+	// Send and deliver packet
+	Server->UpdateAndSend({Client});
+		
+	// We expect object to exist on both server and client
+	UE_NET_ASSERT_TRUE(Cast<UTestReplicatedIrisObject>(Client->GetReplicationBridge()->GetReplicatedObject(ServerObject->NetRefHandle)) != nullptr);
+	UE_NET_ASSERT_TRUE(Cast<UTestReplicatedIrisObject>(Server->GetReplicationBridge()->GetReplicatedObject(ServerObject->NetRefHandle)) != nullptr);
+ 
+	// Stop replication using async path
+	Server->ReplicationBridge->EndReplication(ServerObject, EEndReplicationFlags::AllowAsyncStopReplication | EEndReplicationFlags::Destroy);
+
+	// Send and deliver packet
+	Server->UpdateAndSend({Client});
+
+	// We expect object to no longer exist on client and server
+	UE_NET_ASSERT_TRUE(Cast<UTestReplicatedIrisObject>(Client->GetReplicationBridge()->GetReplicatedObject(ServerObject->NetRefHandle)) == nullptr);
+	UE_NET_ASSERT_TRUE(Cast<UTestReplicatedIrisObject>(Server->GetReplicationBridge()->GetReplicatedObject(ServerObject->NetRefHandle)) == nullptr);
+
+	// This should start replicating the object again using a new handle, as the object had finished the async stop replication.
+	Server->ReplicationBridge->BeginReplication(ServerObject);
+
+	// We should get a new handle for our dynamic object
+	Server->ReplicationBridge->BeginReplication(ServerObject);
+
+	UE_NET_ASSERT_NE(OriginalServerHandle, ServerObject->NetRefHandle);	
+}
+
+UE_NET_TEST_FIXTURE(FReplicationSystemServerClientTestFixture, TestDynamicObjectAsyncStopReplicationAndFlushWitTwoClients)
+{
+	UReplicationSystem* ReplicationSystem = Server->ReplicationSystem;
+
+	// Create clients
+	CreateClient();
+	CreateClient();
+
+	// Create dynamic object
+	UTestReplicatedIrisObject* ServerObject = Server->CreateObject<UTestReplicatedIrisObject>();
+
+	UE_NET_ASSERT_TRUE(ServerObject->NetRefHandle.IsValid());
+	UE_NET_ASSERT_TRUE(ServerObject->NetRefHandle.IsDynamic());
+
+	const FNetRefHandle OriginalServerHandle = ServerObject->NetRefHandle;
+
+	// Send and deliver packet to clients
+	Server->UpdateAndSend(Clients);
+
+	// The object should now exist on all clients
+	UE_NET_ASSERT_TRUE(Cast<UTestReplicatedIrisObject>(Clients[0]->GetReplicationBridge()->GetReplicatedObject(ServerObject->NetRefHandle)) != nullptr);
+	UE_NET_ASSERT_TRUE(Cast<UTestReplicatedIrisObject>(Clients[1]->GetReplicationBridge()->GetReplicatedObject(ServerObject->NetRefHandle)) != nullptr);
+	UE_NET_ASSERT_TRUE(Cast<UTestReplicatedIrisObject>(Server->GetReplicationBridge()->GetReplicatedObject(ServerObject->NetRefHandle)) != nullptr);
+	
+	// Stop replication using async path
+	Server->ReplicationBridge->EndReplication(ServerObject, EEndReplicationFlags::AllowAsyncStopReplication | EEndReplicationFlags::Destroy);
+	// Should be pending stop replication
+	UE_NET_ASSERT_TRUE(Server->ReplicationBridge->IsPendingStopReplication(ServerObject->NetRefHandle));
+
+	// Send and deliver packet to client 0
+	Server->UpdateAndSend({Clients[0]});
+
+	// We expect object to exist on server and clients[1], but not on clients[0]
+	UE_NET_ASSERT_TRUE(Cast<UTestReplicatedIrisObject>(Server->GetReplicationBridge()->GetReplicatedObject(ServerObject->NetRefHandle)) != nullptr);
+	UE_NET_ASSERT_TRUE(Cast<UTestReplicatedIrisObject>(Clients[0]->GetReplicationBridge()->GetReplicatedObject(ServerObject->NetRefHandle)) == nullptr);
+	UE_NET_ASSERT_TRUE(Cast<UTestReplicatedIrisObject>(Clients[1]->GetReplicationBridge()->GetReplicatedObject(ServerObject->NetRefHandle)) != nullptr);
+	// Should be pending stop replication
+	UE_NET_ASSERT_TRUE(Server->ReplicationBridge->IsPendingStopReplication(ServerObject->NetRefHandle));
+
+	// This should cancel the async stop replication and keep replicating the object using the same handle
+	Server->ReplicationBridge->BeginReplication(ServerObject);
+
+	// Should still use the Original handle
+	UE_NET_ASSERT_EQ(OriginalServerHandle, ServerObject->NetRefHandle);
+
+	// Should not be pending stop replication
+	UE_NET_ASSERT_FALSE(Server->ReplicationBridge->IsPendingStopReplication(ServerObject->NetRefHandle));
+
+	// Send and deliver packet to clients
+	Server->UpdateAndSend(Clients);
+
+	// The object should now exist on all clients
+	UE_NET_ASSERT_TRUE(Cast<UTestReplicatedIrisObject>(Clients[0]->GetReplicationBridge()->GetReplicatedObject(ServerObject->NetRefHandle)) != nullptr);
+	UE_NET_ASSERT_TRUE(Cast<UTestReplicatedIrisObject>(Clients[1]->GetReplicationBridge()->GetReplicatedObject(ServerObject->NetRefHandle)) != nullptr);
+	UE_NET_ASSERT_TRUE(Cast<UTestReplicatedIrisObject>(Server->GetReplicationBridge()->GetReplicatedObject(ServerObject->NetRefHandle)) != nullptr);
+
+	// Stop replication using async path
+	Server->ReplicationBridge->EndReplication(ServerObject, EEndReplicationFlags::AllowAsyncStopReplication);
+
+	// Send and deliver packet to client 0
+	Server->UpdateAndSend({Clients[0]});
+
+	// We expect object to no longer exist on client 0
+	UE_NET_ASSERT_TRUE(Cast<UTestReplicatedIrisObject>(Clients[0]->GetReplicationBridge()->GetReplicatedObject(ServerObject->NetRefHandle)) == nullptr);
+	UE_NET_ASSERT_TRUE(Cast<UTestReplicatedIrisObject>(Clients[1]->GetReplicationBridge()->GetReplicatedObject(ServerObject->NetRefHandle)) != nullptr);
+	UE_NET_ASSERT_TRUE(Cast<UTestReplicatedIrisObject>(Server->GetReplicationBridge()->GetReplicatedObject(ServerObject->NetRefHandle)) != nullptr);
+	// Should be pending stop replication
+	UE_NET_ASSERT_TRUE(Server->ReplicationBridge->IsPendingStopReplication(ServerObject->NetRefHandle));
+
+	// Send and deliver packet to all clients
+	Server->UpdateAndSend(Clients);
+
+	// We expect object to no longer exist and no longer be pending end replication
+	UE_NET_ASSERT_TRUE(Cast<UTestReplicatedIrisObject>(Clients[0]->GetReplicationBridge()->GetReplicatedObject(ServerObject->NetRefHandle)) == nullptr);
+	UE_NET_ASSERT_TRUE(Cast<UTestReplicatedIrisObject>(Clients[1]->GetReplicationBridge()->GetReplicatedObject(ServerObject->NetRefHandle)) == nullptr);
+	UE_NET_ASSERT_TRUE(Cast<UTestReplicatedIrisObject>(Server->GetReplicationBridge()->GetReplicatedObject(ServerObject->NetRefHandle)) == nullptr);
+	UE_NET_ASSERT_FALSE(Server->ReplicationBridge->IsPendingStopReplication(ServerObject->NetRefHandle));
+
+	// This should start replicating the same object again using a new handle
+	Server->ReplicationBridge->BeginReplication(ServerObject);
+	UE_NET_ASSERT_NE(OriginalServerHandle, ServerObject->NetRefHandle);	
+
+	// Send and deliver packet to all clients
+	Server->UpdateAndSend(Clients);
+
+	// New instance of dynamic object should exist on all clients
+	UE_NET_ASSERT_TRUE(Cast<UTestReplicatedIrisObject>(Clients[0]->GetReplicationBridge()->GetReplicatedObject(ServerObject->NetRefHandle)) != nullptr);
+	UE_NET_ASSERT_TRUE(Cast<UTestReplicatedIrisObject>(Clients[1]->GetReplicationBridge()->GetReplicatedObject(ServerObject->NetRefHandle)) != nullptr);
+}
+
 } // end namespace UE::Net::Private
 
 
diff --git a/Engine/Source/Runtime/Net/Iris/Private/Iris/ReplicationSystem/NetRefHandleManager.cpp b/Engine/Source/Runtime/Net/Iris/Private/Iris/ReplicationSystem/NetRefHandleManager.cpp
--- a/Engine/Source/Runtime/Net/Iris/Private/Iris/ReplicationSystem/NetRefHandleManager.cpp
+++ b/Engine/Source/Runtime/Net/Iris/Private/Iris/ReplicationSystem/NetRefHandleManager.cpp
@@ -435,6 +435,20 @@ FNetRefHandle FNetRefHandleManager::CreateNetObjectFromRemote(FNetRefHandle Want
 	}
 }
 
+void FNetRefHandleManager::ResetDestroyedStartupObject(FInternalNetRefIndex InternalIndex)
+{
+	// Cleanup cross reference to destruction info
+	if (DestroyedStartupObjectInternalIndices.GetBit(InternalIndex))
+	{
+		DestroyedStartupObjectInternalIndices.ClearBit(InternalIndex);
+		uint32 OtherInternalIndex = 0U;
+		if (DestroyedStartupObject.RemoveAndCopyValue(InternalIndex, OtherInternalIndex))
+		{
+			DestroyedStartupObject.Remove(OtherInternalIndex);
+		}
+	}
+}
+
 void FNetRefHandleManager::InternalDestroyNetObject(FInternalNetRefIndex InternalIndex)
 {
 	UE_LOG(LogIris, Verbose, TEXT("FNetRefHandleManager::InternalDestroyNetObject: %s"), *PrintObjectFromIndex(InternalIndex));
@@ -509,15 +523,7 @@ void FNetRefHandleManager::InternalDestroyNetObject(FInternalNetRefIndex Interna
 	ClearStateForFreedInternalIndex(InternalIndex);
 
 	// Cleanup cross reference to destruction info
-	if (DestroyedStartupObjectInternalIndices.GetBit(InternalIndex))
-	{
-		DestroyedStartupObjectInternalIndices.ClearBit(InternalIndex);
-		uint32 OtherInternalIndex = 0U;
-		if (DestroyedStartupObject.RemoveAndCopyValue(InternalIndex, OtherInternalIndex))
-		{
-			DestroyedStartupObject.Remove(OtherInternalIndex);
-		}
-	}
+	ResetDestroyedStartupObject(InternalIndex);
 
 	--ActiveObjectCount;
 }
diff --git a/Engine/Source/Runtime/Net/Iris/Private/Iris/ReplicationSystem/NetRefHandleManager.h b/Engine/Source/Runtime/Net/Iris/Private/Iris/ReplicationSystem/NetRefHandleManager.h
--- a/Engine/Source/Runtime/Net/Iris/Private/Iris/ReplicationSystem/NetRefHandleManager.h
+++ b/Engine/Source/Runtime/Net/Iris/Private/Iris/ReplicationSystem/NetRefHandleManager.h
@@ -210,6 +210,9 @@ public:
 	// Returns a new handle used to replicate the destruction info, and adds a cross-reference to the original handle being destroyed if that still exists 
 	// Handle is a reference to an already replicated Handle that is about to be destroyed
 	FNetRefHandle CreateHandleForDestructionInfo(FNetRefHandle Handle, const FCreateNetObjectParams& Params);
+
+	// Reset state of object tracked as internal index.
+	void ResetDestroyedStartupObject(FInternalNetRefIndex InternalIndex);
 	
 	void DestroyNetObject(FNetRefHandle Handle);
 
diff --git a/Engine/Source/Runtime/Net/Iris/Private/Iris/ReplicationSystem/ObjectReplicationBridge.cpp b/Engine/Source/Runtime/Net/Iris/Private/Iris/ReplicationSystem/ObjectReplicationBridge.cpp
--- a/Engine/Source/Runtime/Net/Iris/Private/Iris/ReplicationSystem/ObjectReplicationBridge.cpp
+++ b/Engine/Source/Runtime/Net/Iris/Private/Iris/ReplicationSystem/ObjectReplicationBridge.cpp
@@ -478,6 +478,26 @@ void UObjectReplicationBridge::RemoveDestructionInfosForGroup(UE::Net::FNetObjec
 	}
 }
 
+void UObjectReplicationBridge::RemoveDestructionInfo(UE::Net::Private::FInternalNetRefIndex InternalObjectIndex)
+{
+	using namespace UE::Net;
+	using namespace UE::Net::Private;
+
+	if (InternalObjectIndex == FNetRefHandleManager::InvalidInternalIndex)
+	{
+		return;
+	}
+
+	const FNetRefHandleManager::FReplicatedObjectData& ObjectData = NetRefHandleManager->GetReplicatedObjectDataNoCheck(InternalObjectIndex);
+	if (StaticObjectsPendingDestroy.Remove(ObjectData.RefHandle))
+	{
+		FWorldLocations& WorldLocations = GetReplicationSystem()->GetReplicationSystemInternal()->GetWorldLocations();
+
+		NetRefHandleManager->DestroyNetObject(ObjectData.RefHandle);
+		WorldLocations.RemoveObjectInfoCache(InternalObjectIndex);
+	}
+}
+
 void UObjectReplicationBridge::ReadAndExecuteDestructionInfoFromRemote(FReplicationBridgeSerializationContext& Context)
 {
 	using namespace UE::Net;
@@ -586,9 +606,9 @@ UE::Net::FNetRefHandle UObjectReplicationBridge::StartReplicatingNetObject(UObje
 	}
 
 	// The instance is already replicating, just return the handle
-	if (IsReplicatedHandle(AllocatedRefHandle))
+	if (IsReplicatedHandle(AllocatedRefHandle) && RestartOrStopObjectPendingAsyncStopReplication(AllocatedRefHandle, bForceNewHandle))
 	{
-		ensureMsgf(bForceNewHandle == false, TEXT("bForceNewHandle was requested for %s but it is already replicating. The object must stop replicating first."), *PrintObjectFromNetRefHandle(AllocatedRefHandle));
+		ensureMsgf(bForceNewHandle == false, TEXT("bForceNewHandle was requested for %s but it is already replicating. The object must stop replicating first or use async stop replication."), *PrintObjectFromNetRefHandle(AllocatedRefHandle));
 		return AllocatedRefHandle;
 	}
 
@@ -898,7 +918,7 @@ UE::Net::FNetRefHandle UObjectReplicationBridge::StartReplicatingSubObject(UObje
 	}
 
 	FNetRefHandle SubObjectRefHandle = GetReplicatedRefHandle(Instance);
-	if (SubObjectRefHandle.IsValid())
+	if (SubObjectRefHandle.IsValid() && RestartOrStopObjectPendingAsyncStopReplication(SubObjectRefHandle, Params.bForceNewHandle))
 	{
 		// Verify that the existing object is a subobject of the owner
 		const FNetRefHandle ExistingRootObjectHandle = LocalNetRefHandleManager.GetRootObjectOfSubObject(SubObjectRefHandle);
@@ -1188,7 +1208,7 @@ void UObjectReplicationBridge::StopReplicatingNetObject(UObject* Instance, EEndR
 	const FNetRefHandle RefHandle = GetReplicatedRefHandle(Instance, EGetRefHandleFlags::EvenIfGarbage);
 	if (RefHandle.IsValid())
 	{
-		ensureMsgf(IsValid(Instance), TEXT("Calling EndReplication for Invalid Object: %s."), *PrintObjectFromNetRefHandle(RefHandle));
+		ensureMsgf(IsValid(Instance), TEXT("Calling StopReplicatingNetObject for Invalid Object: %s."), *PrintObjectFromNetRefHandle(RefHandle));
 		StopReplicatingNetRefHandle(RefHandle, EndReplicationFlags);
 	}
 }
diff --git a/Engine/Source/Runtime/Net/Iris/Private/Iris/ReplicationSystem/ReplicationBridge.cpp b/Engine/Source/Runtime/Net/Iris/Private/Iris/ReplicationSystem/ReplicationBridge.cpp
--- a/Engine/Source/Runtime/Net/Iris/Private/Iris/ReplicationSystem/ReplicationBridge.cpp
+++ b/Engine/Source/Runtime/Net/Iris/Private/Iris/ReplicationSystem/ReplicationBridge.cpp
@@ -37,6 +37,7 @@
 
 #include "ReplicationFragmentInternal.h"
 #include "ReplicationOperationsInternal.h"
+#include "Iris/ReplicationSystem/ReplicationBridgeTypes.h"
 
 #include UE_INLINE_GENERATED_CPP_BY_NAME(ReplicationBridge)
 
@@ -49,13 +50,6 @@ static FAutoConsoleVariableRef CVarEnableFlushReliableRPCOnDestroy(
 	TEXT("When true EEndReplicationFlags::Flush flag will be appended in EndReplication if we have pending unprocessed attachments/RPC:s when destroying a replicated object.")
 );
 
-static bool bAllowDestroyToCancelFlushAndTearOff = false;
-static FAutoConsoleVariableRef CVarAllowDestroyToCancelFlushAndTearOff(
-	TEXT("net.Iris.AllowDestroyToCancelFlushAndTearOff"),
-	bAllowDestroyToCancelFlushAndTearOff,
-	TEXT("When true issuing a EndReplication on an object that is already Tear-off or pending endreplication will cancel destroy/flush and destroy the replicated objects.")
-);
-
 /**
  * ReplicationBridge Implementation
  */
@@ -390,6 +384,42 @@ void UReplicationBridge::InternalDestroySubObjects(FNetRefHandle OwnerHandle, EE
 	}
 }
 
+void UReplicationBridge::InternalRequestAsyncStopReplication(FNetRefHandle Handle, EEndReplicationFlags EndReplicationFlags)
+{
+	using namespace UE::Net::Private;
+
+	// Destroy SubObjects
+	FInternalNetRefIndex ObjectInternalIndex = NetRefHandleManager->GetInternalIndex(Handle);
+	if (ObjectInternalIndex != FNetRefHandleManager::InvalidInternalIndex)
+	{
+		AddPendingEndReplication(Handle, EndReplicationFlags);
+
+		// If flush is required we need to wait to after scope update before we initialize async end replication, otherwise we can do it now.
+		if (!EnumHasAnyFlags(EndReplicationFlags, EEndReplicationFlags::Flush))
+		{
+			if (!ReplicationSystem->IsInGroup(ReplicationSystem->GetAsyncStopReplicationNetObjectGroup(), Handle))
+			{
+				ReplicationSystem->AddToGroup(ReplicationSystem->GetAsyncStopReplicationNetObjectGroup(), Handle);
+			}
+		}
+
+		for (FInternalNetRefIndex SubObjectInternalIndex : NetRefHandleManager->GetChildSubObjects(ObjectInternalIndex))
+		{
+			FNetRefHandleManager::FReplicatedObjectData& SubObjectData = NetRefHandleManager->GetReplicatedObjectDataNoCheck(SubObjectInternalIndex);
+
+			const FNetRefHandle SubObjectHandle = SubObjectData.RefHandle;
+			const bool bDestroySubObjectWithOwner = SubObjectData.bDestroySubObjectWithOwner;
+				
+			// Tag subobject for async stop replication. The check against the scope is needed since the subobjects array might contain destroyed subobjects
+			if (bDestroySubObjectWithOwner && NetRefHandleManager->IsScopableIndex(SubObjectInternalIndex))
+			{
+				UE_LOG_REPLICATIONBRIDGE(Verbose, TEXT("InternalRequestAsyncStopReplicationForSubObjects %s"), *NetRefHandleManager->PrintObjectFromNetRefHandle(SubObjectHandle));
+				InternalRequestAsyncStopReplication(SubObjectHandle, EndReplicationFlags);
+			}
+		}
+	}
+}
+
 void UReplicationBridge::StopReplicatingNetRefHandle(FNetRefHandle Handle, EEndReplicationFlags EndReplicationFlags)
 {
 	using namespace UE::Net::Private;
@@ -404,7 +434,7 @@ void UReplicationBridge::StopReplicatingNetRefHandle(FNetRefHandle Handle, EEndR
 	{
 		FNetRefHandleManager::FReplicatedObjectData& ObjectData = NetRefHandleManager->GetReplicatedObjectDataNoCheck(InternalReplicationIndex);
 
-		if (ObjectData.bPendingEndReplication && !bAllowDestroyToCancelFlushAndTearOff)
+		if (ObjectData.bPendingEndReplication)
 		{
 			UE_LOG_REPLICATIONBRIDGE(Verbose, TEXT("Ignoring EndReplication called on object already PendingEndReplication %s."), *NetRefHandleManager->PrintObjectFromNetRefHandle(Handle));
 			return;
@@ -438,7 +468,21 @@ void UReplicationBridge::StopReplicatingNetRefHandle(FNetRefHandle Handle, EEndR
 				EnumAddFlags(EndReplicationFlags, EEndReplicationFlags::Flush);
 			}
 
-			if (EnumHasAnyFlags(EndReplicationFlags, EEndReplicationFlags::Flush))
+			const bool bFlushRequested = EnumHasAnyFlags(EndReplicationFlags, EEndReplicationFlags::Flush);
+			if (EnumHasAnyFlags(EndReplicationFlags, EEndReplicationFlags::AllowAsyncStopReplication))
+			{
+				// Defer stop replication until object is no longer scoped by any connections to allow restart replication.
+				AddPendingEndReplication(Handle, EndReplicationFlags);
+
+				if (bFlushRequested)
+				{
+					InternalFlushStateData(Handle);
+				}
+
+				// Note: When using async stop replication We do not detach the instance until it is no longer referenced by an instance.
+				InternalRequestAsyncStopReplication(Handle, EndReplicationFlags);
+			}
+			else if (bFlushRequested)
 			{
 				// Defer destroy until after scope update to allow create/destroy on same frame
 				AddPendingEndReplication(Handle, EndReplicationFlags);
@@ -522,9 +566,19 @@ void UReplicationBridge::UpdateHandlesPendingEndReplication()
 			}
 			else
 			{
-				// If the object is still in scope remove it from scope as objects pending EndReplication should not be added to new connections after the first update.
-				if (NetRefHandleManager->IsScopableIndex(ObjectInternalIndex))
+				if (EnumHasAnyFlags(Info.DestroyFlags, EEndReplicationFlags::AllowAsyncStopReplication))
 				{
+					FNetRefHandleManager::FReplicatedObjectData& ObjectData = NetRefHandleManager->GetReplicatedObjectDataNoCheck(ObjectInternalIndex);
+					// We keep the object in scope but we mark it a not replicated which will stop replicating it for all clients
+					if (!ReplicationSystem->IsInGroup(ReplicationSystem->GetAsyncStopReplicationNetObjectGroup(), Info.Handle))
+					{
+						ReplicationSystem->AddToGroup(ReplicationSystem->GetAsyncStopReplicationNetObjectGroup(), Info.Handle);
+					}
+				}
+				else if (NetRefHandleManager->IsScopableIndex(ObjectInternalIndex))
+				{
+					// If the object is still in scope remove it from scope as objects pending EndReplication should not be added to new connections after the first update.
+
 					// Mark object and subobjects as no longer scopeable, and that we should not propagate changed states
 					NetRefHandleManager->RemoveFromScope(ObjectInternalIndex);
 					for (FInternalNetRefIndex SubObjectIndex : NetRefHandleManager->GetSubObjects(ObjectInternalIndex))
@@ -547,7 +601,7 @@ void UReplicationBridge::UpdateHandlesPendingEndReplication()
 
 void UReplicationBridge::AddPendingEndReplication(FNetRefHandle Handle, EEndReplicationFlags DestroyFlags, EPendingEndReplicationImmediate Immediate)
 {
-	if (ensure(EnumHasAnyFlags(DestroyFlags, EEndReplicationFlags::Flush | EEndReplicationFlags::TearOff)))
+	if (ensure(EnumHasAnyFlags(DestroyFlags, EEndReplicationFlags::Flush | EEndReplicationFlags::TearOff | EEndReplicationFlags::AllowAsyncStopReplication)))
 	{
 		if (!HandlesPendingEndReplication.FindByPredicate([&](const FPendingEndReplicationInfo& Info){ return Info.Handle == Handle; }))
 		{
@@ -556,6 +610,58 @@ void UReplicationBridge::AddPendingEndReplication(FNetRefHandle Handle, EEndRepl
 	}
 }
 
+bool UReplicationBridge::IsPendingStopReplication(FNetRefHandle Handle) const
+{
+	const int32 Index = HandlesPendingEndReplication.FindLastByPredicate([Handle](const FPendingEndReplicationInfo& Info){ return Info.Handle == Handle; });
+	return Index != INDEX_NONE;
+}
+
+bool UReplicationBridge::RestartOrStopObjectPendingAsyncStopReplication(FNetRefHandle Handle, bool bForceStopReplication)
+{
+	using namespace UE::Net;
+	using namespace UE::Net::Private;
+
+	bool bRestarted = false;
+
+	const int32 Index = HandlesPendingEndReplication.FindLastByPredicate([Handle](const FPendingEndReplicationInfo& Info){ return Info.Handle == Handle; });
+	if (Index != INDEX_NONE && EnumHasAnyFlags(HandlesPendingEndReplication[Index].DestroyFlags, EEndReplicationFlags::AllowAsyncStopReplication))
+	{
+		UE_LOG_REPLICATIONBRIDGE(Verbose, TEXT("RestartOrStopObjectPendingAsyncStopReplication: %s ForceStop: %u"), *PrintObjectFromNetRefHandle(Handle), bForceStopReplication ? 1U : 0U);
+
+		const uint32 InternalObjectIndex = NetRefHandleManager->GetInternalIndex(HandlesPendingEndReplication[Index].Handle);
+		if (InternalObjectIndex != FNetRefHandleManager::InvalidInternalIndex)
+		{
+			FNetRefHandleManager::FReplicatedObjectData& ObjectData = NetRefHandleManager->GetReplicatedObjectDataNoCheck(InternalObjectIndex);
+			
+			// In some cases we want to just create a new net object instead so we should just carry on and stop replication now.
+			if (bForceStopReplication)
+			{
+				ObjectData.bPendingEndReplication = 1U;
+				DestroyLocalNetHandle(HandlesPendingEndReplication[Index].Handle, HandlesPendingEndReplication[Index].DestroyFlags);
+				bRestarted = false;
+			}
+			else
+			{
+				bRestarted = true;
+				if (NetRefHandleManager->IsDestroyedStartupObject(InternalObjectIndex))
+				{
+					// Need to cleanse destruction info if we are going to cancel async stop replication.
+					const FInternalNetRefIndex DestructionInfoInternalIndex = NetRefHandleManager->GetOriginalDestroyedStartupObjectIndex(InternalObjectIndex);
+					NetRefHandleManager->ResetDestroyedStartupObject(InternalObjectIndex);
+					RemoveDestructionInfo(DestructionInfoInternalIndex);
+				}
+			}
+		}
+
+		HandlesPendingEndReplication.RemoveAt(Index);
+		ReplicationSystem->RemoveFromGroup(ReplicationSystem->GetAsyncStopReplicationNetObjectGroup(), Handle);
+
+		return bRestarted;
+	}
+
+	return Index == INDEX_NONE;
+}
+
 void UReplicationBridge::InternalFlushStateData(UE::Net::FNetSerializationContext& SerializationContext, UE::Net::Private::FChangeMaskCache& ChangeMaskCache, UE::Net::FNetBitStreamWriter& ChangeMaskWriter, uint32 InternalObjectIndex)
 {
 	using namespace UE::Net;
@@ -596,8 +702,6 @@ void UReplicationBridge::InternalFlushStateData(UE::Net::FNetSerializationContex
 	{
 		InternalFlushStateData(SerializationContext, ChangeMaskCache, ChangeMaskWriter, SubObjectInternalIndex);
 	}
-
-	// $IRIS TODO:  Should we also clear the DirtyTracker flags for this flushed object ?
 }
 
 void UReplicationBridge::InternalFlushStateData(FNetRefHandle Handle)
diff --git a/Engine/Source/Runtime/Net/Iris/Private/Iris/ReplicationSystem/ReplicationBridgeTypes.cpp b/Engine/Source/Runtime/Net/Iris/Private/Iris/ReplicationSystem/ReplicationBridgeTypes.cpp
--- a/Engine/Source/Runtime/Net/Iris/Private/Iris/ReplicationSystem/ReplicationBridgeTypes.cpp
+++ b/Engine/Source/Runtime/Net/Iris/Private/Iris/ReplicationSystem/ReplicationBridgeTypes.cpp
@@ -47,6 +47,11 @@ FString LexToString(EEndReplicationFlags EndReplicationFlags)
 		if (!Flags.IsEmpty()) { Flags += TEXT(','); };
 		Flags += TEXT("Migrated");
 	}
+	if (EnumHasAnyFlags(EndReplicationFlags, EEndReplicationFlags::AllowAsyncStopReplication))
+	{
+		if (!Flags.IsEmpty()) { Flags += TEXT(','); };
+		Flags += TEXT("AllowAsyncStopReplication");
+	}
 
 	return Flags;
 }
diff --git a/Engine/Source/Runtime/Net/Iris/Private/Iris/ReplicationSystem/ReplicationSystem.cpp b/Engine/Source/Runtime/Net/Iris/Private/Iris/ReplicationSystem/ReplicationSystem.cpp
--- a/Engine/Source/Runtime/Net/Iris/Private/Iris/ReplicationSystem/ReplicationSystem.cpp
+++ b/Engine/Source/Runtime/Net/Iris/Private/Iris/ReplicationSystem/ReplicationSystem.cpp
@@ -78,6 +78,7 @@ public:
 	FReplicationSystemDelegates Delegates;
 	uint64 IrisDebugHelperDummy = 0U;
 	FNetObjectGroupHandle NotReplicatedNetObjectGroupHandle;
+	FNetObjectGroupHandle AsyncStopReplicationNetObjectGroupHandle;
 	FNetObjectGroupHandle NetGroupOwnerNetObjectGroupHandle;
 	FNetObjectGroupHandle NetGroupReplayNetObjectGroupHandle;
 	FNetBitArray ConnectionsPendingPostTickDispatchSend;
@@ -122,6 +123,10 @@ public:
 		NotReplicatedNetObjectGroupHandle = ReplicationSystem->CreateGroup(FName("NotReplicated"));
 		check(NotReplicatedNetObjectGroupHandle.IsNotReplicatedNetObjectGroup());
 		ReplicationSystem->AddExclusionFilterGroup(NotReplicatedNetObjectGroupHandle);
+
+		AsyncStopReplicationNetObjectGroupHandle = ReplicationSystem->CreateGroup(FName("AsyncStopReplication"));
+		check(AsyncStopReplicationNetObjectGroupHandle.IsAsyncStopReplicationNetObjectGroup());
+		ReplicationSystem->AddExclusionFilterGroup(AsyncStopReplicationNetObjectGroupHandle);
 		
 		// Setup SubObjectFiltering groups
 		NetGroupOwnerNetObjectGroupHandle = ReplicationSystem->GetOrCreateSubObjectFilter(UE::Net::NetGroupOwner);
@@ -1856,6 +1861,11 @@ UE::Net::FNetObjectGroupHandle UReplicationSystem::GetNotReplicatedNetObjectGrou
 	return Impl->NotReplicatedNetObjectGroupHandle;
 }
 
+UE::Net::FNetObjectGroupHandle UReplicationSystem::GetAsyncStopReplicationNetObjectGroup() const
+{
+	return Impl->AsyncStopReplicationNetObjectGroupHandle;
+}
+
 UE::Net::FNetObjectGroupHandle UReplicationSystem::GetNetGroupOwnerNetObjectGroup() const
 {
 	return Impl->NetGroupOwnerNetObjectGroupHandle;
diff --git a/Engine/Source/Runtime/Net/Iris/Public/Iris/ReplicationSystem/NetObjectGroupHandle.h b/Engine/Source/Runtime/Net/Iris/Public/Iris/ReplicationSystem/NetObjectGroupHandle.h
--- a/Engine/Source/Runtime/Net/Iris/Public/Iris/ReplicationSystem/NetObjectGroupHandle.h
+++ b/Engine/Source/Runtime/Net/Iris/Public/Iris/ReplicationSystem/NetObjectGroupHandle.h
@@ -33,6 +33,7 @@ public:
 	{
 		InvalidNetObjectGroupIndex = 0,
 		NotReplicatedNetObjectGroupIndex,
+		AsyncStopReplicationNetObjectGroupIndex,
 		NetGroupOwnerNetObjectGroupIndex,
 		NetGroupReplayNetObjectGroupIndex
 	};
@@ -60,6 +61,9 @@ public:
 	/** Special group, NetHandles assigned to this group will be filtered out for all connections */
 	bool IsNotReplicatedNetObjectGroup() const { return Index == NotReplicatedNetObjectGroupIndex; }
 
+	/** Special group, NetHandles assigned to this group will be filtered out for all connections used for async end replication */
+	bool IsAsyncStopReplicationNetObjectGroup() const { return Index == AsyncStopReplicationNetObjectGroupIndex; }
+
 	/** Special group, SubObjects assigned to this group will replicate to owner of RootParent */
 	bool IsNetGroupOwnerNetObjectGroup() const { return Index == NetGroupOwnerNetObjectGroupIndex; }
 
diff --git a/Engine/Source/Runtime/Net/Iris/Public/Iris/ReplicationSystem/ObjectReplicationBridge.h b/Engine/Source/Runtime/Net/Iris/Public/Iris/ReplicationSystem/ObjectReplicationBridge.h
--- a/Engine/Source/Runtime/Net/Iris/Public/Iris/ReplicationSystem/ObjectReplicationBridge.h
+++ b/Engine/Source/Runtime/Net/Iris/Public/Iris/ReplicationSystem/ObjectReplicationBridge.h
@@ -352,12 +352,8 @@ protected:
 	IRISCORE_API virtual void DetachInstance(FNetRefHandle Handle) override;
 	IRISCORE_API virtual void OnProtocolMismatchReported(FNetRefHandle RefHandle, uint32 ConnectionId) override;
 	IRISCORE_API virtual void OnErrorWithNetRefHandleReported(UE::Net::ENetRefHandleError ErrorType, FNetRefHandle RefHandle, uint32 ConnectionId, const TArray<FNetRefHandle>& ExtraHandles) override;
-
-	/**
-	 * Remove destruction infos associated with group
-	 * Passing in an invalid group handle indicates that we should remove all destruction infos
-	 */
 	IRISCORE_API virtual void RemoveDestructionInfosForGroup(UE::Net::FNetObjectGroupHandle GroupHandle) override;
+	IRISCORE_API virtual void RemoveDestructionInfo(UE::Net::Private::FInternalNetRefIndex InternalObjectIndex) override;
 
 private:
 
diff --git a/Engine/Source/Runtime/Net/Iris/Public/Iris/ReplicationSystem/ReplicationBridge.h b/Engine/Source/Runtime/Net/Iris/Public/Iris/ReplicationSystem/ReplicationBridge.h
--- a/Engine/Source/Runtime/Net/Iris/Public/Iris/ReplicationSystem/ReplicationBridge.h
+++ b/Engine/Source/Runtime/Net/Iris/Public/Iris/ReplicationSystem/ReplicationBridge.h
@@ -150,6 +150,9 @@ public:
 	/** Print common information about this handle and the object it is mapped to */
 	[[nodiscard]] IRISCORE_API FString PrintObjectFromNetRefHandle(FNetRefHandle RefHandle) const;
 
+	/** Returns true if object is PendingStopReplication or PendingAsyncStopReplication. */
+	IRISCORE_API bool IsPendingStopReplication(FNetRefHandle Handle) const;
+
 protected:
 	/** Initializes the bridge. Is called during ReplicationSystem initialization. */
 	IRISCORE_API virtual void Initialize(UReplicationSystem* InReplicationSystem);
@@ -224,6 +227,14 @@ protected:
 	/** Add SubObjectHandle as SubObject to OwnerHandle. */
 	IRISCORE_API void InternalAddSubObject(FNetRefHandle OwnerHandle, FNetRefHandle SubObjectHandle, FNetRefHandle InsertRelativeToSubObjectHandle, UE::Net::ESubObjectInsertionOrder InsertionOrder);
 
+	/** 
+	 * Restart/cancel pending async stop replication.
+     * Returns true if the handle was pending async stop replication and was restored to fully replicated status.
+	 * NOTE: Always returns true for Handles not marked for stop replication.
+	 * Passing in bForceStop as true will issue the pending async stop replication immediately and return false.
+	 */
+	IRISCORE_API bool RestartOrStopObjectPendingAsyncStopReplication(FNetRefHandle Handle, bool bForceStop);
+
 	inline UE::Net::Private::FReplicationProtocolManager* GetReplicationProtocolManager() const { return ReplicationProtocolManager; }
 	inline UReplicationSystem* GetReplicationSystem() const { return ReplicationSystem; }
 	inline UE::Net::Private::FReplicationStateDescriptorRegistry* GetReplicationStateDescriptorRegistry() const { return ReplicationStateDescriptorRegistry; }
@@ -256,6 +267,9 @@ protected:
 	 */
 	virtual void RemoveDestructionInfosForGroup(UE::Net::FNetObjectGroupHandle GroupHandle) {}
 
+	/** Remove destruction infos associated with a specific NetRefHandle */
+	virtual void RemoveDestructionInfo(UE::Net::Private::FInternalNetRefIndex InternalObjectIndex) {};
+
 private:
 
 	// Internal operations invoked by ReplicationSystem/ReplicationWriter
@@ -292,6 +306,9 @@ private:
 	// Destroy all SubObjects owned by provided handle
 	void InternalDestroySubObjects(FNetRefHandle OwnerHandle, EEndReplicationFlags Flags);
 
+	// Internal method to stop replication by filtering out net oobject and once it is no longer replicating to any client it will be destroyed.
+	void InternalRequestAsyncStopReplication(FNetRefHandle OwnerHandle, EEndReplicationFlags Flags);
+
 	/**
 	 * Called from ReplicationSystem when a streaming container is about to unload.
 	 * Will remove the group associated with the container and remove destruction infos.
diff --git a/Engine/Source/Runtime/Net/Iris/Public/Iris/ReplicationSystem/ReplicationBridgeTypes.h b/Engine/Source/Runtime/Net/Iris/Public/Iris/ReplicationSystem/ReplicationBridgeTypes.h
--- a/Engine/Source/Runtime/Net/Iris/Public/Iris/ReplicationSystem/ReplicationBridgeTypes.h
+++ b/Engine/Source/Runtime/Net/Iris/Public/Iris/ReplicationSystem/ReplicationBridgeTypes.h
@@ -25,6 +25,8 @@ enum class EEndReplicationFlags : uint32
 	SkipPendingEndReplicationValidation = ClearNetPushId << 1U,
 	/** An actor is migrated between servers. Only supported by clients with replication systems configured to use UE::Net::EProxyType::Backend. */
 	Migrated							= SkipPendingEndReplicationValidation << 1U,
+	/** If true, replication can be resumed effectively canceling the EndReplication by issuing a StartReplication before the StopReplication has been acknowledged by all clients. */
+	AllowAsyncStopReplication = Migrated << 1U,
 };
 ENUM_CLASS_FLAGS(EEndReplicationFlags);
 
diff --git a/Engine/Source/Runtime/Net/Iris/Public/Iris/ReplicationSystem/ReplicationSystem.h b/Engine/Source/Runtime/Net/Iris/Public/Iris/ReplicationSystem/ReplicationSystem.h
--- a/Engine/Source/Runtime/Net/Iris/Public/Iris/ReplicationSystem/ReplicationSystem.h
+++ b/Engine/Source/Runtime/Net/Iris/Public/Iris/ReplicationSystem/ReplicationSystem.h
@@ -541,6 +541,9 @@ public:
 	/** Special group, root objects assigned to this group will be filtered out for all connections */
 	IRISCORE_API FNetObjectGroupHandle GetNotReplicatedNetObjectGroup() const;
 
+	/** Special group, used for async stop replication */
+	IRISCORE_API FNetObjectGroupHandle GetAsyncStopReplicationNetObjectGroup() const;
+
 	/** Special group, SubObjects assigned to this group will replicate to owner of RootParent */
 	IRISCORE_API FNetObjectGroupHandle GetNetGroupOwnerNetObjectGroup() const;
 
