Add an outer reference to FReplicationTestCreationHeader to make sure it's exported, needed for proxy tests so that the proxy has a cached reference to the outer when it re-replicates an object.

#Synced-CL 49123499

diff --git a/Engine/Plugins/Runtime/ReplicationSystemTestPlugin/Source/Private/Tests/ReplicationSystem/ReplicatedTestObjectFactory.cpp b/Engine/Plugins/Runtime/ReplicationSystemTestPlugin/Source/Private/Tests/ReplicationSystem/ReplicatedTestObjectFactory.cpp
--- a/Engine/Plugins/Runtime/ReplicationSystemTestPlugin/Source/Private/Tests/ReplicationSystem/ReplicatedTestObjectFactory.cpp
+++ b/Engine/Plugins/Runtime/ReplicationSystemTestPlugin/Source/Private/Tests/ReplicationSystem/ReplicatedTestObjectFactory.cpp
@@ -7,6 +7,7 @@
 #include "Iris/Serialization/NetSerializationContext.h"
 #include "Iris/Serialization/NetBitStreamUtil.h"
 #include "Iris/ReplicationSystem/NetObjectFactory.h"
+#include "Iris/Serialization/ObjectNetSerializer.h"
 
 namespace UE::Net
 {
@@ -22,28 +23,30 @@ FString FReplicationTestCreationHeader::ToString() const
 		return FString::Printf(TEXT("FReplicationTestCreationHeader (ProtocolId:0x%x)\n\t"
 									"StableName=%s"
 									"bForceFailCreationRemoteInstance=%u"),
-						   			GetProtocolId(),
-						   			*ArchetypeOrStableName,
-						   			bForceFailCreateRemoteInstance);	
+									GetProtocolId(),
+									*ArchetypeOrStableName,
+									bForceFailCreateRemoteInstance);
 	}
 	else
 	{
 		return FString::Printf(TEXT("FReplicationTestCreationHeader (ProtocolId:0x%x)\n\t"
 									"ArchetypeName=%s"
+									"OuterReference=%s"
 									"NumComponentsToSpawn=%u"
 									"NumIrisComponentsToSpawn=%u"
 									"NumDynamicComponentsToSpawn=%u"
 									"NumConnectionFilteredComponentsToSpawn=%u"
 									"NumObjectReferenceComponentsToSpawn=%u"
 									"bForceFailCreationRemoteInstance=%u"),
-						   			GetProtocolId(),
-						   			*ArchetypeOrStableName,
-						   			NumComponentsToSpawn,
-						   			NumIrisComponentsToSpawn,
-						   			NumDynamicComponentsToSpawn,
-						   			NumConnectionFilteredComponentsToSpawn,
-						   			NumObjectReferenceComponentsToSpawn,
-						   			bForceFailCreateRemoteInstance);
+									GetProtocolId(),
+									*ArchetypeOrStableName,
+									*DynamicObjectOuterReference.ToString(),
+									NumComponentsToSpawn,
+									NumIrisComponentsToSpawn,
+									NumDynamicComponentsToSpawn,
+									NumConnectionFilteredComponentsToSpawn,
+									NumObjectReferenceComponentsToSpawn,
+									bForceFailCreateRemoteInstance);
 	}
 }
 
@@ -55,6 +58,7 @@ bool FReplicationTestCreationHeader::Serialize(const FCreationHeaderContext& Con
 	Writer.WriteBool(bForceFailCreateRemoteInstance);
 	if (!Writer.WriteBool(bIsStable))
 	{
+		WriteFullNetObjectReference(Context.Serialization, DynamicObjectOuterReference);
 		Writer.WriteBits(NumComponentsToSpawn, 16);
 		Writer.WriteBits(NumIrisComponentsToSpawn, 16);
 		Writer.WriteBits(NumDynamicComponentsToSpawn, 16);
@@ -74,6 +78,7 @@ bool FReplicationTestCreationHeader::Deserialize(const FCreationHeaderContext& C
 	this->bIsStable = Reader.ReadBool();
 	if (!this->bIsStable)
 	{
+		ReadFullNetObjectReference(Context.Serialization, DynamicObjectOuterReference);
 		this->NumComponentsToSpawn = Reader.ReadBits(16);
 		this->NumIrisComponentsToSpawn = Reader.ReadBits(16);
 		this->NumDynamicComponentsToSpawn = Reader.ReadBits(16);
@@ -149,6 +154,11 @@ TUniquePtr<UE::Net::FNetObjectCreationHeader> UReplicatedTestObjectFactory::Crea
 
 	Header->ArchetypeOrStableName = Archetype->GetPathName();
 
+	if (const UObject* ObjectOuter = Object->GetOuter())
+	{
+		Header->DynamicObjectOuterReference = Bridge->GetOrCreateObjectReference(ObjectOuter);
+	}
+
 	return HeaderRef;
 }
 
@@ -205,6 +215,7 @@ UNetObjectFactory::FInstantiateResult UReplicatedTestObjectFactory::InstantiateR
 		check(ArcheType);
 
 		FStaticConstructObjectParameters ConstructObjectParameters(ArcheType->GetClass());
+		ConstructObjectParameters.Outer = Bridge->ResolveObjectReference(Header->DynamicObjectOuterReference, Context.ResolveContext);
 		UObject* CreatedObject = StaticConstructObject_Internal(ConstructObjectParameters);
 
 		const bool bIsSubObject = Context.RootObjectOfSubObject.IsValid();
diff --git a/Engine/Plugins/Runtime/ReplicationSystemTestPlugin/Source/Private/Tests/ReplicationSystem/ReplicatedTestObjectFactory.h b/Engine/Plugins/Runtime/ReplicationSystemTestPlugin/Source/Private/Tests/ReplicationSystem/ReplicatedTestObjectFactory.h
--- a/Engine/Plugins/Runtime/ReplicationSystemTestPlugin/Source/Private/Tests/ReplicationSystem/ReplicatedTestObjectFactory.h
+++ b/Engine/Plugins/Runtime/ReplicationSystemTestPlugin/Source/Private/Tests/ReplicationSystem/ReplicatedTestObjectFactory.h
@@ -3,6 +3,7 @@
 #pragma once
 
 #include "Iris/ReplicationSystem/NetObjectFactory.h"
+#include "Iris/Core/NetObjectReference.h"
 
 #include "ReplicatedTestObjectFactory.generated.h"
 
@@ -13,6 +14,7 @@ class FReplicationTestCreationHeader : public FNetObjectCreationHeader
 {
 public:
 	FString ArchetypeOrStableName;
+	FNetObjectReference DynamicObjectOuterReference;
 	uint32 NumComponentsToSpawn = 0;
 	uint32 NumIrisComponentsToSpawn = 0;
 	uint32 NumDynamicComponentsToSpawn = 0;
diff --git a/Engine/Plugins/Runtime/ReplicationSystemTestPlugin/Source/Private/Tests/Serialization/TestObjectNetSerializer.cpp b/Engine/Plugins/Runtime/ReplicationSystemTestPlugin/Source/Private/Tests/Serialization/TestObjectNetSerializer.cpp
--- a/Engine/Plugins/Runtime/ReplicationSystemTestPlugin/Source/Private/Tests/Serialization/TestObjectNetSerializer.cpp
+++ b/Engine/Plugins/Runtime/ReplicationSystemTestPlugin/Source/Private/Tests/Serialization/TestObjectNetSerializer.cpp
@@ -102,7 +102,7 @@ UE_NET_TEST_FIXTURE(FReplicationSystemServerClientTestFixture, TestCircularRefer
 UE_NET_TEST_FIXTURE(FReplicationSystemServerClientTestFixture, TestCircularReferenceWithLimitedThroughput)
 {
 	// Limit packet size
-	Server->SetMaxSendPacketSize(128U);
+	Server->SetMaxSendPacketSize(136U);
 
 	// Add a client
 	FReplicationSystemTestClient* Client = CreateClient();
