Track which objects have been processed while updating objects with creation dependencies so we don't reprocess objects.

#Synced-CL 49156055

diff --git a/Engine/Plugins/Runtime/ReplicationSystemTestPlugin/Source/Private/Tests/ReplicationSystem/TestDependentObjects.cpp b/Engine/Plugins/Runtime/ReplicationSystemTestPlugin/Source/Private/Tests/ReplicationSystem/TestDependentObjects.cpp
--- a/Engine/Plugins/Runtime/ReplicationSystemTestPlugin/Source/Private/Tests/ReplicationSystem/TestDependentObjects.cpp
+++ b/Engine/Plugins/Runtime/ReplicationSystemTestPlugin/Source/Private/Tests/ReplicationSystem/TestDependentObjects.cpp
@@ -1082,6 +1082,336 @@ UE_NET_TEST_FIXTURE(FReplicationSystemServerClientTestFixture, TestNestedCreatio
 	}
 }
 
+UE_NET_TEST_FIXTURE(FTestFilteredDependentObjectFixture, TestBasicCreationDependencyFiltering)
+{
+	UReplicationSystem* ReplicationSystem = Server->ReplicationSystem;
+	UReplicatedTestObjectBridge* Bridge = Server->GetReplicationBridge();
+
+	// Add a client
+	FReplicationSystemTestClient* Client = CreateClient();
+
+	// Spawn server objects
+	UReplicatedTestObject* ParentObject = Server->CreateObject();
+	UReplicatedTestObject* ChildObject = Server->CreateObject();
+	Bridge->AddCreationDependencyLink(ParentObject->NetRefHandle, ChildObject->NetRefHandle);
+
+	// Send and deliver packets
+	Server->UpdateAndSend({ Client });
+
+	// Objects should now exist on client
+	UE_NET_ASSERT_TRUE(Client->IsResolvableNetRefHandle(ParentObject->NetRefHandle));
+	UE_NET_ASSERT_TRUE(Client->IsResolvableNetRefHandle(ChildObject->NetRefHandle));
+
+	// Filter out parent
+	ReplicationSystem->SetFilter(ParentObject->NetRefHandle, NotRoutedFilterHandle);
+
+	// Send and deliver packets
+	Server->UpdateAndSend({ Client });
+
+	// Objects should not exist on client
+	UE_NET_ASSERT_FALSE(Client->IsResolvableNetRefHandle(ParentObject->NetRefHandle));
+	UE_NET_ASSERT_FALSE(Client->IsResolvableNetRefHandle(ChildObject->NetRefHandle));
+}
+
+UE_NET_TEST_FIXTURE(FTestFilteredDependentObjectFixture, TestNestedCreationDependencyFiltering)
+{
+	UReplicationSystem* ReplicationSystem = Server->ReplicationSystem;
+	UReplicatedTestObjectBridge* Bridge = Server->GetReplicationBridge();
+
+	// Add a client
+	FReplicationSystemTestClient* Client = CreateClient();
+
+	// Spawn root object
+	UReplicatedTestObject* RootObject = Server->CreateObject();
+	TArray<UReplicatedTestObject*> Objects;
+	Objects.Add(RootObject);
+
+	// Spawn chain of dependent objects
+	constexpr uint32 NumObjects = 16;
+	for (uint32 ObjectIndex = 1U; ObjectIndex < NumObjects; ++ObjectIndex)
+	{
+		UReplicatedTestObject* ChildObject = Server->CreateObject();
+		Bridge->AddCreationDependencyLink(Objects.Last()->NetRefHandle, ChildObject->NetRefHandle);
+		Objects.Add(ChildObject);
+	}
+
+	// Send and deliver packets
+	Server->UpdateAndSend({ Client });
+
+	// Objects should all exist on client
+	for (UReplicatedTestObject* Object : Objects)
+	{
+		UE_NET_ASSERT_TRUE(Client->IsResolvableNetRefHandle(Object->NetRefHandle));
+	}
+
+	// Filter out root object
+	ReplicationSystem->SetFilter(Objects[0]->NetRefHandle, NotRoutedFilterHandle);
+
+	// Send and deliver packets
+	Server->UpdateAndSend({ Client });
+
+	// Objects should not exist on client
+	for (UReplicatedTestObject* Object : Objects)
+	{
+		UE_NET_ASSERT_FALSE(Client->IsResolvableNetRefHandle(Object->NetRefHandle));
+	}
+}
+
+// Tests creation dependencies with multiple parents, multiple children, and parents created first.
+UE_NET_TEST_FIXTURE(FTestFilteredDependentObjectFixture, TestComplexCreationDependencyFiltering)
+{
+	UReplicationSystem* ReplicationSystem = Server->ReplicationSystem;
+	UReplicatedTestObjectBridge* Bridge = Server->GetReplicationBridge();
+
+	// Add a client
+	FReplicationSystemTestClient* Client = CreateClient();
+
+	// There are 3 layers of objects in this creation dependency graph.
+	// 1st layer is TopObjects[4]. They don't depend on anything.
+	// 2nd layer is MiddleObjects[4]. MiddleObjects[i] depends on TopObjects[i] and TopObjects[i+1] if it exists.
+	// 3rd layer is BottomObjects[4]. BottomObjects[i] depends on MiddleObjects[i] and MiddleObjects[i+1] if it exists.
+
+	// Create parents first, then children
+	constexpr uint32 ArrayObjectCount = 4;
+
+	// Create TopObjects[4]
+	UReplicatedTestObject* TopObjects[ArrayObjectCount];
+	for (UReplicatedTestObject*& Object : TopObjects)
+	{
+		Object = Server->CreateObject();
+	}
+
+	// Create MiddleObjects[4]
+	UReplicatedTestObject* MiddleObjects[ArrayObjectCount];
+	for (UReplicatedTestObject*& Object : MiddleObjects)
+	{
+		Object = Server->CreateObject();
+	}
+
+	// Create BottomObjects[4]
+	UReplicatedTestObject* BottomObjects[ArrayObjectCount];
+	for (UReplicatedTestObject*& Object : BottomObjects)
+	{
+		Object = Server->CreateObject();
+	}
+
+	// Add creation dependency links
+	for (uint32 ObjectIndex = 0; ObjectIndex < ArrayObjectCount - 1; ++ObjectIndex)
+	{
+		Bridge->AddCreationDependencyLink(TopObjects[ObjectIndex]->NetRefHandle, MiddleObjects[ObjectIndex]->NetRefHandle);
+		Bridge->AddCreationDependencyLink(TopObjects[ObjectIndex + 1]->NetRefHandle, MiddleObjects[ObjectIndex]->NetRefHandle);
+		Bridge->AddCreationDependencyLink(MiddleObjects[ObjectIndex]->NetRefHandle, BottomObjects[ObjectIndex]->NetRefHandle);
+		Bridge->AddCreationDependencyLink(MiddleObjects[ObjectIndex + 1]->NetRefHandle, BottomObjects[ObjectIndex]->NetRefHandle);
+	}
+	Bridge->AddCreationDependencyLink(TopObjects[ArrayObjectCount - 1]->NetRefHandle, MiddleObjects[ArrayObjectCount - 1]->NetRefHandle);
+	Bridge->AddCreationDependencyLink(MiddleObjects[ArrayObjectCount - 1]->NetRefHandle, BottomObjects[ArrayObjectCount - 1]->NetRefHandle);
+
+	// Send and deliver packets
+	Server->UpdateAndSend({ Client });
+
+	// Objects should all exist on client
+	for (uint32 ObjectIndex = 0; ObjectIndex < ArrayObjectCount; ++ObjectIndex)
+	{
+		UE_NET_ASSERT_TRUE(Client->IsResolvableNetRefHandle(TopObjects[ObjectIndex]->NetRefHandle));
+		UE_NET_ASSERT_TRUE(Client->IsResolvableNetRefHandle(MiddleObjects[ObjectIndex]->NetRefHandle));
+		UE_NET_ASSERT_TRUE(Client->IsResolvableNetRefHandle(BottomObjects[ObjectIndex]->NetRefHandle));
+	}
+
+	// Filter out TopObjects[2]
+	ReplicationSystem->SetFilter(TopObjects[2]->NetRefHandle, NotRoutedFilterHandle);
+
+	// Send and deliver packets
+	Server->UpdateAndSend({ Client });
+
+	// TopObjects[2] should not exist on the client
+	UE_NET_ASSERT_TRUE(Client->IsResolvableNetRefHandle(TopObjects[0]->NetRefHandle));
+	UE_NET_ASSERT_TRUE(Client->IsResolvableNetRefHandle(TopObjects[1]->NetRefHandle));
+	UE_NET_ASSERT_FALSE(Client->IsResolvableNetRefHandle(TopObjects[2]->NetRefHandle));
+	UE_NET_ASSERT_TRUE(Client->IsResolvableNetRefHandle(TopObjects[3]->NetRefHandle));
+
+	// MiddleObjects[1-2] should not exist on the client
+	UE_NET_ASSERT_TRUE(Client->IsResolvableNetRefHandle(MiddleObjects[0]->NetRefHandle));
+	UE_NET_ASSERT_FALSE(Client->IsResolvableNetRefHandle(MiddleObjects[1]->NetRefHandle));
+	UE_NET_ASSERT_FALSE(Client->IsResolvableNetRefHandle(MiddleObjects[2]->NetRefHandle));
+	UE_NET_ASSERT_TRUE(Client->IsResolvableNetRefHandle(MiddleObjects[3]->NetRefHandle));
+
+	// BottomObjects[0-2] should not exist on the client
+	UE_NET_ASSERT_FALSE(Client->IsResolvableNetRefHandle(BottomObjects[0]->NetRefHandle));
+	UE_NET_ASSERT_FALSE(Client->IsResolvableNetRefHandle(BottomObjects[1]->NetRefHandle));
+	UE_NET_ASSERT_FALSE(Client->IsResolvableNetRefHandle(BottomObjects[2]->NetRefHandle));
+	UE_NET_ASSERT_TRUE(Client->IsResolvableNetRefHandle(BottomObjects[3]->NetRefHandle));
+}
+
+// Tests creation dependencies with multiple parents, multiple children, and children created first.
+UE_NET_TEST_FIXTURE(FTestFilteredDependentObjectFixture, TestComplexCreationDependencyFilteringReversed)
+{
+	UReplicationSystem* ReplicationSystem = Server->ReplicationSystem;
+	UReplicatedTestObjectBridge* Bridge = Server->GetReplicationBridge();
+
+	// Add a client
+	FReplicationSystemTestClient* Client = CreateClient();
+
+	// There are 3 layers of objects in this creation dependency graph.
+	// 1st layer is TopObjects[4]. They don't depend on anything.
+	// 2nd layer is MiddleObjects[4]. MiddleObjects[i] depends on TopObjects[i] and TopObjects[i+1] if it exists.
+	// 3rd layer is BottomObjects[4]. BottomObjects[i] depends on MiddleObjects[i] and MiddleObjects[i+1] if it exists.
+
+	// Create children first, then parents
+	constexpr uint32 ArrayObjectCount = 4;
+
+	// Create BottomObjects[4]
+	UReplicatedTestObject* BottomObjects[ArrayObjectCount];
+	for (UReplicatedTestObject*& Object : BottomObjects)
+	{
+		Object = Server->CreateObject();
+	}
+
+	// Create MiddleObjects[4]
+	UReplicatedTestObject* MiddleObjects[ArrayObjectCount];
+	for (UReplicatedTestObject*& Object : MiddleObjects)
+	{
+		Object = Server->CreateObject();
+	}
+
+	// Create TopObjects[4]
+	UReplicatedTestObject* TopObjects[ArrayObjectCount];
+	for (UReplicatedTestObject*& Object : TopObjects)
+	{
+		Object = Server->CreateObject();
+	}
+
+	// Add creation dependency links
+	for (uint32 ObjectIndex = 0U; ObjectIndex < ArrayObjectCount - 1; ++ObjectIndex)
+	{
+		Bridge->AddCreationDependencyLink(TopObjects[ObjectIndex]->NetRefHandle, MiddleObjects[ObjectIndex]->NetRefHandle);
+		Bridge->AddCreationDependencyLink(TopObjects[ObjectIndex + 1]->NetRefHandle, MiddleObjects[ObjectIndex]->NetRefHandle);
+		Bridge->AddCreationDependencyLink(MiddleObjects[ObjectIndex]->NetRefHandle, BottomObjects[ObjectIndex]->NetRefHandle);
+		Bridge->AddCreationDependencyLink(MiddleObjects[ObjectIndex + 1]->NetRefHandle, BottomObjects[ObjectIndex]->NetRefHandle);
+	}
+	Bridge->AddCreationDependencyLink(TopObjects[ArrayObjectCount - 1]->NetRefHandle, MiddleObjects[ArrayObjectCount - 1]->NetRefHandle);
+	Bridge->AddCreationDependencyLink(MiddleObjects[ArrayObjectCount - 1]->NetRefHandle, BottomObjects[ArrayObjectCount - 1]->NetRefHandle);
+
+	// Send and deliver packets
+	Server->UpdateAndSend({ Client });
+
+	// Objects should all exist on client
+	for (uint32 ObjectIndex = 0U; ObjectIndex < ArrayObjectCount; ++ObjectIndex)
+	{
+		UE_NET_ASSERT_TRUE(Client->IsResolvableNetRefHandle(TopObjects[ObjectIndex]->NetRefHandle));
+		UE_NET_ASSERT_TRUE(Client->IsResolvableNetRefHandle(MiddleObjects[ObjectIndex]->NetRefHandle));
+		UE_NET_ASSERT_TRUE(Client->IsResolvableNetRefHandle(BottomObjects[ObjectIndex]->NetRefHandle));
+	}
+
+	// Filter out TopObjects[2]
+	ReplicationSystem->SetFilter(TopObjects[2]->NetRefHandle, NotRoutedFilterHandle);
+
+	// Send and deliver packets
+	Server->UpdateAndSend({ Client });
+
+	// TopObjects[2] should not exist on the client
+	UE_NET_ASSERT_TRUE(Client->IsResolvableNetRefHandle(TopObjects[0]->NetRefHandle));
+	UE_NET_ASSERT_TRUE(Client->IsResolvableNetRefHandle(TopObjects[1]->NetRefHandle));
+	UE_NET_ASSERT_FALSE(Client->IsResolvableNetRefHandle(TopObjects[2]->NetRefHandle));
+	UE_NET_ASSERT_TRUE(Client->IsResolvableNetRefHandle(TopObjects[3]->NetRefHandle));
+
+	// MiddleObjects[1-2] should not exist on the client
+	UE_NET_ASSERT_TRUE(Client->IsResolvableNetRefHandle(MiddleObjects[0]->NetRefHandle));
+	UE_NET_ASSERT_FALSE(Client->IsResolvableNetRefHandle(MiddleObjects[1]->NetRefHandle));
+	UE_NET_ASSERT_FALSE(Client->IsResolvableNetRefHandle(MiddleObjects[2]->NetRefHandle));
+	UE_NET_ASSERT_TRUE(Client->IsResolvableNetRefHandle(MiddleObjects[3]->NetRefHandle));
+
+	// BottomObjects[0-2] should not exist on the client
+	UE_NET_ASSERT_FALSE(Client->IsResolvableNetRefHandle(BottomObjects[0]->NetRefHandle));
+	UE_NET_ASSERT_FALSE(Client->IsResolvableNetRefHandle(BottomObjects[1]->NetRefHandle));
+	UE_NET_ASSERT_FALSE(Client->IsResolvableNetRefHandle(BottomObjects[2]->NetRefHandle));
+	UE_NET_ASSERT_TRUE(Client->IsResolvableNetRefHandle(BottomObjects[3]->NetRefHandle));
+}
+
+// Tests creation dependencies with multiple parents, multiple children, and mixed parent/child object creation order.
+UE_NET_TEST_FIXTURE(FTestFilteredDependentObjectFixture, TestComplexCreationDependencyFilteringMixedOrder)
+{
+	UReplicationSystem* ReplicationSystem = Server->ReplicationSystem;
+	UReplicatedTestObjectBridge* Bridge = Server->GetReplicationBridge();
+
+	// Add a client
+	FReplicationSystemTestClient* Client = CreateClient();
+
+	// There are 3 layers of objects in this creation dependency graph.
+	// 1st layer is TopObjects[4]. They don't depend on anything.
+	// 2nd layer is MiddleObjects[4]. MiddleObjects[i] depends on TopObjects[i] and TopObjects[i+1] if it exists.
+	// 3rd layer is BottomObjects[4]. BottomObjects[i] depends on MiddleObjects[i] and MiddleObjects[i+1] if it exists.
+
+	// Create MiddleObjects first to simulate creating some children before their parents and some parents before their children
+	constexpr uint32 ArrayObjectCount = 4;
+
+	// Create MiddleObjects[4]
+	UReplicatedTestObject* MiddleObjects[ArrayObjectCount];
+	for (UReplicatedTestObject*& Object : MiddleObjects)
+	{
+		Object = Server->CreateObject();
+	}
+
+	// Create TopObjects[4]
+	UReplicatedTestObject* TopObjects[ArrayObjectCount];
+	for (UReplicatedTestObject*& Object : TopObjects)
+	{
+		Object = Server->CreateObject();
+	}
+
+	// Create BottomObjects[4]
+	UReplicatedTestObject* BottomObjects[ArrayObjectCount];
+	for (UReplicatedTestObject*& Object : BottomObjects)
+	{
+		Object = Server->CreateObject();
+	}
+
+	// Add creation dependency links
+	for (uint32 ObjectIndex = 0; ObjectIndex < ArrayObjectCount - 1; ++ObjectIndex)
+	{
+		Bridge->AddCreationDependencyLink(TopObjects[ObjectIndex]->NetRefHandle, MiddleObjects[ObjectIndex]->NetRefHandle);
+		Bridge->AddCreationDependencyLink(TopObjects[ObjectIndex + 1]->NetRefHandle, MiddleObjects[ObjectIndex]->NetRefHandle);
+		Bridge->AddCreationDependencyLink(MiddleObjects[ObjectIndex]->NetRefHandle, BottomObjects[ObjectIndex]->NetRefHandle);
+		Bridge->AddCreationDependencyLink(MiddleObjects[ObjectIndex + 1]->NetRefHandle, BottomObjects[ObjectIndex]->NetRefHandle);
+	}
+	Bridge->AddCreationDependencyLink(TopObjects[ArrayObjectCount - 1]->NetRefHandle, MiddleObjects[ArrayObjectCount - 1]->NetRefHandle);
+	Bridge->AddCreationDependencyLink(MiddleObjects[ArrayObjectCount - 1]->NetRefHandle, BottomObjects[ArrayObjectCount - 1]->NetRefHandle);
+
+	// Send and deliver packets
+	Server->UpdateAndSend({ Client });
+
+	// Objects should all exist on client
+	for (uint32 ObjectIndex = 0U; ObjectIndex < ArrayObjectCount; ++ObjectIndex)
+	{
+		UE_NET_ASSERT_TRUE(Client->IsResolvableNetRefHandle(TopObjects[ObjectIndex]->NetRefHandle));
+		UE_NET_ASSERT_TRUE(Client->IsResolvableNetRefHandle(MiddleObjects[ObjectIndex]->NetRefHandle));
+		UE_NET_ASSERT_TRUE(Client->IsResolvableNetRefHandle(BottomObjects[ObjectIndex]->NetRefHandle));
+	}
+
+	// Filter out TopObjects[2]
+	ReplicationSystem->SetFilter(TopObjects[2]->NetRefHandle, NotRoutedFilterHandle);
+
+	// Send and deliver packets
+	Server->UpdateAndSend({ Client });
+
+	// TopObjects[2] should not exist on the client
+	UE_NET_ASSERT_TRUE(Client->IsResolvableNetRefHandle(TopObjects[0]->NetRefHandle));
+	UE_NET_ASSERT_TRUE(Client->IsResolvableNetRefHandle(TopObjects[1]->NetRefHandle));
+	UE_NET_ASSERT_FALSE(Client->IsResolvableNetRefHandle(TopObjects[2]->NetRefHandle));
+	UE_NET_ASSERT_TRUE(Client->IsResolvableNetRefHandle(TopObjects[3]->NetRefHandle));
+
+	// MiddleObjects[1-2] should not exist on the client
+	UE_NET_ASSERT_TRUE(Client->IsResolvableNetRefHandle(MiddleObjects[0]->NetRefHandle));
+	UE_NET_ASSERT_FALSE(Client->IsResolvableNetRefHandle(MiddleObjects[1]->NetRefHandle));
+	UE_NET_ASSERT_FALSE(Client->IsResolvableNetRefHandle(MiddleObjects[2]->NetRefHandle));
+	UE_NET_ASSERT_TRUE(Client->IsResolvableNetRefHandle(MiddleObjects[3]->NetRefHandle));
+
+	// BottomObjects[0-2] should not exist on the client
+	UE_NET_ASSERT_FALSE(Client->IsResolvableNetRefHandle(BottomObjects[0]->NetRefHandle));
+	UE_NET_ASSERT_FALSE(Client->IsResolvableNetRefHandle(BottomObjects[1]->NetRefHandle));
+	UE_NET_ASSERT_FALSE(Client->IsResolvableNetRefHandle(BottomObjects[2]->NetRefHandle));
+	UE_NET_ASSERT_TRUE(Client->IsResolvableNetRefHandle(BottomObjects[3]->NetRefHandle));
+}
+
 UE_NET_TEST_FIXTURE(FTestFilteredDependentObjectFixture, TestCircularDependencyIsDenied)
 {
 	UReplicatedTestObjectBridge* Bridge = Server->GetReplicationBridge();
diff --git a/Engine/Source/Runtime/Net/Iris/Private/Iris/ReplicationSystem/Filtering/ReplicationFiltering.cpp b/Engine/Source/Runtime/Net/Iris/Private/Iris/ReplicationSystem/Filtering/ReplicationFiltering.cpp
--- a/Engine/Source/Runtime/Net/Iris/Private/Iris/ReplicationSystem/Filtering/ReplicationFiltering.cpp
+++ b/Engine/Source/Runtime/Net/Iris/Private/Iris/ReplicationSystem/Filtering/ReplicationFiltering.cpp
@@ -1270,7 +1270,7 @@ void FReplicationFiltering::PreUpdateDynamicFiltering()
 	}
 }
 
-void FReplicationFiltering::UpdateCreationDependentParent(uint32 ChildIndex, const FNetBitArrayView ObjectsWithCreationDependencies,  FNetBitArrayView OutConnectionObjectsInScope, bool bRecursive) const
+void FReplicationFiltering::UpdateCreationDependentParent(uint32 ChildIndex, const FNetBitArrayView ObjectsWithCreationDependencies, FNetBitArrayView OutConnectionObjectsInScope, FNetBitArrayView ProcessedObjects, bool bRecursive) const
 {
 	const bool bStartedRelevant = OutConnectionObjectsInScope.IsBitSet(ChildIndex);
 	if (!bStartedRelevant)
@@ -1279,6 +1279,12 @@ void FReplicationFiltering::UpdateCreationDependentParent(uint32 ChildIndex, con
 		return;
 	}
 
+	const bool bAlreadyProcessed = ProcessedObjects.IsBitSet(ChildIndex);
+	if (bAlreadyProcessed)
+	{
+		return;
+	}
+
 	TConstArrayView<const FInternalNetRefIndex> Parents = NetRefHandleManager->GetCreationDependencies(ChildIndex);
 
 	for (FInternalNetRefIndex ParentIndex : Parents)
@@ -1286,11 +1292,10 @@ void FReplicationFiltering::UpdateCreationDependentParent(uint32 ChildIndex, con
 		const bool bHasCreationDependency = ObjectsWithCreationDependencies.IsBitSet(ParentIndex);
 					
 		// If the parent is a dependent AND relevant check his status first
-		//TODO: Optimize this by flagging which parents were processed already or not
 		if (bHasCreationDependency && OutConnectionObjectsInScope.IsBitSet(ParentIndex))
 		{
 			constexpr bool bIsRecursiveCall = true;
-			UpdateCreationDependentParent(ParentIndex, ObjectsWithCreationDependencies, OutConnectionObjectsInScope, bIsRecursiveCall);
+			UpdateCreationDependentParent(ParentIndex, ObjectsWithCreationDependencies, OutConnectionObjectsInScope, ProcessedObjects, bIsRecursiveCall);
 		}
 
 		const bool bIsParentRelevant = OutConnectionObjectsInScope.IsBitSet(ParentIndex);
@@ -1306,6 +1311,9 @@ void FReplicationFiltering::UpdateCreationDependentParent(uint32 ChildIndex, con
 			OutConnectionObjectsInScope.ClearBit(SubObjectIndex);
 		}
 	}
+
+	// Mark the object as processed so we don't process it again if it's another object's parent.
+	ProcessedObjects.SetBit(ChildIndex);
 };
 
 /**
@@ -1592,12 +1600,17 @@ void FReplicationFiltering::UpdateDynamicFiltering()
 		// Update dependent objects that can only be relevant if their parent also is
 		{
 			IRIS_PROFILER_SCOPE(FReplicationFiltering_UpdateCreationDependencies);
-			FNetBitArrayView ConnectionObjectsInScope = MakeNetBitArrayView(ConnectionInfo.ObjectsInScope);
-			FNetBitArrayView::ForAllSetBits(ObjectsWithCreationDependencies, ConnectionObjectsInScope, FNetBitArrayView::AndOp, [&](uint32 ChildIndex)
+			if (ObjectsWithCreationDependencies.IsAnyBitSet())
 			{
-				constexpr bool bIsNotRecursive = false;
-				this->UpdateCreationDependentParent(ChildIndex, ObjectsWithCreationDependencies, ConnectionObjectsInScope, bIsNotRecursive);
-			});
+				FNetBitArrayView ConnectionObjectsInScope = MakeNetBitArrayView(ConnectionInfo.ObjectsInScope);
+				FNetBitArray ProcessedObjects(ConnectionInfo.ObjectsInScope.GetNumBits());
+				FNetBitArrayView ProcessedObjectsView = MakeNetBitArrayView(ProcessedObjects);
+				FNetBitArrayView::ForAllSetBits(ObjectsWithCreationDependencies, ConnectionObjectsInScope, FNetBitArrayView::AndOp, [ & ](uint32 ChildIndex)
+				{
+					constexpr bool bIsNotRecursive = false;
+					this->UpdateCreationDependentParent(ChildIndex, ObjectsWithCreationDependencies, ConnectionObjectsInScope, ProcessedObjectsView, bIsNotRecursive);
+				});
+			}
 		}
 
 		// The scope for the connection is now fully updated, apart from disabled dependent objects.
diff --git a/Engine/Source/Runtime/Net/Iris/Private/Iris/ReplicationSystem/Filtering/ReplicationFiltering.h b/Engine/Source/Runtime/Net/Iris/Private/Iris/ReplicationSystem/Filtering/ReplicationFiltering.h
--- a/Engine/Source/Runtime/Net/Iris/Private/Iris/ReplicationSystem/Filtering/ReplicationFiltering.h
+++ b/Engine/Source/Runtime/Net/Iris/Private/Iris/ReplicationSystem/Filtering/ReplicationFiltering.h
@@ -238,7 +238,7 @@ private:
 
 	bool HasDynamicFilters() const;
 
-	void UpdateCreationDependentParent(uint32 ChildIndex, const FNetBitArrayView ObjectsWithCreationDependencies, FNetBitArrayView OutConnectionObjectsInScope, bool bIsRecursive) const;
+	void UpdateCreationDependentParent(uint32 ChildIndex, const FNetBitArrayView ObjectsWithCreationDependencies, FNetBitArrayView OutConnectionObjectsInScope, FNetBitArrayView ProcessedObjects, bool bIsRecursive) const;
 
 	// Helper to update and reset group exclusion filter effects if objects are removed from a filter or after a filter status change, returns true if the group filter was changed
 	bool ClearGroupExclusionFilterEffectsForObject(uint32 ObjectIndex, uint32 ConnectionId);
