Iris
* IsRelevantForConnection will now test object status in the ReplicationWriter
* Fixes IsRelevantForConnection returning true when a subobject is filtered out

#rnx
#Synced-CL 49270433

diff --git a/Engine/Plugins/Runtime/ReplicationSystemTestPlugin/Source/Private/Tests/ReplicationSystem/Filtering/TestFiltering.cpp b/Engine/Plugins/Runtime/ReplicationSystemTestPlugin/Source/Private/Tests/ReplicationSystem/Filtering/TestFiltering.cpp
--- a/Engine/Plugins/Runtime/ReplicationSystemTestPlugin/Source/Private/Tests/ReplicationSystem/Filtering/TestFiltering.cpp
+++ b/Engine/Plugins/Runtime/ReplicationSystemTestPlugin/Source/Private/Tests/ReplicationSystem/Filtering/TestFiltering.cpp
@@ -2,6 +2,9 @@
 
 #include "NetworkAutomationTest.h"
 #include "NetworkAutomationTestMacros.h"
+
+#include "UObject/CoreNetTypes.h"
+
 #include "Iris/IrisConstants.h"
 #include "Iris/ReplicationSystem/ReplicationSystem.h"
 #include "Iris/ReplicationSystem/Filtering/NetObjectFilter.h"
@@ -3020,9 +3023,7 @@ UE_NET_TEST_FIXTURE(FTestFilteringFixture, TestIsRelevantForConnection)
 	UReplicatedTestObject* ServerObject = Server->CreateObject(0, 0);
 
 	// Send and deliver packets
-	Server->NetUpdate();
-	Server->SendAndDeliverTo(Client, DeliverPacket);
-	Server->PostSendUpdate();
+	Server->UpdateAndSend({ Client });
 
 	UE_NET_ASSERT_TRUE(Server->ReplicationSystem->IsRelevantForConnection(ServerObject->NetRefHandle, Client->ConnectionIdOnServer));
 
@@ -3033,9 +3034,7 @@ UE_NET_TEST_FIXTURE(FTestFilteringFixture, TestIsRelevantForConnection)
 	Server->ReplicationSystem->SetFilter(ServerObject->NetRefHandle, ToOwnerFilterHandle);
 
 	// Send and deliver packets
-	Server->NetUpdate();
-	Server->SendAndDeliverTo(Client, DeliverPacket);
-	Server->PostSendUpdate();
+	Server->UpdateAndSend({ Client });
 
 	UE_NET_ASSERT_FALSE(Server->ReplicationSystem->IsRelevantForConnection(ServerObject->NetRefHandle, Client->ConnectionIdOnServer));
 
@@ -3046,22 +3045,58 @@ UE_NET_TEST_FIXTURE(FTestFilteringFixture, TestIsRelevantForConnection)
 	Server->ReplicationSystem->SetOwningNetConnection(ServerObject->NetRefHandle, Client->ConnectionIdOnServer);
 
 	// Send and deliver packets
-	Server->NetUpdate();
-	Server->SendAndDeliverTo(Client, DeliverPacket);
-	Server->PostSendUpdate();
+	Server->UpdateAndSend({ Client });
 
 	UE_NET_ASSERT_TRUE(Server->ReplicationSystem->IsRelevantForConnection(ServerObject->NetRefHandle, Client->ConnectionIdOnServer));
 
 	// The object should have been created again
 	UE_NET_ASSERT_NE(Client->GetReplicationBridge()->GetReplicatedObject(ServerObject->NetRefHandle), nullptr);
+	
+
+	// Test Subobjects
+	{
+		// Test a normal subobject
+		UReplicatedTestObject* ServerSubObject = Server->CreateSubObject(ServerObject->NetRefHandle);
+
+		Server->UpdateAndSend({ Client });
+
+		UE_NET_ASSERT_TRUE(Server->ReplicationSystem->IsRelevantForConnection(ServerSubObject->NetRefHandle, Client->ConnectionIdOnServer));
+
+		// Object should now exist on client
+		UE_NET_ASSERT_NE(Client->GetReplicationBridge()->GetReplicatedObject(ServerSubObject->NetRefHandle), nullptr);
+	}
+
+	// Test a normal subobject
+	{
+		UReplicatedTestObject* ServerSubObject = Server->CreateSubObject(ServerObject->NetRefHandle);
+
+		Server->UpdateAndSend({ Client });
+
+		UE_NET_ASSERT_TRUE(Server->ReplicationSystem->IsRelevantForConnection(ServerSubObject->NetRefHandle, Client->ConnectionIdOnServer));
+
+		// Object should now exist on client
+		UE_NET_ASSERT_NE(Client->GetReplicationBridge()->GetReplicatedObject(ServerSubObject->NetRefHandle), nullptr);
+	}
+
+	// Test a subobject set to COND_Never
+	{
+		UReplicatedTestObject* ServerSubObjectNever = Server->CreateSubObject(ServerObject->NetRefHandle);
+
+		Server->ReplicationBridge->SetSubObjectNetCondition(ServerSubObjectNever->NetRefHandle, COND_Never);
+
+		Server->UpdateAndSend({ Client });
+
+		UE_NET_ASSERT_FALSE(Server->ReplicationSystem->IsRelevantForConnection(ServerSubObjectNever->NetRefHandle, Client->ConnectionIdOnServer));
+
+		// Object should never exist on client
+		UE_NET_ASSERT_EQ(Client->GetReplicationBridge()->GetReplicatedObject(ServerSubObjectNever->NetRefHandle), nullptr);
+	}
 
 	// Finally, remove the owning connection
 	Server->ReplicationSystem->SetOwningNetConnection(ServerObject->NetRefHandle, 0);
 
 	// Send and deliver packets
-	Server->NetUpdate();
-	Server->SendAndDeliverTo(Client, DeliverPacket);
-	Server->PostSendUpdate();
+	Server->UpdateAndSend({ Client });
 
 	UE_NET_ASSERT_FALSE(Server->ReplicationSystem->IsRelevantForConnection(ServerObject->NetRefHandle, Client->ConnectionIdOnServer));
 
diff --git a/Engine/Source/Runtime/Net/Iris/Private/Iris/ReplicationSystem/Conditionals/ReplicationConditionals.cpp b/Engine/Source/Runtime/Net/Iris/Private/Iris/ReplicationSystem/Conditionals/ReplicationConditionals.cpp
--- a/Engine/Source/Runtime/Net/Iris/Private/Iris/ReplicationSystem/Conditionals/ReplicationConditionals.cpp
+++ b/Engine/Source/Runtime/Net/Iris/Private/Iris/ReplicationSystem/Conditionals/ReplicationConditionals.cpp
@@ -864,19 +864,22 @@ FReplicationConditionals::FConditionalsMask FReplicationConditionals::GetLifetim
 	}
 
 	ConditionalsMask.SetConditionEnabled(COND_None, true);
-	ConditionalsMask.SetConditionEnabled(COND_Custom, true);
-	ConditionalsMask.SetConditionEnabled(COND_Dynamic, true);
+	ConditionalsMask.SetConditionEnabled(COND_InitialOnly, bIsInitialState);
 	ConditionalsMask.SetConditionEnabled(COND_OwnerOnly, bIsReplicatingToOwner);
 	ConditionalsMask.SetConditionEnabled(COND_SkipOwner, !bIsReplicatingToOwner);
 	ConditionalsMask.SetConditionEnabled(COND_SimulatedOnly, bRoleSimulated);
 	ConditionalsMask.SetConditionEnabled(COND_AutonomousOnly, bRoleAutonomous);
 	ConditionalsMask.SetConditionEnabled(COND_SimulatedOrPhysics, bRoleSimulated | bRepPhysics);
-	ConditionalsMask.SetConditionEnabled(COND_InitialOnly, bIsInitialState);
 	ConditionalsMask.SetConditionEnabled(COND_InitialOrOwner, bIsReplicatingToOwner | bIsInitialState);
+	ConditionalsMask.SetConditionEnabled(COND_Custom, true);
 	ConditionalsMask.SetConditionEnabled(COND_ReplayOrOwner, bIsReplicatingToOwner);
+	//ConditionalsMask.SetConditionEnabled(COND_ReplayOnly, Not support yet);
 	ConditionalsMask.SetConditionEnabled(COND_SimulatedOnlyNoReplay, bRoleSimulated);
 	ConditionalsMask.SetConditionEnabled(COND_SimulatedOrPhysicsNoReplay, bRoleSimulated | bRepPhysics);
 	ConditionalsMask.SetConditionEnabled(COND_SkipReplay, true);
+	ConditionalsMask.SetConditionEnabled(COND_Dynamic, true);
+	//ConditionalsMask.SetConditionEnabled(COND_Never, false); Ctor sets its false by default
+	//ConditionalsMask.SetConditionEnabled(COND_NetGroup, Not Known); NetGroup is evaluated in FReplicationConditionals::GetChildSubObjectsToReplicate
 	
 	return ConditionalsMask;
 }
diff --git a/Engine/Source/Runtime/Net/Iris/Private/Iris/ReplicationSystem/ObjectReplicationBridge.cpp b/Engine/Source/Runtime/Net/Iris/Private/Iris/ReplicationSystem/ObjectReplicationBridge.cpp
--- a/Engine/Source/Runtime/Net/Iris/Private/Iris/ReplicationSystem/ObjectReplicationBridge.cpp
+++ b/Engine/Source/Runtime/Net/Iris/Private/Iris/ReplicationSystem/ObjectReplicationBridge.cpp
@@ -1011,7 +1011,7 @@ void UObjectReplicationBridge::SetSubObjectNetCondition(FNetRefHandle SubObjectR
 	using namespace UE::Net::Private;
 
 	// We assume that we can store the condition in an int8;
-	static_assert(ELifetimeCondition::COND_Max <= 127);
+	static_assert(ELifetimeCondition::COND_Max <= INT8_MAX);
 
 	FReplicationSystemInternal* ReplicationSystemInternal = GetReplicationSystem()->GetReplicationSystemInternal();
 	FNetRefHandleManager& LocalNetRefHandleManager = ReplicationSystemInternal->GetNetRefHandleManager();
@@ -3063,6 +3063,8 @@ void UObjectReplicationBridge::OnProtocolMismatchReported(FNetRefHandle RefHandl
 
 		UE::Net::Private::ObjectBridgeDebugging::RemoteProtocolMismatchDetected(ArchetypesAlreadyPrinted, ReplicationSystem, ConnectionId, Fragments, ObjArchetype, ObjInstance);
 	}
+
+	//TODO: Set object as broken in RepWriter
 }
 
 void UObjectReplicationBridge::OnErrorWithNetRefHandleReported(UE::Net::ENetRefHandleError ErrorType, FNetRefHandle RefHandle, uint32 ConnectionId, const TArray<FNetRefHandle>& ExtraHandles)
@@ -3133,6 +3135,8 @@ void UObjectReplicationBridge::OnErrorCannotReplicateObject(FNetRefHandle RefHan
 	else
 	{
 		UE_LOG_BRIDGEID(LogIris, Error, TEXT("Client cannot replicate object: %s. Reported by client: %s"), *NetRefHandleManager->PrintObjectFromIndex(ObjectInternalIndex), *PrintConnectionInfo(ConnectionId));
+
+		// TODO: Tell RepWriter the creation failed
 	}
 
 	ensureMsgf(false, TEXT("Client disabled replication for %s. Client was unable to read the creation data and did not spawn the netobject."), *NetRefHandleManager->PrintObjectFromNetRefHandle(RefHandle));
diff --git a/Engine/Source/Runtime/Net/Iris/Private/Iris/ReplicationSystem/ReplicationSystem.cpp b/Engine/Source/Runtime/Net/Iris/Private/Iris/ReplicationSystem/ReplicationSystem.cpp
--- a/Engine/Source/Runtime/Net/Iris/Private/Iris/ReplicationSystem/ReplicationSystem.cpp
+++ b/Engine/Source/Runtime/Net/Iris/Private/Iris/ReplicationSystem/ReplicationSystem.cpp
@@ -1386,7 +1386,7 @@ bool UReplicationSystem::IsRelevantForConnection(FNetRefHandle Handle, uint32 Co
 	using namespace UE::Net::Private;
 
 	const FReplicationConnections& Connections = Impl->ReplicationSystemInternal.GetConnections();
-	if (!ensureMsgf(Connections.IsValidConnection(ConnectionId), TEXT("Invalid ConnectionId %u passed to UReplicationSystem::GetDataStream."), ConnectionId))
+	if (!ensureMsgf(Connections.IsValidConnection(ConnectionId), TEXT("Invalid ConnectionId %u passed to UReplicationSystem::IsRelevantForConnection."), ConnectionId))
 	{
 		return false;
 	}
@@ -1399,9 +1399,11 @@ bool UReplicationSystem::IsRelevantForConnection(FNetRefHandle Handle, uint32 Co
 	}
 
 	const FReplicationFiltering& Filtering = Impl->ReplicationSystemInternal.GetFiltering();
-	const  FNetBitArrayView ObjectsInScope = Filtering.GetRelevantObjectsInScope(ConnectionId);
+	const FNetBitArrayView ObjectsInScope = Filtering.GetRelevantObjectsInScope(ConnectionId);
+
+	const FReplicationConnection* Connection = Connections.GetConnection(ConnectionId);
 
-	return ObjectsInScope.GetNumBits() > ObjectInternalIndex ? ObjectsInScope.IsBitSet(ObjectInternalIndex) : false;
+	return ObjectsInScope.GetNumBits() > ObjectInternalIndex ? (ObjectsInScope.IsBitSet(ObjectInternalIndex) && Connection->ReplicationWriter->IsInReplicatingState(ObjectInternalIndex)) : false;
 }
 
 void UReplicationSystem::CloseDataStream(uint32 ConnectionId, FName Name)
diff --git a/Engine/Source/Runtime/Net/Iris/Private/Iris/ReplicationSystem/ReplicationWriter.cpp b/Engine/Source/Runtime/Net/Iris/Private/Iris/ReplicationSystem/ReplicationWriter.cpp
--- a/Engine/Source/Runtime/Net/Iris/Private/Iris/ReplicationSystem/ReplicationWriter.cpp
+++ b/Engine/Source/Runtime/Net/Iris/Private/Iris/ReplicationSystem/ReplicationWriter.cpp
@@ -281,6 +281,22 @@ bool FReplicationWriter::IsReplicationEnabled() const
 	return bReplicationEnabled;
 }
 
+bool FReplicationWriter::IsInReplicatingState(FInternalNetRefIndex ObjectInternalIndex) const
+{
+	if (!ReplicatedObjects.IsValidIndex(ObjectInternalIndex))
+	{
+		ensureMsgf(false, TEXT("Received invalid index: %u"), ObjectInternalIndex);
+		return false;
+	}
+
+	// We let InvalidInternalIndex be accessed because it has an object state too.
+	const FReplicationInfo& Info = GetReplicationInfo(ObjectInternalIndex);
+	const EReplicatedObjectState CurrentState = Info.GetState();
+	
+	// We consider an object to be relevant while it's creation data is in flight and before we sent destruction data.
+	return CurrentState >= EReplicatedObjectState::WaitOnCreateConfirmation && CurrentState <= EReplicatedObjectState::WaitOnFlush;
+}
+
 // $IRIS TODO : May need to introduce queue and send behaviors. For example one may want to send only with object.
 // One may not want to send unless the object is replicated very soon etc.
 bool FReplicationWriter::QueueNetObjectAttachments(FInternalNetRefIndex OwnerInternalIndex, FInternalNetRefIndex SubObjectInternalIndex, TArrayView<const TRefCountPtr<FNetBlob>> InAttachments, ENetObjectAttachmentSendPolicyFlags SendFlags)
diff --git a/Engine/Source/Runtime/Net/Iris/Private/Iris/ReplicationSystem/ReplicationWriter.h b/Engine/Source/Runtime/Net/Iris/Private/Iris/ReplicationSystem/ReplicationWriter.h
--- a/Engine/Source/Runtime/Net/Iris/Private/Iris/ReplicationSystem/ReplicationWriter.h
+++ b/Engine/Source/Runtime/Net/Iris/Private/Iris/ReplicationSystem/ReplicationWriter.h
@@ -198,6 +198,9 @@ public:
 	void SetReplicationEnabled(bool bInReplicationEnabled);
 	bool IsReplicationEnabled() const;
 
+	/** Is the object in a state where we can send data about it. It will be false when the object is about to be destroyed or we are waiting for destroy confirmation. */
+	bool IsInReplicatingState(FInternalNetRefIndex ObjectInternalIndex) const;
+
 	void SetNetExports(FNetExports& InNetExports);
 
 	// Attachments
