Iris
* Set net.Iris.IsUsingNewDetachAPI to default to true so that we run the new Detach API by default

#rnx
#Synced-CL 49546171

diff --git a/Engine/Plugins/Runtime/ReplicationSystemTestPlugin/Source/Private/Tests/ReplicationSystem/ReplicatedTestObjectFactory.cpp b/Engine/Plugins/Runtime/ReplicationSystemTestPlugin/Source/Private/Tests/ReplicationSystem/ReplicatedTestObjectFactory.cpp
--- a/Engine/Plugins/Runtime/ReplicationSystemTestPlugin/Source/Private/Tests/ReplicationSystem/ReplicatedTestObjectFactory.cpp
+++ b/Engine/Plugins/Runtime/ReplicationSystemTestPlugin/Source/Private/Tests/ReplicationSystem/ReplicatedTestObjectFactory.cpp
@@ -284,6 +284,36 @@ void UReplicatedTestObjectFactory::DetachedFromReplication(const FDestroyedConte
 	}
 }
 
+void UReplicatedTestObjectFactory::DetachedFromReplication(const FDetachContext& Context, const TOptional<FSubObjectDetachContext>& SubObjectContext)
+{
+	using namespace UE::Net;
+
+	if (Context.DetachedInstance == nullptr)
+	{
+		return;
+	}
+
+	if (Context.Reason == EDetachReason::TornOff)
+	{
+		return;
+	}
+
+	// Always destroy dynamic objects and stable objects only when static destroyed
+	const bool bDestroyInstance = Context.Reason == EDetachReason::StaticDestroyed || !Context.DetachedInstance->IsNameStableForNetworking();
+
+	if (bDestroyInstance)
+	{
+		// Remove the object from the created objects on the node
+		if (CreatedObjectsOnNode)
+		{
+			CreatedObjectsOnNode->Remove(TStrongObjectPtr<UObject>(Context.DetachedInstance));
+		}
+
+		Context.DetachedInstance->PreDestroyFromReplication();
+		Context.DetachedInstance->MarkAsGarbage();
+	}
+}
+
 TOptional<UNetObjectFactory::FWorldInfoData> UReplicatedTestObjectFactory::GetWorldInfo(const FWorldInfoContext& Context) const
 {
 	if (GetInstanceWorldObjectInfoFunction)
@@ -313,4 +343,13 @@ void UReplicatedTestObjectFactory::SubObjectDetachedFromReplication(const FDestr
 	{
 		RootInstance->OnSubObjectDestroyed(Context.DestroyedInstance);
 	}
+}
+
+void UReplicatedTestObjectFactory::SubObjectDetachedFromReplication(const FDetachContext& Context, const FSubObjectDetachContext& SubObjectContext)
+{
+	UReplicatedTestObject* RootInstance = Cast<UReplicatedTestObject>(SubObjectContext.RootObject);
+	if (RootInstance)
+	{
+		RootInstance->OnSubObjectDestroyed(Context.DetachedInstance);
+	}
 }
\ No newline at end of file
diff --git a/Engine/Plugins/Runtime/ReplicationSystemTestPlugin/Source/Private/Tests/ReplicationSystem/ReplicatedTestObjectFactory.h b/Engine/Plugins/Runtime/ReplicationSystemTestPlugin/Source/Private/Tests/ReplicationSystem/ReplicatedTestObjectFactory.h
--- a/Engine/Plugins/Runtime/ReplicationSystemTestPlugin/Source/Private/Tests/ReplicationSystem/ReplicatedTestObjectFactory.h
+++ b/Engine/Plugins/Runtime/ReplicationSystemTestPlugin/Source/Private/Tests/ReplicationSystem/ReplicatedTestObjectFactory.h
@@ -45,10 +45,16 @@ public:
 
 	// UNetObjectFactory interface
 	virtual void PostInit(const FPostInitContext& Context) override;
+
 	virtual void DetachedFromReplication(const FDestroyedContext& Context) override;
+	virtual void DetachedFromReplication(const FDetachContext& Context, const TOptional<FSubObjectDetachContext>& SubObjectContext) override;
+
 	virtual TOptional<FWorldInfoData> GetWorldInfo(const FWorldInfoContext& Context) const override;
+
 	virtual void SubObjectCreatedFromReplication(UE::Net::FNetRefHandle RootObject, UE::Net::FNetRefHandle SubObjectCreated) override;
+
 	virtual void SubObjectDetachedFromReplication(const FDestroyedContext& Context) override;
+	virtual void SubObjectDetachedFromReplication(const FDetachContext& Context, const FSubObjectDetachContext& SubObjectContext) override;
 
 	//------------------------------------------------------------------------
 
diff --git a/Engine/Source/Runtime/Net/Iris/Private/Iris/IrisConfig.cpp b/Engine/Source/Runtime/Net/Iris/Private/Iris/IrisConfig.cpp
--- a/Engine/Source/Runtime/Net/Iris/Private/Iris/IrisConfig.cpp
+++ b/Engine/Source/Runtime/Net/Iris/Private/Iris/IrisConfig.cpp
@@ -18,7 +18,7 @@ static FAutoConsoleVariableRef CVarUseIrisReplicationRef(TEXT("net.Iris.UseIrisR
 static bool bScheduleCreationDependenciesFirst = true;
 static FAutoConsoleVariableRef CVarScheduleCreationDependenciesFirst(TEXT("net.Iris.ScheduleCreationDependenciesFirst"), bScheduleCreationDependenciesFirst, TEXT("If enabled we will schedule an objects creation dependencies before the object if they are not confirmed as created."));
 
-static bool bIsUsingNewDetachAPI = false;
+static bool bIsUsingNewDetachAPI = true;
 static FAutoConsoleVariableRef CVarIsUsingNewDetachAPI(TEXT("net.Iris.IsUsingNewDetachAPI"), bIsUsingNewDetachAPI, TEXT("If enabled clients will use the new detach API."));
 
 }
