Don't call Iris factory detach functions in the new detach API (net.Iris.IsUsingNewDetachAPI) when the detach reason is Migrated, to match the old detach API behavior.
Add unit tests for replication systems in proxy mode re-using objects after a mock migration.

#Synced-CL 49807081

diff --git a/Engine/Plugins/Runtime/ReplicationSystemTestPlugin/Source/Private/Tests/ReplicationSystem/TestReplicationSystemProxy.cpp b/Engine/Plugins/Runtime/ReplicationSystemTestPlugin/Source/Private/Tests/ReplicationSystem/TestReplicationSystemProxy.cpp
--- a/Engine/Plugins/Runtime/ReplicationSystemTestPlugin/Source/Private/Tests/ReplicationSystem/TestReplicationSystemProxy.cpp
+++ b/Engine/Plugins/Runtime/ReplicationSystemTestPlugin/Source/Private/Tests/ReplicationSystem/TestReplicationSystemProxy.cpp
@@ -128,6 +128,287 @@ UE_NET_TEST_FIXTURE(FReplicationSystemProxyTestFixture, TestTwoServers)
 	UE_NET_ASSERT_EQ(ClientObject2->IntA, IntServer2);
 }
 
-} // end namespace UE::Net::Private
+// Test proxy mode replication systems replicating properties of a reused object after a mock migration to another backend replication system/server.
+UE_NET_TEST_FIXTURE(FReplicationSystemProxyTestFixture, ReuseObjectPropertyReplication)
+{
+	// Add two backend servers. Explicitly use the same initial NetRefHandleIndex because
+	// this test fakes migration where the handle IDs have to match. We will only spawn one root object
+	// on each server (the one on server 2 being when it's "migrated" from server 1).
+	constexpr uint32 InitialNetRefHandleIndexServer1 = 1000;
+	FReplicationSystemTestServer* Server1 = CreateServer(InitialNetRefHandleIndexServer1);
+
+	constexpr uint32 InitialNetRefHandleIndexServer2 = 1000;
+	FReplicationSystemTestServer* Server2 = CreateServer(InitialNetRefHandleIndexServer2);
+
+	// Add a proxy server
+	FReplicationSystemTestProxy* Proxy = CreateProxy();
+
+	// Add a client
+	FReplicationSystemTestClient* Client = CreateClient();
+
+	// Make connections
+	ConnectProxyToBackendServer(Proxy, Server1);
+	ConnectProxyToBackendServer(Proxy, Server2);
+	ConnectClientToProxy(Client, Proxy);
+
+	// Spawn one object on one server
+	UTestReplicatedIrisObject* Server1RootObject = Server1->CreateObject();
+	UTestReplicatedIrisObject* Server1SubObject = Server1->CreateSubObject(Server1RootObject->NetRefHandle);
+
+	const int32 RootObjectInt = 10;
+	const int32 SubObjectInt = 30;
+	Server1RootObject->IntA = RootObjectInt;
+	Server1SubObject->IntA = SubObjectInt;
+
+	// Send and deliver packet, servers -> proxy
+	UpdateAndSendServerToAllProxies(Server1);
+	UpdateAndSendServerToAllProxies(Server2);
+
+	// Verify that objects have been spawned on the proxy
+	FReplicationSystemTestClient* ProxyBackendClient1 = GetProxyClientForBackendServer(Server1, Proxy);
+	FReplicationSystemTestClient* ProxyBackendClient2 = GetProxyClientForBackendServer(Server2, Proxy);
+	UTestReplicatedIrisObject* ProxyRootObject = ProxyBackendClient1->GetObjectAs<UTestReplicatedIrisObject>(Server1RootObject->NetRefHandle);
+	UTestReplicatedIrisObject* ProxySubObject = ProxyBackendClient1->GetObjectAs<UTestReplicatedIrisObject>(Server1SubObject->NetRefHandle);
+	UE_NET_ASSERT_NE(ProxyRootObject, nullptr);
+	UE_NET_ASSERT_NE(ProxyRootObject, Server1RootObject);
+	UE_NET_ASSERT_NE(ProxySubObject, nullptr);
+	UE_NET_ASSERT_NE(ProxySubObject, Server1SubObject);
+
+	// Check property values on the proxy
+	UE_NET_ASSERT_EQ(ProxyRootObject->IntA, RootObjectInt);
+	UE_NET_ASSERT_EQ(ProxySubObject->IntA, SubObjectInt);
+
+	// Let the proxy replicate the object to its client
+	Proxy->GetFrontendServer()->GetReplicationBridge()->BeginReplication(ProxyRootObject);
+	Proxy->GetFrontendServer()->GetReplicationBridge()->BeginReplication(ProxyRootObject->NetRefHandle, ProxySubObject);
+
+	// Send and deliver packet, proxy -> client
+	UpdateAndSendProxyToAllClients(Proxy);
+
+	// Verify that the object has been spawned on client
+	UTestReplicatedIrisObject* ClientRootObject = Client->GetObjectAs<UTestReplicatedIrisObject>(Server1RootObject->NetRefHandle);
+	UTestReplicatedIrisObject* ClientSubObject = Client->GetObjectAs<UTestReplicatedIrisObject>(Server1SubObject->NetRefHandle);
+	UE_NET_ASSERT_NE(ClientRootObject, nullptr);
+	UE_NET_ASSERT_NE(ClientRootObject, ProxyRootObject);
+	UE_NET_ASSERT_NE(ClientSubObject, nullptr);
+	UE_NET_ASSERT_NE(ClientSubObject, ProxySubObject);
 
+	// Verify that the values are correct on the client
+	UE_NET_ASSERT_EQ(ClientRootObject->IntA, RootObjectInt);
+	UE_NET_ASSERT_EQ(ClientSubObject->IntA, SubObjectInt);
+
+	// Fake a "migration" of the object to server 2
+	UTestReplicatedIrisObject* Server2RootObject = Server2->CreateObject();
+	UTestReplicatedIrisObject* Server2SubObject = Server2->CreateSubObject(Server2RootObject->NetRefHandle);
+	Server2RootObject->IntA = Server1RootObject->IntA;
+	Server2SubObject->IntA = Server1SubObject->IntA;
+
+	// Make sure the RefHandles match
+	UE_NET_ASSERT_EQ(Server1RootObject->NetRefHandle.GetId(), Server2RootObject->NetRefHandle.GetId());
+	UE_NET_ASSERT_EQ(Server1SubObject->NetRefHandle.GetId(), Server2SubObject->NetRefHandle.GetId());
+
+	Server1->DestroyObject(Server1SubObject, EEndReplicationFlags::Migrated);
+	Server1->DestroyObject(Server1RootObject, EEndReplicationFlags::Migrated);
+	Server1RootObject = nullptr;
+	Server1SubObject = nullptr;
+
+	// Send and deliver packet, servers -> proxy
+	UpdateAndSendServerToAllProxies(Server1);
+	UpdateAndSendServerToAllProxies(Server2);
+
+	// Verify that the root object is still the same on the proxy
+	UTestReplicatedIrisObject* ProxyRootObjectAfterMigration = ProxyBackendClient2->GetObjectAs<UTestReplicatedIrisObject>(Server2RootObject->NetRefHandle);
+	UE_NET_ASSERT_EQ(ProxyRootObjectAfterMigration, ProxyRootObject);
+	UE_NET_ASSERT_EQ(ProxyRootObjectAfterMigration->IntA, RootObjectInt);
+
+	// Verify that the subobject is still the same on the proxy
+	UTestReplicatedIrisObject* ProxySubObjectAfterMigration = ProxyBackendClient2->GetObjectAs<UTestReplicatedIrisObject>(Server2SubObject->NetRefHandle);
+	UE_NET_ASSERT_EQ(ProxySubObjectAfterMigration, ProxySubObject);
+	UE_NET_ASSERT_EQ(ProxySubObjectAfterMigration->IntA, SubObjectInt);
+
+	// Send and deliver packet, proxy -> client
+	UpdateAndSendProxyToAllClients(Proxy);
+
+	// Verify that the object is still the same on the client
+	UTestReplicatedIrisObject* ClientRootObjectAfterMigration = Client->GetObjectAs<UTestReplicatedIrisObject>(Server2RootObject->NetRefHandle);
+	UE_NET_ASSERT_EQ(ClientRootObjectAfterMigration, ClientRootObject);
+	UE_NET_ASSERT_EQ(ClientRootObjectAfterMigration->IntA, RootObjectInt);
+
+	UTestReplicatedIrisObject* ClientSubObjectAfterMigration = Client->GetObjectAs<UTestReplicatedIrisObject>(Server2SubObject->NetRefHandle);
+	UE_NET_ASSERT_EQ(ClientSubObjectAfterMigration, ClientSubObject);
+	UE_NET_ASSERT_EQ(ClientSubObjectAfterMigration->IntA, SubObjectInt);
+
+	// Replicate a new value on the object from server 2
+	const int32 NewRootObjectIntValue = RootObjectInt * 2;
+	Server2RootObject->IntA = NewRootObjectIntValue;
+
+	const int32 NewSubObjectIntValue = SubObjectInt * 2;
+	Server2SubObject->IntA = NewSubObjectIntValue;
+
+	// Send and deliver packet, servers -> proxy
+	UpdateAndSendServerToAllProxies(Server1);
+	UpdateAndSendServerToAllProxies(Server2);
+
+	UE_NET_ASSERT_EQ(ProxyRootObjectAfterMigration->IntA, NewRootObjectIntValue);
+	UE_NET_ASSERT_EQ(ProxySubObjectAfterMigration->IntA, NewSubObjectIntValue);
+
+	// Send and deliver packet, proxy -> client
+	UpdateAndSendProxyToAllClients(Proxy);
+
+	UE_NET_ASSERT_EQ(ClientRootObjectAfterMigration->IntA, NewRootObjectIntValue);
+	UE_NET_ASSERT_EQ(ClientSubObjectAfterMigration->IntA, NewSubObjectIntValue);
+}
+
+// Test proxy mode replication systems destroying a reused object after a mock migration to another backend replication system/server.
+UE_NET_TEST_FIXTURE(FReplicationSystemProxyTestFixture, ReuseObjectDestruction)
+{
+	// Add two backend servers. Explicitly use the same initial NetRefHandleIndex because
+	// this test fakes migration where the handle IDs have to match. We will only spawn one object
+	// on each server (the one on server 2 being when it's "migrated" from server 1).
+	constexpr uint32 InitialNetRefHandleIndexServer1 = 1000;
+	FReplicationSystemTestServer* Server1 = CreateServer(InitialNetRefHandleIndexServer1);
+
+	constexpr uint32 InitialNetRefHandleIndexServer2 = 1000;
+	FReplicationSystemTestServer* Server2 = CreateServer(InitialNetRefHandleIndexServer2);
 
+	// Add a proxy server
+	FReplicationSystemTestProxy* Proxy = CreateProxy();
+
+	// Add a client
+	FReplicationSystemTestClient* Client = CreateClient();
+
+	// Make connections
+	ConnectProxyToBackendServer(Proxy, Server1);
+	ConnectProxyToBackendServer(Proxy, Server2);
+	ConnectClientToProxy(Client, Proxy);
+
+	// Spawn one object on one server
+	UTestReplicatedIrisObject* Server1RootObject = Server1->CreateObject();
+	UTestReplicatedIrisObject* Server1SubObject = Server1->CreateSubObject(Server1RootObject->NetRefHandle);
+
+	// Send and deliver packet, servers -> proxy
+	UpdateAndSendServerToAllProxies(Server1);
+	UpdateAndSendServerToAllProxies(Server2);
+
+	// Verify that objects have been spawned on the proxy
+	FReplicationSystemTestClient* ProxyBackendClient1 = GetProxyClientForBackendServer(Server1, Proxy);
+	FReplicationSystemTestClient* ProxyBackendClient2 = GetProxyClientForBackendServer(Server2, Proxy);
+	UTestReplicatedIrisObject* ProxyRootObject = ProxyBackendClient1->GetObjectAs<UTestReplicatedIrisObject>(Server1RootObject->NetRefHandle);
+	UTestReplicatedIrisObject* ProxySubObject = ProxyBackendClient1->GetObjectAs<UTestReplicatedIrisObject>(Server1SubObject->NetRefHandle);
+	UE_NET_ASSERT_NE(ProxyRootObject, nullptr);
+	UE_NET_ASSERT_NE(ProxyRootObject, Server1RootObject);
+	UE_NET_ASSERT_NE(ProxySubObject, nullptr);
+	UE_NET_ASSERT_NE(ProxySubObject, Server1SubObject);
+
+	// Let the proxy replicate the object to its client
+	Proxy->GetFrontendServer()->GetReplicationBridge()->BeginReplication(ProxyRootObject);
+	Proxy->GetFrontendServer()->GetReplicationBridge()->BeginReplication(ProxyRootObject->NetRefHandle, ProxySubObject);
+
+	// Send and deliver packet, proxy -> client
+	UpdateAndSendProxyToAllClients(Proxy);
+
+	// Verify that the object has been spawned on client
+	UTestReplicatedIrisObject* ClientRootObject = Client->GetObjectAs<UTestReplicatedIrisObject>(Server1RootObject->NetRefHandle);
+	UTestReplicatedIrisObject* ClientSubObject = Client->GetObjectAs<UTestReplicatedIrisObject>(Server1SubObject->NetRefHandle);
+	UE_NET_ASSERT_NE(ClientRootObject, nullptr);
+	UE_NET_ASSERT_NE(ClientRootObject, ProxyRootObject);
+	UE_NET_ASSERT_NE(ClientSubObject, nullptr);
+	UE_NET_ASSERT_NE(ClientSubObject, ProxySubObject);
+
+	// Fake a "migration" of the object to server 2
+	UTestReplicatedIrisObject* Server2RootObject = Server2->CreateObject();
+	UTestReplicatedIrisObject* Server2SubObject = Server2->CreateSubObject(Server2RootObject->NetRefHandle);
+
+	// Make sure the RefHandles match
+	FNetRefHandle SavedRootObjectHandle = Server1RootObject->NetRefHandle;
+	FNetRefHandle SavedSubObjectHandle = Server1SubObject->NetRefHandle;
+	UE_NET_ASSERT_EQ(Server1RootObject->NetRefHandle.GetId(), Server2RootObject->NetRefHandle.GetId());
+	UE_NET_ASSERT_EQ(Server1SubObject->NetRefHandle.GetId(), Server2SubObject->NetRefHandle.GetId());
+
+	Server1->DestroyObject(Server1RootObject, EEndReplicationFlags::Migrated);
+	Server1RootObject = nullptr;
+
+	// Send and deliver packet, servers -> proxy
+	UpdateAndSendServerToAllProxies(Server1);
+	UpdateAndSendServerToAllProxies(Server2);
+
+	// Make sure proxy object is still the same
+	UTestReplicatedIrisObject* ProxyRootObjectAfterMigration = ProxyBackendClient2->GetObjectAs<UTestReplicatedIrisObject>(SavedRootObjectHandle);
+	UTestReplicatedIrisObject* ProxySubObjectAfterMigration = ProxyBackendClient2->GetObjectAs<UTestReplicatedIrisObject>(SavedSubObjectHandle);
+	UE_NET_ASSERT_EQ(ProxyRootObjectAfterMigration, ProxyRootObject);
+	UE_NET_ASSERT_EQ(ProxySubObjectAfterMigration, ProxySubObject);
+
+	// Save weak pointers so we can check for the object actually being destroyed on the proxy
+	TWeakObjectPtr<UTestReplicatedIrisObject> WeakProxyRootObject = ProxyRootObjectAfterMigration;
+	TWeakObjectPtr<UTestReplicatedIrisObject> WeakProxySubObject = ProxySubObjectAfterMigration;
+
+	// Send and deliver packet, proxy -> client
+	UpdateAndSendProxyToAllClients(Proxy);
+
+	// Verify that the object is still the same on the client
+	UTestReplicatedIrisObject* ClientRootObjectAfterMigration = Client->GetObjectAs<UTestReplicatedIrisObject>(SavedRootObjectHandle);
+	UTestReplicatedIrisObject* ClientSubObjectAfterMigration = Client->GetObjectAs<UTestReplicatedIrisObject>(SavedSubObjectHandle);
+	UE_NET_ASSERT_EQ(ClientRootObjectAfterMigration, ClientRootObject);
+	UE_NET_ASSERT_EQ(ClientSubObjectAfterMigration, ClientSubObject);
+
+	// Save weak pointers so we can check for the object actually being destroyed on the client
+	TWeakObjectPtr<UTestReplicatedIrisObject> WeakClientRootObject = ClientRootObjectAfterMigration;
+	TWeakObjectPtr<UTestReplicatedIrisObject> WeakClientSubObject = ClientSubObjectAfterMigration;
+
+	// Destroy the subobject now on server 2
+	Server2->DestroyObject(Server2SubObject, EEndReplicationFlags::Destroy);
+	
+	// Send and deliver packet, servers -> proxy
+	UpdateAndSendServerToAllProxies(Server1);
+	UpdateAndSendServerToAllProxies(Server2);
+
+	// Root object should still exist
+	UTestReplicatedIrisObject* ProxyRootObjectAfterSubObjectDestruction = ProxyBackendClient2->GetObjectAs<UTestReplicatedIrisObject>(SavedRootObjectHandle);
+	UE_NET_ASSERT_EQ(ProxyRootObjectAfterSubObjectDestruction, ProxyRootObjectAfterMigration);
+	UE_NET_ASSERT_TRUE(WeakProxyRootObject.IsValid());
+
+	// Subobject should be destroyed
+	UTestReplicatedIrisObject* ProxySubObjectAfterSubObjectDestruction = ProxyBackendClient2->GetObjectAs<UTestReplicatedIrisObject>(SavedSubObjectHandle);
+	UE_NET_ASSERT_EQ(ProxySubObjectAfterSubObjectDestruction, nullptr);
+	UE_NET_ASSERT_FALSE(WeakProxySubObject.IsValid());
+
+	// Explicitly stop replicating on the proxy frontend
+	Proxy->GetFrontendServer()->GetReplicationBridge()->StopReplicatingNetRefHandle(SavedSubObjectHandle, EEndReplicationFlags::Destroy);
+
+	// Send and deliver packet, proxy -> client
+	UpdateAndSendProxyToAllClients(Proxy);
+
+	// Root object should still exist
+	UTestReplicatedIrisObject* ClientRootObjectAfterSubObjectDestruction = Client->GetObjectAs<UTestReplicatedIrisObject>(SavedRootObjectHandle);
+	UE_NET_ASSERT_EQ(ClientRootObjectAfterSubObjectDestruction, ClientRootObjectAfterMigration);
+	UE_NET_ASSERT_TRUE(WeakClientRootObject.IsValid());
+
+	// Subobject should be destroyed
+	UTestReplicatedIrisObject* ClientSubObjectAfterSubObjectDestruction = Client->GetObjectAs<UTestReplicatedIrisObject>(SavedSubObjectHandle);
+	UE_NET_ASSERT_EQ(ClientSubObjectAfterSubObjectDestruction, nullptr);
+	UE_NET_ASSERT_FALSE(WeakClientSubObject.IsValid());
+
+	// Destroy the root object now on server 2
+	Server2->DestroyObject(Server2RootObject, EEndReplicationFlags::Destroy);
+
+	// Send and deliver packet, servers -> proxy
+	UpdateAndSendServerToAllProxies(Server1);
+	UpdateAndSendServerToAllProxies(Server2);
+
+	// Root object should be destroyed
+	UTestReplicatedIrisObject* ProxyRootObjectAfterRootObjectDestruction = ProxyBackendClient2->GetObjectAs<UTestReplicatedIrisObject>(SavedRootObjectHandle);
+	UE_NET_ASSERT_EQ(ProxyRootObjectAfterRootObjectDestruction, nullptr);
+	UE_NET_ASSERT_FALSE(WeakProxyRootObject.IsValid());
+
+	// Explicitly stop replicating on the proxy frontend
+	Proxy->GetFrontendServer()->GetReplicationBridge()->StopReplicatingNetRefHandle(SavedRootObjectHandle, EEndReplicationFlags::Destroy);
+
+	// Send and deliver packet, proxy -> client
+	UpdateAndSendProxyToAllClients(Proxy);
+
+	UTestReplicatedIrisObject* ClientRootObjectAfterRootObjectDestruction = Client->GetObjectAs<UTestReplicatedIrisObject>(SavedRootObjectHandle);
+	UE_NET_ASSERT_EQ(ClientRootObjectAfterRootObjectDestruction, nullptr);
+	UE_NET_ASSERT_FALSE(WeakClientRootObject.IsValid());
+}
+
+} // end namespace UE::Net::Private
diff --git a/Engine/Source/Runtime/Net/Iris/Private/Iris/ReplicationSystem/ObjectReplicationBridge.cpp b/Engine/Source/Runtime/Net/Iris/Private/Iris/ReplicationSystem/ObjectReplicationBridge.cpp
--- a/Engine/Source/Runtime/Net/Iris/Private/Iris/ReplicationSystem/ObjectReplicationBridge.cpp
+++ b/Engine/Source/Runtime/Net/Iris/Private/Iris/ReplicationSystem/ObjectReplicationBridge.cpp
@@ -1420,24 +1420,24 @@ void UObjectReplicationBridge::DetachRootObjectFromRemote(FNetRefHandle RootObje
 			.Reason = DetachReason,
 		};
 
-		// Tell the root object he will lose this subobject
-		RootObjectFactory->SubObjectDetachedFromReplication(DetachContext, SubObjectDetachContext);
-
-		// Now ask the subobject's factory to handle the detachment
-		if (UNetObjectFactory* SubObjectFactory = GetNetFactory(SubObjectData.NetFactoryId))
+		if (DetachReason != EDetachReason::Migrated)
 		{
-			/*PRAGMA_DISABLE_DEPRECATION_WARNINGS
-			if (bCallDeprecatedFunctions)
+			// Tell the root object he will lose this subobject
+			RootObjectFactory->SubObjectDetachedFromReplication(DetachContext, SubObjectDetachContext);
+
+			// Now ask the subobject's factory to handle the detachment
+			if (UNetObjectFactory* SubObjectFactory = GetNetFactory(SubObjectData.NetFactoryId))
 			{
-				SubObjectFactory->DetachedFromReplication(SubObjectDestroyContext);
+				/*PRAGMA_DISABLE_DEPRECATION_WARNINGS
+				if (bCallDeprecatedFunctions)
+				{
+					SubObjectFactory->DetachedFromReplication(SubObjectDestroyContext);
+				}
+				PRAGMA_ENABLE_DEPRECATION_WARNINGS*/
+
+				SubObjectFactory->DetachedFromReplication(DetachContext, SubObjectDetachContext);
 			}
-			PRAGMA_ENABLE_DEPRECATION_WARNINGS*/
 
-			SubObjectFactory->DetachedFromReplication(DetachContext, SubObjectDetachContext);
-		}
-		
-		if (DetachReason != EDetachReason::Migrated)
-		{
 			UnregisterInstance(SubObjectHandle);
 		}
 	}
@@ -1470,14 +1470,15 @@ void UObjectReplicationBridge::DetachRootObjectFromRemote(FNetRefHandle RootObje
 		.DetachedObjectHandle = RootObjectHandle,
 		.Reason = DetachReason,
 	};
-	RootObjectFactory->DetachedFromReplication(DetachContext, {});
 
 	// Migrated objects need to keep the Handle->Pointer binding alive since proxies reuse pointers instead of recreating them. This won't be required in the future if proxies simply do not instantiate a local instance and use protocol data directly.
 	if (DetachReason == EDetachReason::Migrated)
 	{
-        return;
+		return;
 	}
 
+	RootObjectFactory->DetachedFromReplication(DetachContext, {});
+
 	UnregisterInstance(RootObjectHandle);
 }
 
@@ -1569,8 +1570,6 @@ void UObjectReplicationBridge::DetachSubObjectFromRemote(FNetRefHandle DetachedS
 
 		RootObjectFactory->SubObjectDetachedFromReplication(DetachContext, SubObjectDetachContext);
 	}
-	
-	SubObjectFactory->DetachedFromReplication(DetachContext, SubObjectDetachContext);
 
 	// Migrated objects need to keep the Handle->Pointer binding alive since proxies reuse pointers instead of recreating them. This won't be required in the future if proxies simply do not instantiate a local instance and use protocol data directly.
 	if (DetachReason == EDetachReason::Migrated)
@@ -1578,6 +1577,8 @@ void UObjectReplicationBridge::DetachSubObjectFromRemote(FNetRefHandle DetachedS
 		return;
 	}
 
+	SubObjectFactory->DetachedFromReplication(DetachContext, SubObjectDetachContext);
+
 	UnregisterInstance(DetachedSubObjectHandle);
 }
 
