Iris
- Added ensure if reliable attachementqueue is overcommited.

#14251

#Synced-CL 49833605

diff --git a/Engine/Plugins/Runtime/ReplicationSystemTestPlugin/Source/Private/Tests/ReplicationSystem/RPC/TestObjectRPC.cpp b/Engine/Plugins/Runtime/ReplicationSystemTestPlugin/Source/Private/Tests/ReplicationSystem/RPC/TestObjectRPC.cpp
--- a/Engine/Plugins/Runtime/ReplicationSystemTestPlugin/Source/Private/Tests/ReplicationSystem/RPC/TestObjectRPC.cpp
+++ b/Engine/Plugins/Runtime/ReplicationSystemTestPlugin/Source/Private/Tests/ReplicationSystem/RPC/TestObjectRPC.cpp
@@ -566,6 +566,9 @@ namespace UE::Net::Private
 
 		ClientObject->ReplicationSystem = Client->GetReplicationSystem();
 
+		// Explicitly set ReliableRPCQueueSize to avoid issues if it is changed as what we want to test is filling the reliable sequence window.
+		TGuardConsoleVariable<int32> CVarOverride(TEXT("net.ReliableRPCQueueSize"), 10000);
+		
 		// Issue enough RPCs to overshoot reliable send window
 		const int32 RPCCount = 10000;
 
diff --git a/Engine/Source/Runtime/Net/Iris/Private/Iris/ReplicationSystem/AttachmentReplication.cpp b/Engine/Source/Runtime/Net/Iris/Private/Iris/ReplicationSystem/AttachmentReplication.cpp
--- a/Engine/Source/Runtime/Net/Iris/Private/Iris/ReplicationSystem/AttachmentReplication.cpp
+++ b/Engine/Source/Runtime/Net/Iris/Private/Iris/ReplicationSystem/AttachmentReplication.cpp
@@ -85,6 +85,7 @@ public:
 		// Assume that we won't fit all the attachments unless they can fit in the pre-queue.
 		if (AttachmentCount + PreQueueCount > MaxPreQueueCount)
 		{
+			ensureMsgf(false, TEXT("%u reliable RPCs not enqueued as adding them would exceed the current MaxPreQueueCount (%u)"), AttachmentCount, MaxPreQueueCount);
 			return false;
 		}
 
