Make callstacks include the test name for network low level tests.

#Synced-CL 49914824

diff --git a/Engine/Plugins/Runtime/ReplicationSystemTestPlugin/Source/Private/NetworkAutomationTest.h b/Engine/Plugins/Runtime/ReplicationSystemTestPlugin/Source/Private/NetworkAutomationTest.h
--- a/Engine/Plugins/Runtime/ReplicationSystemTestPlugin/Source/Private/NetworkAutomationTest.h
+++ b/Engine/Plugins/Runtime/ReplicationSystemTestPlugin/Source/Private/NetworkAutomationTest.h
@@ -118,7 +118,7 @@ public:
 
 #define UE_NET_TEST_CATCH_INTERNAL(TestSuite, TestCase, BaseClass) \
 	constexpr TCHAR TestSuite ## _ ## TestCase ## Name[] = TEXT("Net." #TestSuite "." #TestCase); \
-	TEST_CASE_METHOD((TCatchFixtureWrapper<BaseClass, TestSuite ## _ ## TestCase ## Name>), UE_STRINGIZE(TestSuite ## _ ## TestCase))
+	INTERNAL_CATCH_TEST_CASE_METHOD2(Catch::DefaultGroup, TestSuite ## _ ## TestCase, (TCatchFixtureWrapper<BaseClass, TestSuite ## _ ## TestCase ## Name>), UE_STRINGIZE(TestSuite ## _ ## TestCase))
 
 #define UE_NET_TEST_FIXTURE(TestFixture, TestCase) UE_NET_TEST_CATCH_INTERNAL(TestFixture, TestCase, TestFixture)
 #define UE_NET_TEST(TestSuite, TestCase) UE_NET_TEST_CATCH_INTERNAL(TestSuite, TestCase, UE::Net::FNetworkAutomationTestSuiteFixture)
