Explicitly register and deregister replicated objects with the frontend replication system when using the multiserver proxy.

#Synced-CL 50100874

diff --git a/Engine/Plugins/Runtime/MultiServerReplication/Source/MultiServerReplication/Private/MultiServerProxy.cpp b/Engine/Plugins/Runtime/MultiServerReplication/Source/MultiServerReplication/Private/MultiServerProxy.cpp
--- a/Engine/Plugins/Runtime/MultiServerReplication/Source/MultiServerReplication/Private/MultiServerProxy.cpp
+++ b/Engine/Plugins/Runtime/MultiServerReplication/Source/MultiServerReplication/Private/MultiServerProxy.cpp
@@ -510,6 +510,19 @@ bool UProxyBackendNetDriver::InitBase(bool bInitAsClient, FNetworkNotify* InNoti
 	{
 		NetConnectionClass = UProxyBackendNetConnection::StaticClass();
 		ChildNetConnectionClass = UProxyBackendChildNetConnection::StaticClass();
+
+		if (GetReplicationSystem())
+		{
+			if (UObjectReplicationBridge* Bridge = GetReplicationSystem()->GetReplicationBridge())
+			{
+				using namespace UE::Net::Private;
+
+				OnIrisRootObjectPostInitHandle = FObjectReplicationBridgeDelegates::GetRootObjectPostInitDelegate(Bridge).AddUObject(this, &UProxyBackendNetDriver::OnIrisRootObjectPostInit);
+				OnIrisRootObjectDetachedHandle = FObjectReplicationBridgeDelegates::GetRootObjectDetachedDelegate(Bridge).AddUObject(this, &UProxyBackendNetDriver::OnIrisRootObjectDetached);
+				OnIrisSubObjectPostInitHandle = FObjectReplicationBridgeDelegates::GetSubObjectPostInitDelegate(Bridge).AddUObject(this, &UProxyBackendNetDriver::OnIrisSubObjectPostInit);
+				OnIrisSubObjectDetachedHandle = FObjectReplicationBridgeDelegates::GetSubObjectDetachedDelegate(Bridge).AddUObject(this, &UProxyBackendNetDriver::OnIrisSubObjectDetached);
+			}
+		}
 	}
 
 	SetReplicateTransactionally(false);
@@ -517,6 +530,22 @@ bool UProxyBackendNetDriver::InitBase(bool bInitAsClient, FNetworkNotify* InNoti
 	return bSuccess;
 }
 
+void UProxyBackendNetDriver::Shutdown()
+{
+	if (GetReplicationSystem())
+	{
+		if (UObjectReplicationBridge* Bridge = GetReplicationSystem()->GetReplicationBridge())
+		{
+			UE::Net::Private::FObjectReplicationBridgeDelegates::GetRootObjectPostInitDelegate(Bridge).Remove(OnIrisRootObjectPostInitHandle);
+			UE::Net::Private::FObjectReplicationBridgeDelegates::GetRootObjectDetachedDelegate(Bridge).Remove(OnIrisRootObjectDetachedHandle);
+			UE::Net::Private::FObjectReplicationBridgeDelegates::GetSubObjectPostInitDelegate(Bridge).Remove(OnIrisSubObjectPostInitHandle);
+			UE::Net::Private::FObjectReplicationBridgeDelegates::GetSubObjectDetachedDelegate(Bridge).Remove(OnIrisSubObjectDetachedHandle);
+		}
+	}
+
+	Super::Shutdown();
+}
+
 void UProxyBackendNetDriver::ForwardRemoteFunction(UObject* RootObject, UObject* SubObject, UFunction* Function, void* Parms)
 {
 	check(ProxyNetDriver);
@@ -615,6 +644,136 @@ bool UProxyBackendNetDriver::ShouldRegisterMetricsDatabaseListeners() const
 	return false;
 }
 
+void UProxyBackendNetDriver::OnIrisRootObjectPostInit(UE::Net::FNetRefHandle RootHandle, UObject* RootObject, UE::Net::FNetObjectFactoryId FactoryId)
+{
+	if (!GetReplicationSystem())
+	{
+		return;
+	}
+
+	if (!ProxyNetDriver || !ProxyNetDriver->GetReplicationSystem())
+	{
+		return;
+	}
+	
+	UObjectReplicationBridge* BackendBridge = GetReplicationSystem()->GetReplicationBridgeAs<UObjectReplicationBridge>();
+	if (!BackendBridge)
+	{
+		return;
+	}
+
+	UObjectReplicationBridge* FrontendBridge = ProxyNetDriver->GetReplicationSystem()->GetReplicationBridgeAs<UObjectReplicationBridge>();
+	if (!FrontendBridge)
+	{
+		return;
+	}
+
+	UE_LOG(LogNetProxy, Verbose, TEXT("[%s] Attaching root object %s (%s) to %s"), *GetNameSafe(this), *GetNameSafe(RootObject), *RootHandle.ToString(), *GetNameSafe(ProxyNetDriver));
+	FrontendBridge->StartReplicatingRootObject(RootObject, UObjectReplicationBridge::FRootObjectReplicationParams(), FactoryId);
+}
+
+void UProxyBackendNetDriver::OnIrisRootObjectDetached(UE::Net::FNetRefHandle RootHandle, UObject* RootObject, UE::Net::EDetachReason Reason)
+{
+	if (!GetReplicationSystem())
+	{
+		return;
+	}
+
+	if (!ProxyNetDriver || !ProxyNetDriver->GetReplicationSystem())
+	{
+		return;
+	}
+	
+	UObjectReplicationBridge* BackendBridge = GetReplicationSystem()->GetReplicationBridgeAs<UObjectReplicationBridge>();
+	if (!BackendBridge)
+	{
+		return;
+	}
+
+	UObjectReplicationBridge* FrontendBridge = ProxyNetDriver->GetReplicationSystem()->GetReplicationBridgeAs<UObjectReplicationBridge>();
+	if (!FrontendBridge)
+	{
+		return;
+	}
+
+	UE_LOG(LogNetProxy, Verbose, TEXT("[%s] Detaching root object %s (%s) from %s"), *GetNameSafe(this), *GetNameSafe(RootObject), *RootHandle.ToString(), *GetNameSafe(ProxyNetDriver));
+	const EEndReplicationFlags Flags = (Reason == UE::Net::EDetachReason::TornOff) ? EEndReplicationFlags::TearOff : EEndReplicationFlags::Destroy;
+	FrontendBridge->StopReplicatingNetObject(RootObject, Flags);
+}
+
+void UProxyBackendNetDriver::OnIrisSubObjectPostInit(UE::Net::FNetRefHandle SubObjectHandle, UObject* SubObject, UE::Net::FNetRefHandle RootObjectHandle, UObject* RootObject, UE::Net::FNetObjectFactoryId FactoryId)
+{
+	if (!GetReplicationSystem())
+	{
+		return;
+	}
+
+	if (!ProxyNetDriver || !ProxyNetDriver->GetReplicationSystem())
+	{
+		return;
+	}
+
+	UObjectReplicationBridge* BackendBridge = GetReplicationSystem()->GetReplicationBridgeAs<UObjectReplicationBridge>();
+	if (!BackendBridge)
+	{
+		return;
+	}
+
+	UObjectReplicationBridge* FrontendBridge = ProxyNetDriver->GetReplicationSystem()->GetReplicationBridgeAs<UObjectReplicationBridge>();
+	if (!FrontendBridge)
+	{
+		return;
+	}
+	
+	UE::Net::FNetRefHandle FrontendRootObjectHandle = FrontendBridge->GetReplicatedRefHandle(RootObject);
+	if (FrontendRootObjectHandle.IsValid())
+	{
+		UE_LOG(LogNetProxy, Verbose, TEXT("[%s] Attaching sub-object %s (%s) to %s (%s) on %s"), 
+			   *GetNameSafe(this), 
+			   *GetNameSafe(SubObject), *SubObjectHandle.ToString(), 
+			   *GetNameSafe(RootObject), *RootObjectHandle.ToString(), 
+			   *GetNameSafe(ProxyNetDriver));
+
+		UObjectReplicationBridge::FSubObjectReplicationParams Params;
+		Params.RootObjectHandle = FrontendRootObjectHandle;
+		FrontendBridge->StartReplicatingSubObject(SubObject, Params, FactoryId);
+	}
+}
+
+void UProxyBackendNetDriver::OnIrisSubObjectDetached(UE::Net::FNetRefHandle SubObjectHandle, UObject* SubObject, UE::Net::FNetRefHandle RootObjectHandle, UObject* RootObject, UE::Net::EDetachReason Reason)
+{
+	if (!GetReplicationSystem())
+	{
+		return;
+	}
+
+	if (!ProxyNetDriver || !ProxyNetDriver->GetReplicationSystem())
+	{
+		return;
+	}
+
+	UObjectReplicationBridge* BackendBridge = GetReplicationSystem()->GetReplicationBridgeAs<UObjectReplicationBridge>();
+	if (!BackendBridge)
+	{
+		return;
+	}
+
+	UObjectReplicationBridge* FrontendBridge = ProxyNetDriver->GetReplicationSystem()->GetReplicationBridgeAs<UObjectReplicationBridge>();
+	if (!FrontendBridge)
+	{
+		return;
+	}
+
+	UE_LOG(LogNetProxy, Verbose, TEXT("[%s] Detaching sub-object %s (%s) from %s (%s) on %s"), 
+		   *GetNameSafe(this), 
+		   *GetNameSafe(SubObject), *SubObjectHandle.ToString(),
+		   *GetNameSafe(RootObject), *RootObjectHandle.ToString(),
+		   *GetNameSafe(ProxyNetDriver));
+
+	const EEndReplicationFlags Flags = (Reason == UE::Net::EDetachReason::TornOff) ? EEndReplicationFlags::TearOff : EEndReplicationFlags::Destroy;
+	FrontendBridge->StopReplicatingNetObject(SubObject, Flags);
+}
+
 void UProxyNetDriver::RegisterGameServer(const FURL& GameServerURL)
 {
 	UE_LOG(LogNetProxy, Log, TEXT("Registering Game Server for URL: %s"), *GameServerURL.ToString());
@@ -1013,8 +1172,11 @@ void UProxyNetDriver::Shutdown()
 
 	for (FGameServerConnectionState& ConnectionState : GameServerConnections)
 	{
-		GEngine->DestroyNamedNetDriver(ConnectionState.World, ConnectionState.NetDriverName);
-		ConnectionState.GameServerNotify->NetDriver = nullptr;
+		if (ConnectionState.NetDriver)
+		{
+			GEngine->DestroyNamedNetDriver(ConnectionState.World, ConnectionState.NetDriverName);
+			ConnectionState.GameServerNotify->NetDriver = nullptr;
+		}
 	}
 
 	GameServerConnections.Reset();
diff --git a/Engine/Plugins/Runtime/MultiServerReplication/Source/MultiServerReplication/Public/MultiServerProxy.h b/Engine/Plugins/Runtime/MultiServerReplication/Source/MultiServerReplication/Public/MultiServerProxy.h
--- a/Engine/Plugins/Runtime/MultiServerReplication/Source/MultiServerReplication/Public/MultiServerProxy.h
+++ b/Engine/Plugins/Runtime/MultiServerReplication/Source/MultiServerReplication/Public/MultiServerProxy.h
@@ -6,6 +6,8 @@
 #include "Engine/ChildConnection.h"
 #include "Engine/PackageMapClient.h"
 #include "GameFramework/OnlineReplStructs.h"
+#include "Iris/ReplicationSystem/NetRefHandle.h"
+#include "Iris/ReplicationSystem/NetObjectFactory.h"
 #include "IpNetDriver.h"
 #include "IpConnection.h"
 #include "MultiServerProxy.generated.h"
@@ -279,6 +281,7 @@ public:
 	TObjectPtr<UProxyNetDriver> GetProxyNetDriver();
 
 	virtual bool InitBase(bool bInitAsClient, FNetworkNotify* InNotify, const FURL& URL, bool bReuseAddressAndPort, FString& Error);
+	virtual void Shutdown() override;
 	virtual void ForwardRemoteFunction(UObject* RootObject, UObject* SubObject, UFunction* Function, void* Parms) override;
 	virtual bool ShouldSkipRepNotifies() const override;
 	virtual bool ShouldClientDestroyActor(AActor* Actor, EChannelCloseReason CloseReason) const;
@@ -288,8 +291,18 @@ public:
 
 private:
 
+	void OnIrisRootObjectPostInit(UE::Net::FNetRefHandle RootHandle, UObject* RootObject, UE::Net::FNetObjectFactoryId FactoryId);
+	void OnIrisRootObjectDetached(UE::Net::FNetRefHandle RootHandle, UObject* RootObject, UE::Net::EDetachReason Reason);
+	void OnIrisSubObjectPostInit(UE::Net::FNetRefHandle SubObjectHandle, UObject* SubObject, UE::Net::FNetRefHandle RootObjectHandle, UObject* RootObject, UE::Net::FNetObjectFactoryId FactoryId);
+	void OnIrisSubObjectDetached(UE::Net::FNetRefHandle SubObjectHandle, UObject* SubObject, UE::Net::FNetRefHandle RootObjectHandle, UObject* RootObject, UE::Net::EDetachReason Reason);
+
 	UPROPERTY()
 	TObjectPtr<UProxyNetDriver> ProxyNetDriver;
+
+	FDelegateHandle OnIrisRootObjectPostInitHandle;
+	FDelegateHandle OnIrisRootObjectDetachedHandle;
+	FDelegateHandle OnIrisSubObjectPostInitHandle;
+	FDelegateHandle OnIrisSubObjectDetachedHandle;
 };
 
 /**
diff --git a/Engine/Plugins/Runtime/ReplicationSystemTestPlugin/Source/Private/Tests/ReplicationSystem/ReplicationSystemProxyTestFixture.cpp b/Engine/Plugins/Runtime/ReplicationSystemTestPlugin/Source/Private/Tests/ReplicationSystem/ReplicationSystemProxyTestFixture.cpp
--- a/Engine/Plugins/Runtime/ReplicationSystemTestPlugin/Source/Private/Tests/ReplicationSystem/ReplicationSystemProxyTestFixture.cpp
+++ b/Engine/Plugins/Runtime/ReplicationSystemTestPlugin/Source/Private/Tests/ReplicationSystem/ReplicationSystemProxyTestFixture.cpp
@@ -79,8 +79,21 @@ FReplicationSystemTestClient* FReplicationSystemProxyTestFixture::CreateClient()
 void FReplicationSystemProxyTestFixture::ConnectProxyToBackendServer(FReplicationSystemTestProxy* Proxy, FReplicationSystemTestServer* BackendServer)
 {
 	FReplicationSystemTestClient* BackendClient = Proxy->CreateAndConnectBackendClient(BackendServer);
+	if (BackendClient == nullptr)
+	{
+		return;
+	}
+
 	TArray<FReplicationSystemTestClient*>& ServerClients = ServerToProxyClients.FindOrAdd(BackendServer);
 	ServerClients.Add(BackendClient);
+
+	UObjectReplicationBridge* BackendBridge = Cast<UObjectReplicationBridge>(GetProxyClientForBackendServer(BackendServer, Proxy)->GetReplicationBridge());
+	UObjectReplicationBridge* FrontendBridge = Cast<UObjectReplicationBridge>(Proxy->GetFrontendServer()->GetReplicationBridge());
+	
+	if (ensure(BackendBridge) && ensure(FrontendBridge))
+	{
+		BindProxyFrontendAndBackend(FrontendBridge, BackendBridge);
+	}
 }
 
 void FReplicationSystemProxyTestFixture::ConnectClientToProxy(FReplicationSystemTestClient* Client, FReplicationSystemTestProxy* Proxy)
@@ -167,4 +180,59 @@ FReplicationSystemTestClient* FReplicationSystemProxyTestFixture::GetProxyClient
 	return nullptr;
 }
 
+FOnRootObjectPostInit::RegistrationType& FReplicationSystemProxyTestFixture::GetRootObjectPostInitDelegate(UObjectReplicationBridge* Bridge)
+{
+	return UE::Net::Private::FObjectReplicationBridgeDelegates::GetRootObjectPostInitDelegate(Bridge);
+}
+
+FOnRootObjectDetached::RegistrationType& FReplicationSystemProxyTestFixture::GetRootObjectDetachedDelegate(UObjectReplicationBridge* Bridge)
+{
+	return UE::Net::Private::FObjectReplicationBridgeDelegates::GetRootObjectDetachedDelegate(Bridge);
+}
+
+FOnSubObjectPostInit::RegistrationType& FReplicationSystemProxyTestFixture::GetSubObjectPostInitDelegate(UObjectReplicationBridge* Bridge)
+{
+	return UE::Net::Private::FObjectReplicationBridgeDelegates::GetSubObjectPostInitDelegate(Bridge);
+}
+
+FOnSubObjectDetached::RegistrationType& FReplicationSystemProxyTestFixture::GetSubObjectDetachedDelegate(UObjectReplicationBridge* Bridge)
+{
+	return UE::Net::Private::FObjectReplicationBridgeDelegates::GetSubObjectDetachedDelegate(Bridge);
+}
+
+void FReplicationSystemProxyTestFixture::BindProxyFrontendAndBackend(UObjectReplicationBridge* FrontendBridge, UObjectReplicationBridge* BackendBridge)
+{
+	GetRootObjectPostInitDelegate(BackendBridge).AddLambda(
+		[FrontendBridge](UE::Net::FNetRefHandle RootHandle, UObject* RootObject, UE::Net::FNetObjectFactoryId FactoryId)
+		{
+			FrontendBridge->StartReplicatingRootObject(RootObject, UObjectReplicationBridge::FRootObjectReplicationParams(), FactoryId);
+		});
+
+	GetRootObjectDetachedDelegate(BackendBridge).AddLambda(
+		[FrontendBridge](UE::Net::FNetRefHandle RootHandle, UObject* RootObject, UE::Net::EDetachReason Reason)
+		{
+			const EEndReplicationFlags Flags = (Reason == UE::Net::EDetachReason::TornOff) ? EEndReplicationFlags::TearOff : EEndReplicationFlags::Destroy;
+			FrontendBridge->StopReplicatingNetObject(RootObject, Flags);
+		});
+
+	GetSubObjectPostInitDelegate(BackendBridge).AddLambda(
+		[FrontendBridge](UE::Net::FNetRefHandle SubObjectHandle, UObject* SubObject, UE::Net::FNetRefHandle RootObjectHandle, UObject* RootObject, UE::Net::FNetObjectFactoryId FactoryId)
+		{
+			UE::Net::FNetRefHandle FrontendRootObjectHandle = FrontendBridge->GetReplicatedRefHandle(RootObject);
+			if (ensure(FrontendRootObjectHandle.IsValid()))
+			{
+				UObjectReplicationBridge::FSubObjectReplicationParams Params;
+				Params.RootObjectHandle = FrontendRootObjectHandle;
+				FrontendBridge->StartReplicatingSubObject(SubObject, Params, FactoryId);
+			}
+		});
+
+	GetSubObjectDetachedDelegate(BackendBridge).AddLambda(
+		[FrontendBridge](UE::Net::FNetRefHandle SubObjectHandle, UObject* SubObject, UE::Net::FNetRefHandle RootObjectHandle, UObject* RootObject, UE::Net::EDetachReason Reason)
+		{
+			const EEndReplicationFlags Flags = (Reason == UE::Net::EDetachReason::TornOff) ? EEndReplicationFlags::TearOff : EEndReplicationFlags::Destroy;
+			FrontendBridge->StopReplicatingNetObject(SubObject, Flags);
+		});
+}
+
 }
diff --git a/Engine/Plugins/Runtime/ReplicationSystemTestPlugin/Source/Private/Tests/ReplicationSystem/ReplicationSystemProxyTestFixture.h b/Engine/Plugins/Runtime/ReplicationSystemTestPlugin/Source/Private/Tests/ReplicationSystem/ReplicationSystemProxyTestFixture.h
--- a/Engine/Plugins/Runtime/ReplicationSystemTestPlugin/Source/Private/Tests/ReplicationSystem/ReplicationSystemProxyTestFixture.h
+++ b/Engine/Plugins/Runtime/ReplicationSystemTestPlugin/Source/Private/Tests/ReplicationSystem/ReplicationSystemProxyTestFixture.h
@@ -70,6 +70,16 @@ protected:
 	// Map of frontend proxy server nodes to their corresponding client nodes
 	// Stores non-owning pointers, do not delete them!
 	TMap<FReplicationSystemTestServer*, TArray<FReplicationSystemTestClient*>> ProxyFrontendToClients;
+	
+	static FOnRootObjectPostInit::RegistrationType& GetRootObjectPostInitDelegate(UObjectReplicationBridge* Bridge);
+	static FOnRootObjectDetached::RegistrationType& GetRootObjectDetachedDelegate(UObjectReplicationBridge* Bridge);
+	static FOnSubObjectPostInit::RegistrationType& GetSubObjectPostInitDelegate(UObjectReplicationBridge* Bridge);
+	static FOnSubObjectDetached::RegistrationType& GetSubObjectDetachedDelegate(UObjectReplicationBridge* Bridge);
+
+private:
+
+	// Setup delegates to start and stop replicating objects between the backend and frontend replication systems.
+	void BindProxyFrontendAndBackend(UObjectReplicationBridge* FrontendBridge, UObjectReplicationBridge* BackendBridge);
 };
 
 }
diff --git a/Engine/Plugins/Runtime/ReplicationSystemTestPlugin/Source/Private/Tests/ReplicationSystem/TestReplicationSystemProxy.cpp b/Engine/Plugins/Runtime/ReplicationSystemTestPlugin/Source/Private/Tests/ReplicationSystem/TestReplicationSystemProxy.cpp
--- a/Engine/Plugins/Runtime/ReplicationSystemTestPlugin/Source/Private/Tests/ReplicationSystem/TestReplicationSystemProxy.cpp
+++ b/Engine/Plugins/Runtime/ReplicationSystemTestPlugin/Source/Private/Tests/ReplicationSystem/TestReplicationSystemProxy.cpp
@@ -7,6 +7,89 @@
 namespace UE::Net::Private
 {
 
+UE_NET_TEST_FIXTURE(FReplicationSystemProxyTestFixture, ReplicationDelegates)
+{
+	// Add a backend server
+	FReplicationSystemTestServer* Server = CreateServer();
+
+	// Add a proxy server
+	FReplicationSystemTestProxy* Proxy = CreateProxy();
+
+	// Make connections
+	ConnectProxyToBackendServer(Proxy, Server);
+
+	struct FDelegateCounter
+	{
+		uint32 RootObjectCreate = 0;
+		uint32 RootObjectDetach = 0;
+		uint32 SubObjectCreate = 0;
+		uint32 SubObjectDetach = 0;	
+	};
+
+	FDelegateCounter DelegateCounter;
+
+	// Track the amount of times each delegate is called.
+	UObjectReplicationBridge* BackendBridge = Cast<UObjectReplicationBridge>(GetProxyClientForBackendServer(Server, Proxy)->GetReplicationBridge());
+
+	FDelegateHandle RootObjectPostInitHandle = GetRootObjectPostInitDelegate(BackendBridge).AddLambda(
+		[&DelegateCounter](UE::Net::FNetRefHandle RootHandle, UObject* RootObject, UE::Net::FNetObjectFactoryId FactoryId) 
+		{
+			DelegateCounter.RootObjectCreate++; 
+		});
+	FDelegateHandle RootObjectDetachHandle = GetRootObjectDetachedDelegate(BackendBridge).AddLambda(
+		[&DelegateCounter](UE::Net::FNetRefHandle RootHandle, UObject* RootObject, UE::Net::EDetachReason Reason)
+		{
+			DelegateCounter.RootObjectDetach++; 
+		});
+	FDelegateHandle SubObjectPostInitHandle = GetSubObjectPostInitDelegate(BackendBridge).AddLambda(
+		[&DelegateCounter](UE::Net::FNetRefHandle SubObjectHandle, UObject* SubObject, UE::Net::FNetRefHandle RootObjectHandle, UObject* RootObject, UE::Net::FNetObjectFactoryId FactoryId)
+		{
+			DelegateCounter.SubObjectCreate++;
+		});
+	FDelegateHandle SubObjectDetachedHandle = GetSubObjectDetachedDelegate(BackendBridge).AddLambda(
+		[&DelegateCounter](UE::Net::FNetRefHandle SubObjectHandle, UObject* SubObject, UE::Net::FNetRefHandle RootObjectHandle, UObject* RootObject, UE::Net::EDetachReason Reason)
+		{
+			DelegateCounter.SubObjectDetach++;
+		});
+
+	// Spawn object on server
+	UTestReplicatedIrisObject* ServerObject = Server->CreateObject();
+	UpdateAndSendServerToAllProxies(Server);
+	UE_NET_ASSERT_EQ(DelegateCounter.RootObjectCreate, 1U);
+	UE_NET_ASSERT_EQ(DelegateCounter.RootObjectDetach, 0U);
+	UE_NET_ASSERT_EQ(DelegateCounter.SubObjectCreate, 0U);
+	UE_NET_ASSERT_EQ(DelegateCounter.SubObjectDetach, 0U);
+
+	// Spawn sub-object on server
+	UTestReplicatedIrisObject* ServerSubObject = Server->CreateSubObject(ServerObject->NetRefHandle);
+	UpdateAndSendServerToAllProxies(Server);
+	UE_NET_ASSERT_EQ(DelegateCounter.RootObjectCreate, 1U);
+	UE_NET_ASSERT_EQ(DelegateCounter.RootObjectDetach, 0U);
+	UE_NET_ASSERT_EQ(DelegateCounter.SubObjectCreate, 1U);
+	UE_NET_ASSERT_EQ(DelegateCounter.SubObjectDetach, 0U);
+
+	// Destroy sub-object on server
+	Server->DestroyObject(ServerSubObject);
+	UpdateAndSendServerToAllProxies(Server);
+	UE_NET_ASSERT_EQ(DelegateCounter.RootObjectCreate, 1U);
+	UE_NET_ASSERT_EQ(DelegateCounter.RootObjectDetach, 0U);
+	UE_NET_ASSERT_EQ(DelegateCounter.SubObjectCreate, 1U);
+	UE_NET_ASSERT_EQ(DelegateCounter.SubObjectDetach, 1U);
+
+	// Destroy object on server
+	Server->DestroyObject(ServerObject);
+	UpdateAndSendServerToAllProxies(Server);
+	UE_NET_ASSERT_EQ(DelegateCounter.RootObjectCreate, 1U);
+	UE_NET_ASSERT_EQ(DelegateCounter.RootObjectDetach, 1U);
+	UE_NET_ASSERT_EQ(DelegateCounter.SubObjectCreate, 1U);
+	UE_NET_ASSERT_EQ(DelegateCounter.SubObjectDetach, 1U);
+
+	GetRootObjectPostInitDelegate(BackendBridge).Remove(RootObjectPostInitHandle);
+	GetRootObjectDetachedDelegate(BackendBridge).Remove(RootObjectDetachHandle);
+	GetSubObjectPostInitDelegate(BackendBridge).Remove(SubObjectPostInitHandle);
+	GetSubObjectDetachedDelegate(BackendBridge).Remove(SubObjectDetachedHandle);
+}
+
 UE_NET_TEST_FIXTURE(FReplicationSystemProxyTestFixture, ReplicateAndDestroySingleObject)
 {
 	// Add a backend server
@@ -28,13 +111,14 @@ UE_NET_TEST_FIXTURE(FReplicationSystemProxyTestFixture, ReplicateAndDestroySingl
 	// Send and deliver packet, server -> proxy
 	UpdateAndSendServerToAllProxies(Server);
 
-	// Verify that created server handle now also exists on the proxy
+	// Verify that created server handle now also exists on the proxy backend
 	FReplicationSystemTestClient* ProxyClient = GetProxyClientForBackendServer(Server, Proxy);
-	UTestReplicatedIrisObject* ProxyTestObject = ProxyClient->GetObjectAs<UTestReplicatedIrisObject>(ServerObject->NetRefHandle);
-	UE_NET_ASSERT_NE(ProxyTestObject, nullptr);
+	UTestReplicatedIrisObject* ProxyBackendTestObject = ProxyClient->GetObjectAs<UTestReplicatedIrisObject>(ServerObject->NetRefHandle);
+	UE_NET_ASSERT_NE(ProxyBackendTestObject, nullptr);
 
-	// Let the proxy replicate the object to its clients
-	Proxy->GetFrontendServer()->GetReplicationBridge()->BeginReplication(ProxyTestObject);
+	// Verify the create server handle now also exists in the proxy frontend
+	UTestReplicatedIrisObject* ProxyFrontendTestObject = Proxy->GetFrontendServer()->GetObjectAs<UTestReplicatedIrisObject>(ServerObject->NetRefHandle);
+	UE_NET_ASSERT_NE(ProxyFrontendTestObject, nullptr);
 
 	// Send and deliver packet, proxy -> client
 	UpdateAndSendProxyToAllClients(Proxy);
@@ -51,9 +135,6 @@ UE_NET_TEST_FIXTURE(FReplicationSystemProxyTestFixture, ReplicateAndDestroySingl
 	// Verify that object now is destroyed on proxy as well
 	UE_NET_ASSERT_FALSE(ProxyClient->IsValidNetRefHandle(ServerObject->NetRefHandle));
 
-	// Have the proxy destroy the object on its clients
-	Proxy->GetFrontendServer()->GetReplicationBridge()->StopReplicatingNetRefHandle(ServerObject->NetRefHandle, EEndReplicationFlags::Destroy);
-
 	// Send and deliver packet, proxy -> client
 	UpdateAndSendProxyToAllClients(Proxy);
 
@@ -96,7 +177,7 @@ UE_NET_TEST_FIXTURE(FReplicationSystemProxyTestFixture, TestTwoServers)
 	UpdateAndSendServerToAllProxies(Server1);
 	UpdateAndSendServerToAllProxies(Server2);
 
-	// Verify that objects have been spawned on the proxy
+	// Verify that objects have been spawned on the proxy backend
 	FReplicationSystemTestClient* ProxyBackendClient1 = GetProxyClientForBackendServer(Server1, Proxy);
 	FReplicationSystemTestClient* ProxyBackendClient2 = GetProxyClientForBackendServer(Server2, Proxy);
 	UTestReplicatedIrisObject* ProxyObject1 = ProxyBackendClient1->GetObjectAs<UTestReplicatedIrisObject>(ServerObject1->NetRefHandle);
@@ -109,9 +190,9 @@ UE_NET_TEST_FIXTURE(FReplicationSystemProxyTestFixture, TestTwoServers)
 	UE_NET_ASSERT_EQ(ProxyObject1->IntA, IntServer1);
 	UE_NET_ASSERT_EQ(ProxyObject2->IntA, IntServer2);
 
-	// Let the proxy replicate the object to its clients
-	Proxy->GetFrontendServer()->GetReplicationBridge()->BeginReplication(ProxyObject1);
-	Proxy->GetFrontendServer()->GetReplicationBridge()->BeginReplication(ProxyObject2);
+	// Verify that objects have been added to the proxy frontend
+	UE_NET_ASSERT_TRUE(Proxy->GetFrontendServer()->IsValidNetRefHandle(ProxyObject1->NetRefHandle));
+	UE_NET_ASSERT_TRUE(Proxy->GetFrontendServer()->IsValidNetRefHandle(ProxyObject2->NetRefHandle));
 
 	// Send and deliver packet, proxy -> client
 	UpdateAndSendProxyToAllClients(Proxy);
@@ -177,10 +258,10 @@ UE_NET_TEST_FIXTURE(FReplicationSystemProxyTestFixture, ReuseObjectPropertyRepli
 	// Check property values on the proxy
 	UE_NET_ASSERT_EQ(ProxyRootObject->IntA, RootObjectInt);
 	UE_NET_ASSERT_EQ(ProxySubObject->IntA, SubObjectInt);
-
-	// Let the proxy replicate the object to its client
-	Proxy->GetFrontendServer()->GetReplicationBridge()->BeginReplication(ProxyRootObject);
-	Proxy->GetFrontendServer()->GetReplicationBridge()->BeginReplication(ProxyRootObject->NetRefHandle, ProxySubObject);
+	
+	// Verify that objects have been added to the proxy frontend
+	UE_NET_ASSERT_TRUE(Proxy->GetFrontendServer()->IsValidNetRefHandle(ProxyRootObject->NetRefHandle));
+	UE_NET_ASSERT_TRUE(Proxy->GetFrontendServer()->IsValidNetRefHandle(ProxySubObject->NetRefHandle));
 
 	// Send and deliver packet, proxy -> client
 	UpdateAndSendProxyToAllClients(Proxy);
@@ -300,10 +381,6 @@ UE_NET_TEST_FIXTURE(FReplicationSystemProxyTestFixture, ReuseObjectDestruction)
 	UE_NET_ASSERT_NE(ProxySubObject, nullptr);
 	UE_NET_ASSERT_NE(ProxySubObject, Server1SubObject);
 
-	// Let the proxy replicate the object to its client
-	Proxy->GetFrontendServer()->GetReplicationBridge()->BeginReplication(ProxyRootObject);
-	Proxy->GetFrontendServer()->GetReplicationBridge()->BeginReplication(ProxyRootObject->NetRefHandle, ProxySubObject);
-
 	// Send and deliver packet, proxy -> client
 	UpdateAndSendProxyToAllClients(Proxy);
 
@@ -374,9 +451,6 @@ UE_NET_TEST_FIXTURE(FReplicationSystemProxyTestFixture, ReuseObjectDestruction)
 	UE_NET_ASSERT_EQ(ProxySubObjectAfterSubObjectDestruction, nullptr);
 	UE_NET_ASSERT_FALSE(WeakProxySubObject.IsValid());
 
-	// Explicitly stop replicating on the proxy frontend
-	Proxy->GetFrontendServer()->GetReplicationBridge()->StopReplicatingNetRefHandle(SavedSubObjectHandle, EEndReplicationFlags::Destroy);
-
 	// Send and deliver packet, proxy -> client
 	UpdateAndSendProxyToAllClients(Proxy);
 
@@ -402,9 +476,6 @@ UE_NET_TEST_FIXTURE(FReplicationSystemProxyTestFixture, ReuseObjectDestruction)
 	UE_NET_ASSERT_EQ(ProxyRootObjectAfterRootObjectDestruction, nullptr);
 	UE_NET_ASSERT_FALSE(WeakProxyRootObject.IsValid());
 
-	// Explicitly stop replicating on the proxy frontend
-	Proxy->GetFrontendServer()->GetReplicationBridge()->StopReplicatingNetRefHandle(SavedRootObjectHandle, EEndReplicationFlags::Destroy);
-
 	// Send and deliver packet, proxy -> client
 	UpdateAndSendProxyToAllClients(Proxy);
 
diff --git a/Engine/Source/Runtime/Engine/Public/Tests/NetTestHelpers.h b/Engine/Source/Runtime/Engine/Public/Tests/NetTestHelpers.h
--- a/Engine/Source/Runtime/Engine/Public/Tests/NetTestHelpers.h
+++ b/Engine/Source/Runtime/Engine/Public/Tests/NetTestHelpers.h
@@ -4,6 +4,7 @@
 
 #include "HAL/Platform.h"
 #include "Engine/NetDriver.h"
+#include "UObject/UObjectIterator.h"
 
 #include "Containers/Array.h"
 #include "GameFramework/PlayerController.h"
@@ -45,6 +46,21 @@ int32 NumActors(UWorld* World)
 	return Actors.Num();
 }
 
+template<class T>
+int32 NumObjects(UWorld* World)
+{
+	int32 Count = 0;
+	for (TObjectIterator<T> It; It; ++It)
+	{
+		if (IsValid(*It) && It->GetWorld() == World)
+		{
+			Count++;
+		}
+	}
+
+	return Count;
+}
+
 /*
  * Used to set a cvar and restore it to it's original value when destroyed
  */
diff --git a/Engine/Source/Runtime/Net/Iris/Private/Iris/ReplicationSystem/ObjectReplicationBridge.cpp b/Engine/Source/Runtime/Net/Iris/Private/Iris/ReplicationSystem/ObjectReplicationBridge.cpp
--- a/Engine/Source/Runtime/Net/Iris/Private/Iris/ReplicationSystem/ObjectReplicationBridge.cpp
+++ b/Engine/Source/Runtime/Net/Iris/Private/Iris/ReplicationSystem/ObjectReplicationBridge.cpp
@@ -182,6 +182,26 @@ namespace UE::Net::Private
 		}
 	};
 	typedef TUniquePtr<FReplicationInstanceProtocol, FReplicationInstanceProtocolDeleter> FReplicationInstanceProtocolPtr;
+
+	FOnRootObjectPostInit::RegistrationType& FObjectReplicationBridgeDelegates::GetRootObjectPostInitDelegate(UObjectReplicationBridge* Bridge)
+	{
+		return Bridge->OnRootObjectPostInit;
+	}
+
+	FOnRootObjectDetached::RegistrationType& FObjectReplicationBridgeDelegates::GetRootObjectDetachedDelegate(UObjectReplicationBridge* Bridge)
+	{
+		return Bridge->OnRootObjectDetached;
+	}
+
+	FOnSubObjectPostInit::RegistrationType& FObjectReplicationBridgeDelegates::GetSubObjectPostInitDelegate(UObjectReplicationBridge* Bridge)
+	{
+		return Bridge->OnSubObjectPostInit;
+	}
+
+	FOnSubObjectDetached::RegistrationType& FObjectReplicationBridgeDelegates::GetSubObjectDetachedDelegate(UObjectReplicationBridge* Bridge)
+	{
+		return Bridge->OnSubObjectDetached;
+	}
 }
 
 UObjectReplicationBridge::UObjectReplicationBridge()
@@ -1425,6 +1445,11 @@ void UObjectReplicationBridge::DetachRootObjectFromRemote(FNetRefHandle RootObje
 			.Reason = DetachReason,
 		};
 
+		if (OnSubObjectDetached.IsBound())
+		{
+			OnSubObjectDetached.Broadcast(SubObjectHandle, DetachContext.DetachedInstance, RootObjectHandle, RootObjectInstance, DetachReason);
+		}
+
 		// Tell the root object he will lose this subobject
 		RootObjectFactory->SubObjectDetachedFromReplication(DetachContext, SubObjectDetachContext);
 
@@ -1465,6 +1490,11 @@ void UObjectReplicationBridge::DetachRootObjectFromRemote(FNetRefHandle RootObje
 	}
 	PRAGMA_ENABLE_DEPRECATION_WARNINGS
 	*/
+
+	if (OnRootObjectDetached.IsBound())
+	{
+		OnRootObjectDetached.Broadcast(RootObjectHandle, RootObjectInstance, DetachReason);
+	}
 	
 	const UNetObjectFactory::FDetachContext DetachContext
 	{
@@ -1511,6 +1541,11 @@ void UObjectReplicationBridge::DetachSubObjectFromRemote(FNetRefHandle DetachedS
 	UObject* RootObjectInstance = GetObjectFromReferenceHandle(RootObjectHandle);
 	const FNetRefHandleManager::FReplicatedObjectData& RootObjectData = NetRefHandleManager->GetReplicatedObjectData(RootObjectIndex);
 
+	if (OnSubObjectDetached.IsBound())
+	{
+		OnSubObjectDetached.Broadcast(DetachedSubObjectHandle, DetachedSubObjectInstance, RootObjectHandle, RootObjectInstance, DetachReason);
+	}
+
 	const UNetObjectFactory::FDetachContext DetachContext
 	{
 		.DetachedInstance = DetachedSubObjectInstance,
@@ -1940,7 +1975,8 @@ void UObjectReplicationBridge::PostApplyInitialState(UE::Net::Private::FInternal
 	using namespace UE::Net;
 	using namespace UE::Net::Private;
 
-	const FNetObjectFactoryId FactoryId = NetRefHandleManager->GetReplicatedObjectDataNoCheck(InternalObjectIndex).NetFactoryId;
+	const FNetRefHandleManager::FReplicatedObjectData& ObjectData = NetRefHandleManager->GetReplicatedObjectDataNoCheck(InternalObjectIndex);
+	const FNetObjectFactoryId FactoryId = ObjectData.NetFactoryId;
 	UNetObjectFactory* Factory = GetNetFactory(FactoryId);
 	check(Factory);
 	const UNetObjectFactory::FPostInitContext Context 
@@ -1949,6 +1985,23 @@ void UObjectReplicationBridge::PostApplyInitialState(UE::Net::Private::FInternal
 		.Handle = NetRefHandleManager->GetNetRefHandleFromInternalIndex(InternalObjectIndex)
 	};
 	Factory->PostInit(Context);
+
+	if (ObjectData.IsRootObject())
+	{
+		if (OnRootObjectPostInit.IsBound())
+		{
+			OnRootObjectPostInit.Broadcast(Context.Handle, Context.Instance, FactoryId);
+		}
+	}
+	else
+	{
+		if (OnSubObjectPostInit.IsBound())
+		{
+			const FNetRefHandle RootObjectHandle = NetRefHandleManager->GetNetRefHandleFromInternalIndex(ObjectData.SubObjectRootIndex);
+			UObject* RootObject = NetRefHandleManager->GetReplicatedObjectInstance(ObjectData.SubObjectRootIndex);
+			OnSubObjectPostInit.Broadcast(Context.Handle, Context.Instance, RootObjectHandle, RootObject, FactoryId);
+		}
+	}
 }
 
 void UObjectReplicationBridge::PreSendUpdateSingleHandle(FNetRefHandle RefHandle)
diff --git a/Engine/Source/Runtime/Net/Iris/Public/Iris/ReplicationSystem/ObjectReplicationBridge.h b/Engine/Source/Runtime/Net/Iris/Public/Iris/ReplicationSystem/ObjectReplicationBridge.h
--- a/Engine/Source/Runtime/Net/Iris/Public/Iris/ReplicationSystem/ObjectReplicationBridge.h
+++ b/Engine/Source/Runtime/Net/Iris/Public/Iris/ReplicationSystem/ObjectReplicationBridge.h
@@ -23,6 +23,8 @@
 
 DECLARE_LOG_CATEGORY_EXTERN(LogIrisFilterConfig, Log, All);
 
+class UProxyBackendNetDriver;
+
 namespace UE::Net
 {
 	enum class ENetRefHandleError : uint32;
@@ -48,6 +50,13 @@ namespace UE::Net
 		class FReplicationSystemImpl;
 		class FBridgeSerialization;
 	}
+
+	class FReplicationSystemProxyTestFixture;
+
+	using FOnRootObjectPostInit = TMulticastDelegate<void(FNetRefHandle RootObjectHandle, UObject* RootObject, FNetObjectFactoryId FactoryId)>;
+	using FOnRootObjectDetached = TMulticastDelegate<void(FNetRefHandle RootObjectHandle, UObject* RootObject, EDetachReason DetachReason)>;
+	using FOnSubObjectPostInit = TMulticastDelegate<void(FNetRefHandle SubObjectHandle, UObject* SubObject, FNetRefHandle RootObjectHandle, UObject* RootObject, FNetObjectFactoryId FactoryId)>;
+	using FOnSubObjectDetached = TMulticastDelegate<void(FNetRefHandle SubObjectHandle, UObject* SubObject, FNetRefHandle RootObjectHandle, UObject* RootObject, EDetachReason DetachReason)>;
 }
 
 namespace UE::Net::Private
@@ -69,6 +78,30 @@ struct FStaticDestructionInfo
 	UE::Net::Private::FInternalNetRefIndex InternalReplicationIndex;
 };
 
+/** 
+ * Delegates that are triggered as part of root object and sub-object lifecycles in the replication bridge.
+ *
+ * These delegates are currently used to enable a proxy server to start and stop replicating objects in the frontend
+ * replication system when they are attached or detached from a backend replication system.
+ */
+class FObjectReplicationBridgeDelegates
+{
+private:
+
+	/** (Client only) After a root object has been attached and initialized. */
+	IRISCORE_API static FOnRootObjectPostInit::RegistrationType& GetRootObjectPostInitDelegate(UObjectReplicationBridge* Bridge);
+	/** (Client only) Before a root object is detached; any sub-objects that are part of this root will also be detached before the root. */
+	IRISCORE_API static FOnRootObjectDetached::RegistrationType& GetRootObjectDetachedDelegate(UObjectReplicationBridge* Bridge);
+	/** (Client only) After a sub-object has been attached and initialized. */
+	IRISCORE_API static FOnSubObjectPostInit::RegistrationType& GetSubObjectPostInitDelegate(UObjectReplicationBridge* Bridge);
+	/** (Client only) Before a sub-object is detached, either on it's own or before it's parent root object is detached. */
+	IRISCORE_API static FOnSubObjectDetached::RegistrationType& GetSubObjectDetachedDelegate(UObjectReplicationBridge* Bridge);
+
+	friend class ::UProxyBackendNetDriver;
+	friend class UE::Net::FReplicationSystemProxyTestFixture;
+
+};
+
 } // end namespace UE::Net::Private
 
 struct FObjectReplicationBridgeInstantiateResult
@@ -598,6 +631,8 @@ private:
 	friend UE::Net::Private::FObjectPoller;
 	friend UE::Net::Private::FReplicationSystemImpl;
 	
+	friend UE::Net::Private::FObjectReplicationBridgeDelegates;
+
 	friend UReplicationSystem;
 
 	TMap<const UClass*, FName> ConfigClassPathNameCache;
@@ -711,6 +746,11 @@ private:
 	/** Protects access to CachedCreationHeaders inside WriteNetRefHandleCreationInfo */
 	FCriticalSection WriteNetRefHandleCreationCS;
 #endif // UE_SUPPORT_PARALLEL_IRIS
+
+	UE::Net::FOnRootObjectPostInit OnRootObjectPostInit;
+	UE::Net::FOnRootObjectDetached OnRootObjectDetached;
+	UE::Net::FOnSubObjectPostInit OnSubObjectPostInit;
+	UE::Net::FOnSubObjectDetached OnSubObjectDetached;
 };
 
 //------------------------------------------------------------------------
