Iris
- Fixed internal leaks of subobjects being torned off.
- Fixed edge cases related to subobject conditionals flush and tear-off
- Added detection and warning if internal net objects does not get destroyed.

#Synced-CL 50186587

diff --git a/Engine/Plugins/Runtime/ReplicationSystemTestPlugin/Source/Private/Tests/ReplicationSystem/ReplicationSystemServerClientTestFixture.cpp b/Engine/Plugins/Runtime/ReplicationSystemTestPlugin/Source/Private/Tests/ReplicationSystem/ReplicationSystemServerClientTestFixture.cpp
--- a/Engine/Plugins/Runtime/ReplicationSystemTestPlugin/Source/Private/Tests/ReplicationSystem/ReplicationSystemServerClientTestFixture.cpp
+++ b/Engine/Plugins/Runtime/ReplicationSystemTestPlugin/Source/Private/Tests/ReplicationSystem/ReplicationSystemServerClientTestFixture.cpp
@@ -509,16 +509,16 @@ bool FReplicationSystemTestClient::UpdateAndSend(FReplicationSystemTestServer* S
 	bool bSuccess = false;
 
 	NetUpdate();
+	const bool bSentPacket = SendUpdate(Desc);
+	PostSendUpdate();
 
-	if (SendUpdate(Desc))
+	if (bSentPacket)
 	{
 		constexpr uint32 ServerRemoteConnectionId = 0x01;
 		DeliverTo(*Server, LocalConnectionId, ServerRemoteConnectionId, bDeliver);
 		bSuccess = true;
 	}
 
-	PostSendUpdate();
-
 	return bSuccess;
 }
 
@@ -563,18 +563,33 @@ void FReplicationSystemTestServer::DeliverTo(FReplicationSystemTestClient* Clien
 
 bool FReplicationSystemTestServer::UpdateAndSend(const TArrayView<FReplicationSystemTestClient*const>& Clients, bool bDeliver /*= true*/, const TCHAR* Desc)
 {
-	bool bSuccess = true;
+	bool bSentPackets = false;
 
 	NetUpdate();
 
+	TArray<FReplicationSystemTestClient*> ClientsWithData;
+	ClientsWithData.Reserve(Clients.Num());
+
+	FNetBitArray SentPackets;
+	SentPackets.Init(Clients.Num());
+
 	for (FReplicationSystemTestClient* Client : Clients)
 	{
-		bSuccess &= SendAndDeliverTo(Client, bDeliver, Desc);
+		if (SendTo(Client, Desc))
+		{
+			ClientsWithData.Add(Client);
+		}
 	}
 	
 	PostSendUpdate();
 
-	return bSuccess;
+	for (FReplicationSystemTestClient* Client : ClientsWithData)
+	{
+		DeliverTo(Client, bDeliver);
+		bSentPackets = true;
+	}
+
+	return bSentPackets;
 }
 
 //*****************************************************************************
diff --git a/Engine/Plugins/Runtime/ReplicationSystemTestPlugin/Source/Private/Tests/ReplicationSystem/TestReplicationSystem.cpp b/Engine/Plugins/Runtime/ReplicationSystemTestPlugin/Source/Private/Tests/ReplicationSystem/TestReplicationSystem.cpp
--- a/Engine/Plugins/Runtime/ReplicationSystemTestPlugin/Source/Private/Tests/ReplicationSystem/TestReplicationSystem.cpp
+++ b/Engine/Plugins/Runtime/ReplicationSystemTestPlugin/Source/Private/Tests/ReplicationSystem/TestReplicationSystem.cpp
@@ -1091,9 +1091,33 @@ UE_NET_TEST_FIXTURE(FReplicationSystemServerClientTestFixture, TestNetTemporary)
 }
 
 // Tests for TearOff
+class FTearOffAndFlushTestFixture : public FReplicationSystemServerClientTestFixture
+{
+public:
+	void UpdateAndVerifyThatAllInternalIndicesAreReleased(uint32 MaxIterattions = 3)
+	{
+		for (uint32 It = 0U; It < MaxIterattions; ++It)
+		{
+			Server->UpdateAndSend({Clients});
+		}
+		
+		bool bHasInternalObjectsPendingDestroy = false;
+		// Verify that we have released all internal objects and have no pending references.
+		{
+			UE::Net::Private::FNetRefHandleManager* NetRefHandleManager = &Server->ReplicationSystem->GetReplicationSystemInternal()->GetNetRefHandleManager();
+			UE_NET_ASSERT_EQ(NetRefHandleManager->GetObjectsPendingDestroy().Num(), 0);
+		}
+
+		for (FReplicationSystemTestClient* Client : Clients)
+		{
+			UE::Net::Private::FNetRefHandleManager* NetRefHandleManager = &Client->ReplicationSystem->GetReplicationSystemInternal()->GetNetRefHandleManager();
+			UE_NET_ASSERT_EQ(NetRefHandleManager->GetObjectsPendingDestroy().Num(), 0);
+		}
+	}
+};
 
 // Test TearOff for existing confirmed object
-UE_NET_TEST_FIXTURE(FReplicationSystemServerClientTestFixture, TestTearOffExistingObject)
+UE_NET_TEST_FIXTURE(FTearOffAndFlushTestFixture, TestTearOffExistingObject)
 {
 	UReplicationSystem* ReplicationSystem = Server->ReplicationSystem;
 
@@ -1134,7 +1158,7 @@ UE_NET_TEST_FIXTURE(FReplicationSystemServerClientTestFixture, TestTearOffExisti
 }
 
 // Test TearOff for new object
-UE_NET_TEST_FIXTURE(FReplicationSystemServerClientTestFixture, TestTearOffOnNewlyCreatedObject)
+UE_NET_TEST_FIXTURE(FTearOffAndFlushTestFixture, TestTearOffOnNewlyCreatedObject)
 {
 	UReplicationSystem* ReplicationSystem = Server->ReplicationSystem;
 
@@ -1171,7 +1195,7 @@ UE_NET_TEST_FIXTURE(FReplicationSystemServerClientTestFixture, TestTearOffOnNewl
 }
 
 // Test TearOff resend for existing confirmed object with no state changes
-UE_NET_TEST_FIXTURE(FReplicationSystemServerClientTestFixture, TestTearOffResendForExistingObjectWithoutDirtyState)
+UE_NET_TEST_FIXTURE(FTearOffAndFlushTestFixture, TestTearOffResendForExistingObjectWithoutDirtyState)
 {
 	UReplicationSystem* ReplicationSystem = Server->ReplicationSystem;
 
@@ -1215,7 +1239,7 @@ UE_NET_TEST_FIXTURE(FReplicationSystemServerClientTestFixture, TestTearOffResend
 }
 
 // Test TearOff for new object and resend, this requires creation info to be cached.
-UE_NET_TEST_FIXTURE(FReplicationSystemServerClientTestFixture, TestTearOffImmediateOnNewlyCreatedObjectResend)
+UE_NET_TEST_FIXTURE(FTearOffAndFlushTestFixture, TestTearOffImmediateOnNewlyCreatedObjectResend)
 {
 	UReplicationSystem* ReplicationSystem = Server->ReplicationSystem;
 
@@ -1253,7 +1277,7 @@ UE_NET_TEST_FIXTURE(FReplicationSystemServerClientTestFixture, TestTearOffImmedi
 }
 
 // Test TearOff for new object and resend if not fitting in packet, this requires creation info to be cached and that we retain object in scope
-UE_NET_TEST_FIXTURE(FReplicationSystemServerClientTestFixture, TestTearOffImmediateOnNewlyCreatedObjectIfNotSentOnFirstUpdate)
+UE_NET_TEST_FIXTURE(FTearOffAndFlushTestFixture, TestTearOffImmediateOnNewlyCreatedObjectIfNotSentOnFirstUpdate)
 {
 	UReplicationSystem* ReplicationSystem = Server->ReplicationSystem;
 
@@ -1291,7 +1315,7 @@ UE_NET_TEST_FIXTURE(FReplicationSystemServerClientTestFixture, TestTearOffImmedi
 }
 
 // Test TearOff for new subobject and resend, this requires creation info to be cached.
-UE_NET_TEST_FIXTURE(FReplicationSystemServerClientTestFixture, TestTearOffImmediateOnNewlyCreatedSubObjectResend)
+UE_NET_TEST_FIXTURE(FTearOffAndFlushTestFixture, TestTearOffImmediateOnNewlyCreatedSubObjectResend)
 {
 	UReplicationSystem* ReplicationSystem = Server->ReplicationSystem;
 
@@ -1339,7 +1363,7 @@ UE_NET_TEST_FIXTURE(FReplicationSystemServerClientTestFixture, TestTearOffImmedi
 	UE_NET_ASSERT_EQ(ServerSubObject->IntA, ClientSubObjectThatWillBeTornOff->IntA);
 }
 
-UE_NET_TEST_FIXTURE(FReplicationSystemServerClientTestFixture, TestDefferedTearOffOnNewlyCreatedObjectResend)
+UE_NET_TEST_FIXTURE(FTearOffAndFlushTestFixture, TestDefferedTearOffOnNewlyCreatedObjectResend)
 {
 	UReplicationSystem* ReplicationSystem = Server->ReplicationSystem;
 
@@ -1384,7 +1408,7 @@ UE_NET_TEST_FIXTURE(FReplicationSystemServerClientTestFixture, TestDefferedTearO
 }
 
 // Test TearOff for existing not yet confirmed object
-UE_NET_TEST_FIXTURE(FReplicationSystemServerClientTestFixture, TestTearOffObjectPendingCreateConfirmation)
+UE_NET_TEST_FIXTURE(FTearOffAndFlushTestFixture, TestTearOffObjectPendingCreateConfirmation)
 {
 	UReplicationSystem* ReplicationSystem = Server->ReplicationSystem;
 
@@ -1425,7 +1449,7 @@ UE_NET_TEST_FIXTURE(FReplicationSystemServerClientTestFixture, TestTearOffObject
 }
 
 // Test TearOff for existing object pending destroy (should do nothing)
-UE_NET_TEST_FIXTURE(FReplicationSystemServerClientTestFixture, TestTearOffExistingObjectPendingDestroy)
+UE_NET_TEST_FIXTURE(FTearOffAndFlushTestFixture, TestTearOffExistingObjectPendingDestroy)
 {
 	UReplicationSystem* ReplicationSystem = Server->ReplicationSystem;
 
@@ -1469,7 +1493,7 @@ UE_NET_TEST_FIXTURE(FReplicationSystemServerClientTestFixture, TestTearOffExisti
 }
 
 // Test TearOff resend 
-UE_NET_TEST_FIXTURE(FReplicationSystemServerClientTestFixture, TestTearOffResend)
+UE_NET_TEST_FIXTURE(FTearOffAndFlushTestFixture, TestTearOffResend)
 {
 	UReplicationSystem* ReplicationSystem = Server->ReplicationSystem;
 
@@ -1518,7 +1542,7 @@ UE_NET_TEST_FIXTURE(FReplicationSystemServerClientTestFixture, TestTearOffResend
 }
 
 // Test TearOff does not pickup statechanges after tear off
-UE_NET_TEST_FIXTURE(FReplicationSystemServerClientTestFixture, TestTornedOffObjectDoesNotCopyStateChanges)
+UE_NET_TEST_FIXTURE(FTearOffAndFlushTestFixture, TestTornedOffObjectDoesNotCopyStateChanges)
 {
 	UReplicationSystem* ReplicationSystem = Server->ReplicationSystem;
 
@@ -1567,7 +1591,7 @@ UE_NET_TEST_FIXTURE(FReplicationSystemServerClientTestFixture, TestTornedOffObje
 }
 
 // Test TearOff and SubObjects, SubObjects must apply state?
-UE_NET_TEST_FIXTURE(FReplicationSystemServerClientTestFixture, TestImmediateTearOffExistingObjectWithSubObject)
+UE_NET_TEST_FIXTURE(FTearOffAndFlushTestFixture, TestImmediateTearOffExistingObjectWithSubObject)
 {
 	UReplicationSystem* ReplicationSystem = Server->ReplicationSystem;
 
@@ -1585,9 +1609,7 @@ UE_NET_TEST_FIXTURE(FReplicationSystemServerClientTestFixture, TestImmediateTear
 	ServerSubObject->IntA = 1;
 
 	// Send and deliver packet
-	Server->NetUpdate();
-	Server->SendAndDeliverTo(Client, true);
-	Server->PostSendUpdate();
+	Server->UpdateAndSend({Client});
 
 	// Store Pointer to objects
 	UTestReplicatedIrisObject* ClientObjectThatWillBeTornOff = Cast<UTestReplicatedIrisObject>(Client->GetReplicationBridge()->GetReplicatedObject(ServerObject->NetRefHandle));
@@ -1605,17 +1627,420 @@ UE_NET_TEST_FIXTURE(FReplicationSystemServerClientTestFixture, TestImmediateTear
 	Server->ReplicationBridge->EndReplication(ServerObject, EEndReplicationFlags::TearOff);
 
 	// Send and deliver packet
-	Server->NetUpdate();
-	Server->SendAndDeliverTo(Client, true);
-	Server->PostSendUpdate();
+	Server->UpdateAndSend({Client});
 
 	// Verify that ClientObject is torn-off and that the final state was applied to subObject 
 	UE_NET_ASSERT_EQ(ServerSubObject->IntA, ClientSubObjectThatWillBeTornOff->IntA);
-	UE_NET_ASSERT_TRUE(Cast<UTestReplicatedIrisObject>(Client->GetReplicationBridge()->GetReplicatedObject(ServerSubObject->NetRefHandle)) == nullptr);
+	UE_NET_ASSERT_TRUE(Client->GetReplicationBridge()->GetReplicatedObject(ServerSubObject->NetRefHandle) == nullptr);
+
+	// Verify that we have released all internal objects and have no pending references.
+	UpdateAndVerifyThatAllInternalIndicesAreReleased();
+}
+
+/**
+ * Below follows a series of tests that ensure that we do not leak internal net objects when doing tear off and flush in combination with subobject filtering
+ */
+UE_NET_TEST_FIXTURE(FTearOffAndFlushTestFixture, TestAddCondNeverSubObjectAndTearOffSameFrame)
+{
+	UReplicationSystem* ReplicationSystem = Server->ReplicationSystem;
+
+	// Add a client
+	FReplicationSystemTestClient* Client = CreateClient();
+
+	// Spawn object on server
+	UTestReplicatedIrisObject* ServerObject = Server->CreateObject(0,0);
+
+	// Send and deliver packet
+	Server->UpdateAndSend({Client});
+
+	// Spawn second object on server as a subobject
+	UTestReplicatedIrisObject* ServerSubObject = Server->CreateSubObject(ServerObject->NetRefHandle, 0, 0);
+	Server->GetReplicationBridge()->SetSubObjectNetCondition(ServerSubObject->NetRefHandle, ELifetimeCondition::COND_Never);
+
+	// Store Pointer to objects
+	UE_NET_ASSERT_TRUE(Client->GetReplicationBridge()->GetReplicatedObject(ServerObject->NetRefHandle) != nullptr);
+
+	// TearOff the object using immediate tear-off
+	Server->ReplicationBridge->EndReplication(ServerObject, EEndReplicationFlags::TearOff);
+
+	// Send and deliver packet
+	Server->UpdateAndSend({Client});
+
+	// Verify that we cannot find the objects on the client
+	UE_NET_ASSERT_TRUE(Client->GetReplicationBridge()->GetReplicatedObject(ServerObject->NetRefHandle) == nullptr);
+	UE_NET_ASSERT_TRUE(Client->GetReplicationBridge()->GetReplicatedObject(ServerSubObject->NetRefHandle) == nullptr);
+
+	// Verify that we have released all internal objects and have no pending references.
+	UpdateAndVerifyThatAllInternalIndicesAreReleased();
+}
+
+UE_NET_TEST_FIXTURE(FTearOffAndFlushTestFixture, TestAddCondNeverSubObjectAndTearOffSameFrameDelayBeforeSend)
+{
+	UReplicationSystem* ReplicationSystem = Server->ReplicationSystem;
+
+	// Add a client
+	FReplicationSystemTestClient* Client = CreateClient();
+
+	// Spawn object on server
+	UTestReplicatedIrisObject* ServerObject = Server->CreateObject(0,0);
+
+	// Send and deliver packet
+	Server->UpdateAndSend({Client});
+
+	// Spawn second object on server as a subobject
+	UTestReplicatedIrisObject* ServerSubObject = Server->CreateSubObject(ServerObject->NetRefHandle, 0, 0);
+	Server->GetReplicationBridge()->SetSubObjectNetCondition(ServerSubObject->NetRefHandle, ELifetimeCondition::COND_Never);
+
+	// Store Pointer to objects
+	UE_NET_ASSERT_TRUE(Client->GetReplicationBridge()->GetReplicatedObject(ServerObject->NetRefHandle) != nullptr);
+
+	// TearOff the object using immediate tear-off
+	Server->ReplicationBridge->EndReplication(ServerObject, EEndReplicationFlags::TearOff);
+
+	// Do a few updates without sending
+	Server->UpdateAndSend({});
+	Server->UpdateAndSend({});
+
+	// Send and deliver packet
+	Server->UpdateAndSend({Client});
+
+	// Verify that we cannot find the objects on the client
+	UE_NET_ASSERT_TRUE(Client->GetReplicationBridge()->GetReplicatedObject(ServerObject->NetRefHandle) == nullptr);
+	UE_NET_ASSERT_TRUE(Client->GetReplicationBridge()->GetReplicatedObject(ServerSubObject->NetRefHandle) == nullptr);
+
+	// Verify that we have released all internal objects and have no pending references.
+	UpdateAndVerifyThatAllInternalIndicesAreReleased();
+}
+
+UE_NET_TEST_FIXTURE(FTearOffAndFlushTestFixture, TestAddCondNeverSubObjectAndTearOffSameFrameLostPacket)
+{
+	UReplicationSystem* ReplicationSystem = Server->ReplicationSystem;
+
+	// Add a client
+	FReplicationSystemTestClient* Client = CreateClient();
+
+	// Spawn object on server
+	UTestReplicatedIrisObject* ServerObject = Server->CreateObject(0,0);
+
+	// Send and deliver packet
+	Server->UpdateAndSend({Client});
+
+	// Spawn second object on server as a subobject
+	UTestReplicatedIrisObject* ServerSubObject = Server->CreateSubObject(ServerObject->NetRefHandle, 0, 0);
+	Server->GetReplicationBridge()->SetSubObjectNetCondition(ServerSubObject->NetRefHandle, ELifetimeCondition::COND_Never);
+
+	// Store Pointer to objects
+	UE_NET_ASSERT_TRUE(Client->GetReplicationBridge()->GetReplicatedObject(ServerObject->NetRefHandle) != nullptr);
+
+	// TearOff the object using immediate tear-off
+	Server->ReplicationBridge->EndReplication(ServerObject, EEndReplicationFlags::TearOff);
+
+	// Send and drop packet
+	Server->UpdateAndSend({Client}, /* bDeliver = */ false);
+
+	// Update without sending
+	Server->UpdateAndSend({});
+
+	// Send and deliver packet
+	Server->UpdateAndSend({Client});
+
+	// Verify that we cannot find the objectson the client
+	UE_NET_ASSERT_TRUE(Client->GetReplicationBridge()->GetReplicatedObject(ServerObject->NetRefHandle) == nullptr);
+	UE_NET_ASSERT_TRUE(Client->GetReplicationBridge()->GetReplicatedObject(ServerSubObject->NetRefHandle) == nullptr);
+
+	// Verify that we have released all internal objects and have no pending references.
+	UpdateAndVerifyThatAllInternalIndicesAreReleased();
+}
+
+UE_NET_TEST_FIXTURE(FTearOffAndFlushTestFixture, TestAddCondNeverSubObjectAndDestroyRootWithFlushRootSameFrame)
+{
+	UReplicationSystem* ReplicationSystem = Server->ReplicationSystem;
+
+	// Add a client
+	FReplicationSystemTestClient* Client = CreateClient();
+
+	// Spawn object on server
+	UTestReplicatedIrisObject* ServerObject = Server->CreateObject(0,0);
+
+	// Send and deliver packet
+	Server->UpdateAndSend({Client});
+
+	// Spawn second object on server as a subobject
+	UTestReplicatedIrisObject* ServerSubObject = Server->CreateSubObject(ServerObject->NetRefHandle, 0, 0);
+	Server->GetReplicationBridge()->SetSubObjectNetCondition(ServerSubObject->NetRefHandle, ELifetimeCondition::COND_Never);
+
+	Server->ReplicationBridge->EndReplication(ServerObject, EEndReplicationFlags::Destroy | EEndReplicationFlags::Flush);
+
+	// Send and deliver packet
+	Server->UpdateAndSend({Client});
+
+	// Verify that we have released all internal objects and have no pending references.
+	UpdateAndVerifyThatAllInternalIndicesAreReleased();
+}
+
+UE_NET_TEST_FIXTURE(FTearOffAndFlushTestFixture, TestCreateRootAndCondNeverSubObjectAndDestroyRootWithFlushRootSameFrame)
+{
+	LOG_SCOPE_VERBOSITY_OVERRIDE(LogIris, ELogVerbosity::VeryVerbose);
+
+	UReplicationSystem* ReplicationSystem = Server->ReplicationSystem;
+
+	// Add a client
+	FReplicationSystemTestClient* Client = CreateClient();
+
+	// Spawn object on server
+	UTestReplicatedIrisObject* ServerObject = Server->CreateObject(0,0);
+
+	// Spawn second object on server as a subobject
+	UTestReplicatedIrisObject* ServerSubObject = Server->CreateSubObject(ServerObject->NetRefHandle, 0, 0);
+	Server->GetReplicationBridge()->SetSubObjectNetCondition(ServerSubObject->NetRefHandle, ELifetimeCondition::COND_Never);
+
+	// Send and deliver packet
+	Server->UpdateAndSend({});
+
+	Server->ReplicationBridge->EndReplication(ServerObject, EEndReplicationFlags::Destroy | EEndReplicationFlags::Flush);
+
+	// Send and deliver packet
+	Server->UpdateAndSend({});
+
+	// Send and deliver packet
+	Server->UpdateAndSend({Client});
+
+	// Send and deliver packet
+	Server->UpdateAndSend({Client});
+
+	// Verify that we have released all internal objects and have no pending references.
+	UpdateAndVerifyThatAllInternalIndicesAreReleased();
+}
+
+
+UE_NET_TEST_FIXTURE(FTearOffAndFlushTestFixture, TestAddCondNeverSubObjectAndDestroyRootWithFlushRootSameFrameDelayBeforeSend)
+{
+	UReplicationSystem* ReplicationSystem = Server->ReplicationSystem;
+
+	// Add a client
+	FReplicationSystemTestClient* Client = CreateClient();
+
+	// Spawn object on server
+	UTestReplicatedIrisObject* ServerObject = Server->CreateObject(0,0);
+
+	// Send and deliver packet
+	Server->UpdateAndSend({Client});
+
+	// Spawn second object on server as a subobject
+	UTestReplicatedIrisObject* ServerSubObject = Server->CreateSubObject(ServerObject->NetRefHandle, 0, 0);
+	Server->GetReplicationBridge()->SetSubObjectNetCondition(ServerSubObject->NetRefHandle, ELifetimeCondition::COND_Never);
+
+	Server->ReplicationBridge->EndReplication(ServerObject, EEndReplicationFlags::Destroy | EEndReplicationFlags::Flush);
+
+	// Do a few updates without sending
+	Server->UpdateAndSend({});
+
+	// Verify that we have released all internal objects and have no pending references.
+	UpdateAndVerifyThatAllInternalIndicesAreReleased();
+}
+
+// Verify that we release all references to filtered out subobjects
+UE_NET_TEST_FIXTURE(FTearOffAndFlushTestFixture, TestAddCondNeverSubObjectAndDestroyRootWithFlushRootSameFrameLostPacket)
+{
+	UReplicationSystem* ReplicationSystem = Server->ReplicationSystem;
+
+	// Add a client
+	FReplicationSystemTestClient* Client = CreateClient();
+
+	// Spawn object on server
+	UTestReplicatedIrisObject* ServerObject = Server->CreateObject(0,0);
+
+	// Send and deliver packet
+	Server->UpdateAndSend({Client});
+
+	// Spawn second object on server as a subobject
+	UTestReplicatedIrisObject* ServerSubObject = Server->CreateSubObject(ServerObject->NetRefHandle, 0, 0);
+	Server->GetReplicationBridge()->SetSubObjectNetCondition(ServerSubObject->NetRefHandle, ELifetimeCondition::COND_Never);
+
+	Server->ReplicationBridge->EndReplication(ServerObject, EEndReplicationFlags::Destroy | EEndReplicationFlags::Flush);
+
+	// Send and drop packet
+	Server->UpdateAndSend({Client}, /* bDeliver = */ false);
+
+	// Send and deliver packet
+	Server->UpdateAndSend({Client});
+
+	// Verify that we have released all internal objects and have no pending references.
+	UpdateAndVerifyThatAllInternalIndicesAreReleased();
+}
+
+UE_NET_TEST_FIXTURE(FTearOffAndFlushTestFixture, TestFilterOutCreatedCondNeverSubObjectAndDestroyRootWithFlush)
+{
+	UReplicationSystem* ReplicationSystem = Server->ReplicationSystem;
+
+	// Add a client
+	FReplicationSystemTestClient* Client = CreateClient();
+
+	// Spawn object on server
+	UTestReplicatedIrisObject* ServerObject = Server->CreateObject(0,0);
+
+	// Spawn second object on server as a subobject
+	UTestReplicatedIrisObject* ServerSubObject = Server->CreateSubObject(ServerObject->NetRefHandle, 0, 0);
+
+	// Send and deliver packet
+	Server->UpdateAndSend({Client});
+
+	// Verify that subobject exists
+	UTestReplicatedIrisObject* ClientSubObjectThatWillBeDestroyedWithFlush = Cast<UTestReplicatedIrisObject>(Client->GetReplicationBridge()->GetReplicatedObject(ServerSubObject->NetRefHandle));
+	UE_NET_ASSERT_TRUE(ClientSubObjectThatWillBeDestroyedWithFlush != nullptr);
+
+	// Now filter out the subobject, as it was created before being filtered out it will be destroyed
+	Server->GetReplicationBridge()->SetSubObjectNetCondition(ServerSubObject->NetRefHandle, ELifetimeCondition::COND_Never);
+
+	// Just dirty it
+	++ServerSubObject->IntA;
+
+	// End replication of root.
+	Server->ReplicationBridge->EndReplication(ServerObject, EEndReplicationFlags::Destroy | EEndReplicationFlags::Flush);
+
+	// Verify that we have released all internal objects and have no pending references.
+	UpdateAndVerifyThatAllInternalIndicesAreReleased();
+
+	// Verify that we did not get the final value as object was filtered out, we get away with this due the the fact that we keep all objects created on the client until GC
+	UE_NET_ASSERT_NE(ClientSubObjectThatWillBeDestroyedWithFlush->IntA, ServerSubObject->IntA);
+}
+
+UE_NET_TEST_FIXTURE(FTearOffAndFlushTestFixture, TestFilterOutCreatedCondNeverSubObjectAndTearOff)
+{
+	UReplicationSystem* ReplicationSystem = Server->ReplicationSystem;
+
+	// Add a client
+	FReplicationSystemTestClient* Client = CreateClient();
+
+	// Spawn object on server
+	UTestReplicatedIrisObject* ServerObject = Server->CreateObject(0,0);
+
+	// Spawn second object on server as a subobject
+	UTestReplicatedIrisObject* ServerSubObject = Server->CreateSubObject(ServerObject->NetRefHandle, 0, 0);
+
+	// Send and deliver packet
+	Server->UpdateAndSend({Client});
+
+	// Verify that subobject exists
+	UTestReplicatedIrisObject* ClientSubObjectThatWillBeDestroyedWithFlush = Cast<UTestReplicatedIrisObject>(Client->GetReplicationBridge()->GetReplicatedObject(ServerSubObject->NetRefHandle));
+	UE_NET_ASSERT_TRUE(ClientSubObjectThatWillBeDestroyedWithFlush != nullptr);
+
+	// Now filter out the subobject, but as we have created it must still replicate it with the tear-off.
+	Server->GetReplicationBridge()->SetSubObjectNetCondition(ServerSubObject->NetRefHandle, ELifetimeCondition::COND_Never);
+
+	// Just dirty it
+	++ServerSubObject->IntA;
+
+	// End replication of root.
+	Server->ReplicationBridge->EndReplication(ServerObject, EEndReplicationFlags::TearOff);
+
+	// Send and deliver packet
+	Server->UpdateAndSend({Client});
+
+	// Verify that we have released all internal objects and have no pending references.
+	UpdateAndVerifyThatAllInternalIndicesAreReleased();
+
+	// Verify that we did not get the final value as object was filtered out, we get away with this due the the fact that we keep all objects created on the client until GC
+	UE_NET_ASSERT_NE(ClientSubObjectThatWillBeDestroyedWithFlush->IntA, ServerSubObject->IntA);
+}
+
+UE_NET_TEST_FIXTURE(FTearOffAndFlushTestFixture, TestFilterOutCondNeverSubObjectAndRenabledBeforeDestroyRootWithFlush)
+{
+	LOG_SCOPE_VERBOSITY_OVERRIDE(LogIris, ELogVerbosity::VeryVerbose);
+
+	UReplicationSystem* ReplicationSystem = Server->ReplicationSystem;
+
+	// Add a client
+	FReplicationSystemTestClient* Client = CreateClient();
+
+	// Spawn object on server
+	UTestReplicatedIrisObject* ServerObject = Server->CreateObject(0,0);
+
+	// Spawn second object on server as a subobject
+	UTestReplicatedIrisObject* ServerSubObject = Server->CreateSubObject(ServerObject->NetRefHandle, 0, 0);
+
+	// Filter out the subobject
+	Server->GetReplicationBridge()->SetSubObjectNetCondition(ServerSubObject->NetRefHandle, ELifetimeCondition::COND_Never);
+
+	// Send and deliver packet
+	Server->UpdateAndSend({Client});
+
+	// Verify that subobject does not exists
+	UE_NET_ASSERT_TRUE(Client->GetReplicationBridge()->GetReplicatedObject(ServerSubObject->NetRefHandle) == nullptr);
+
+	// Re-enable filtered out the subobject
+	Server->GetReplicationBridge()->SetSubObjectNetCondition(ServerSubObject->NetRefHandle, ELifetimeCondition::COND_None);
+
+	// End replication of root.
+	Server->ReplicationBridge->EndReplication(ServerObject, EEndReplicationFlags::Destroy | EEndReplicationFlags::Flush);
+
+	// Send and deliver packet
+	Server->UpdateAndSend({Client});
+
+	// Verify that subobject exists
+	UE_NET_ASSERT_TRUE(Client->GetReplicationBridge()->GetReplicatedObject(ServerSubObject->NetRefHandle) != nullptr);
+
+	Server->UpdateAndSend({Client});
+
+	// Verify that subobject does not exists
+	UE_NET_ASSERT_TRUE(Client->GetReplicationBridge()->GetReplicatedObject(ServerSubObject->NetRefHandle) == nullptr);
+
+	// Verify that we have released all internal objects and have no pending references.
+	UpdateAndVerifyThatAllInternalIndicesAreReleased();
+}
+
+UE_NET_TEST_FIXTURE(FTearOffAndFlushTestFixture, TestFilterOutCondNeverSubObjectAndRenabledBeforeTearOff)
+{
+	UReplicationSystem* ReplicationSystem = Server->ReplicationSystem;
+
+	// Add a client
+	FReplicationSystemTestClient* Client = CreateClient();
+
+	// Spawn object on server
+	UTestReplicatedIrisObject* ServerObject = Server->CreateObject(0,0);
+
+	// Spawn second object on server as a subobject
+	UTestReplicatedIrisObject* ServerSubObject = Server->CreateSubObject(ServerObject->NetRefHandle, 0, 0);
+
+	// Filter out the subobject
+	Server->GetReplicationBridge()->SetSubObjectNetCondition(ServerSubObject->NetRefHandle, ELifetimeCondition::COND_Never);
+
+	// Send and deliver packet
+	Server->UpdateAndSend({Client});
+
+	// Verify that subobject does not exists
+	UE_NET_ASSERT_TRUE(Client->GetReplicationBridge()->GetReplicatedObject(ServerSubObject->NetRefHandle) == nullptr);
+
+	// Re-enable filtered out the subobject
+	Server->GetReplicationBridge()->SetSubObjectNetCondition(ServerSubObject->NetRefHandle, ELifetimeCondition::COND_None);
+
+	// Set a value that we can find to verify that the SubObject included state
+	ServerSubObject->IntA = 3;
+
+	// End replication of root.
+	Server->ReplicationBridge->EndReplication(ServerObject, EEndReplicationFlags::TearOff);
+
+	const int32 PreUpdateCreatedObjectCount = Client->CreatedObjects.Num();
+
+	// Send and deliver packet
+	Server->UpdateAndSend({Client});
+
+	const int32 PostUpdateCreatedObjectCount = Client->CreatedObjects.Num();
+	UE_NET_ASSERT_GT(PostUpdateCreatedObjectCount, PreUpdateCreatedObjectCount);
+
+	// We expect the SubObject to have been replicated so it should exist in the persistent CreatedObjects array on the client
+	UTestReplicatedIrisObject* ClientSubObject = Cast<UTestReplicatedIrisObject>(Client->CreatedObjects[PostUpdateCreatedObjectCount - 1].Get());
+	UE_NET_ASSERT_EQ(ServerSubObject->IntA, ClientSubObject->IntA);
+
+	// Verify that subobject does not exists as a replicated object
+	UE_NET_ASSERT_TRUE(Client->GetReplicationBridge()->GetReplicatedObject(ServerSubObject->NetRefHandle) == nullptr);
+
+	// Verify that we have released all internal objects and have no pending references.
+	UpdateAndVerifyThatAllInternalIndicesAreReleased();
 }
 
 // Test TearOff and SubObjects, SubObjects must apply state?
-UE_NET_TEST_FIXTURE(FReplicationSystemServerClientTestFixture, TestImmediateTearOffExistingObjectWithSubObjectDroppedData)
+UE_NET_TEST_FIXTURE(FTearOffAndFlushTestFixture, TestImmediateTearOffExistingObjectWithSubObjectDroppedData)
 {
 	UReplicationSystem* ReplicationSystem = Server->ReplicationSystem;
 
@@ -1664,13 +2089,12 @@ UE_NET_TEST_FIXTURE(FReplicationSystemServerClientTestFixture, TestImmediateTear
 
 	// Verify that ClientObject is torn-off and that the final state was applied to subObject 
 	UE_NET_ASSERT_EQ(ServerSubObject->IntA, ClientSubObjectThatWillBeTornOff->IntA);
-	UE_NET_ASSERT_TRUE(Cast<UTestReplicatedIrisObject>(Client->GetReplicationBridge()->GetReplicatedObject(ServerSubObject->NetRefHandle)) == nullptr);
+	UE_NET_ASSERT_TRUE(Client->GetReplicationBridge()->GetReplicatedObject(ServerSubObject->NetRefHandle) == nullptr);
 }
 
 // Test creating a subobject after the root is already torn off
-UE_NET_TEST_FIXTURE(FReplicationSystemServerClientTestFixture, TestCreatingSubobjectAfterRootObjectTearOff)
+UE_NET_TEST_FIXTURE(FTearOffAndFlushTestFixture, TestCreatingSubobjectAfterRootObjectTearOff)
 {
-
 	// Add a client
 	FReplicationSystemTestClient* Client = CreateClient();
 
@@ -1902,7 +2326,7 @@ UE_NET_TEST_FIXTURE(FReplicationSystemServerClientTestFixture, TestReplicatedSub
 
 
 // Test tear-off object in PendingCreate state to ensure that tear-off logic works as expected
-UE_NET_TEST_FIXTURE(FReplicationSystemServerClientTestFixture, TestTearOffObjectWithNoFragmentsDoesNotTriggerCheckIfPendingCreateWhenDestroyed)
+UE_NET_TEST_FIXTURE(FTearOffAndFlushTestFixture, TestTearOffObjectWithNoFragmentsDoesNotTriggerCheckIfPendingCreateWhenDestroyed)
 {
 	UReplicationSystem* ReplicationSystem = Server->ReplicationSystem;
 
@@ -1925,7 +2349,7 @@ UE_NET_TEST_FIXTURE(FReplicationSystemServerClientTestFixture, TestTearOffObject
 }
 
 // Test tear-off subobject in PendingCreate state to ensure that tear-off logic works as expected
-UE_NET_TEST_FIXTURE(FReplicationSystemServerClientTestFixture, TestTearOffSubObjectWithNoFragmentsDoesNotTriggerCheckIfPendingCreateWhenDestroyed)
+UE_NET_TEST_FIXTURE(FTearOffAndFlushTestFixture, TestTearOffSubObjectWithNoFragmentsDoesNotTriggerCheckIfPendingCreateWhenDestroyed)
 {
 	UReplicationSystem* ReplicationSystem = Server->ReplicationSystem;
 
@@ -1952,7 +2376,7 @@ UE_NET_TEST_FIXTURE(FReplicationSystemServerClientTestFixture, TestTearOffSubObj
 }
 
 // Test TearOff and SubObjects, SubObjects must apply state?
-UE_NET_TEST_FIXTURE(FReplicationSystemServerClientTestFixture, TestTearOffNextUpdateExistingObjectWithSubObject)
+UE_NET_TEST_FIXTURE(FTearOffAndFlushTestFixture, TestTearOffNextUpdateExistingObjectWithSubObject)
 {
 	UReplicationSystem* ReplicationSystem = Server->ReplicationSystem;
 
@@ -2000,7 +2424,7 @@ UE_NET_TEST_FIXTURE(FReplicationSystemServerClientTestFixture, TestTearOffNextUp
 }
 
 // Test TearOff and destroy of SubObjects that are still pending create/tearoff
-UE_NET_TEST_FIXTURE(FReplicationSystemServerClientTestFixture, TestTearOffNextUpdateExistingObjectWithSubObjectPendingCreation)
+UE_NET_TEST_FIXTURE(FTearOffAndFlushTestFixture, TestTearOffNextUpdateExistingObjectWithSubObjectPendingCreation)
 {
 	UReplicationSystem* ReplicationSystem = Server->ReplicationSystem;
 	FNetRefHandleManager* NetRefHandleManager = &ReplicationSystem->GetReplicationSystemInternal()->GetNetRefHandleManager();
@@ -2620,6 +3044,9 @@ UE_NET_TEST_FIXTURE(FReplicationSystemServerClientTestFixture, TestStablyNamedOb
 	// Send and deliver packet
 	Server->UpdateAndSend({Client});
 
+	// Send and deliver packet
+	Server->UpdateAndSend({Client});
+
 	// We expect object to no longer exist on client and server
 	UE_NET_ASSERT_TRUE(Cast<UTestReplicatedIrisObject>(Client->GetReplicationBridge()->GetReplicatedObject(ServerObject->NetRefHandle)) == nullptr);
 	UE_NET_ASSERT_TRUE(Cast<UTestReplicatedIrisObject>(Server->GetReplicationBridge()->GetReplicatedObject(ServerObject->NetRefHandle)) == nullptr);
@@ -2702,6 +3129,9 @@ UE_NET_TEST_FIXTURE(FReplicationSystemServerClientTestFixture, TestStablyNamedOb
 	// Update and send to all 
 	Server->UpdateAndSend(Clients);
 
+	// Send and deliver packet
+	Server->UpdateAndSend(Clients);
+
 	// We now expect object to be considered fully destroyed.
 	UE_NET_ASSERT_TRUE(Cast<UTestReplicatedIrisObject>(Clients[0]->GetReplicationBridge()->GetReplicatedObject(ServerObject->NetRefHandle)) == nullptr);
 	UE_NET_ASSERT_TRUE(Cast<UTestReplicatedIrisObject>(Clients[1]->GetReplicationBridge()->GetReplicatedObject(ServerObject->NetRefHandle)) == nullptr);
@@ -2756,6 +3186,9 @@ UE_NET_TEST_FIXTURE(FReplicationSystemServerClientTestFixture, TestDynamicObject
 	// Send and deliver packet
 	Server->UpdateAndSend({Client});
 
+	// Send and deliver packet, really to process delivery notifications.
+	Server->UpdateAndSend({Client});
+
 	// We expect object to no longer exist on client and server
 	UE_NET_ASSERT_TRUE(Cast<UTestReplicatedIrisObject>(Client->GetReplicationBridge()->GetReplicatedObject(ServerObject->NetRefHandle)) == nullptr);
 	UE_NET_ASSERT_TRUE(Cast<UTestReplicatedIrisObject>(Server->GetReplicationBridge()->GetReplicatedObject(ServerObject->NetRefHandle)) == nullptr);
@@ -2771,8 +3204,6 @@ UE_NET_TEST_FIXTURE(FReplicationSystemServerClientTestFixture, TestDynamicObject
 
 UE_NET_TEST_FIXTURE(FReplicationSystemServerClientTestFixture, TestDynamicObjectAsyncStopReplicationForSubObject)
 {
-	LOG_SCOPE_VERBOSITY_OVERRIDE(LogIris, ELogVerbosity::Verbose);
-
 	UReplicationSystem* ReplicationSystem = Server->ReplicationSystem;
 
 	// Add a client
@@ -2835,6 +3266,9 @@ UE_NET_TEST_FIXTURE(FReplicationSystemServerClientTestFixture, TestDynamicObject
 	UE_NET_ASSERT_TRUE(Cast<UTestReplicatedIrisObject>(Client->GetReplicationBridge()->GetReplicatedObject(ServerObject->NetRefHandle)) != nullptr);
 	UE_NET_ASSERT_TRUE(Cast<UTestReplicatedIrisObject>(Server->GetReplicationBridge()->GetReplicatedObject(ServerObject->NetRefHandle)) != nullptr);
 
+	// Send and deliver packet
+	Server->UpdateAndSend({Client});
+
 	// We expect SubObject to no longer exist on client and server
 	UE_NET_ASSERT_TRUE(Cast<UReplicatedTestObject>(Client->GetReplicationBridge()->GetReplicatedObject(ServerSubObject->NetRefHandle)) == nullptr);
 	UE_NET_ASSERT_TRUE(Cast<UReplicatedTestObject>(Server->GetReplicationBridge()->GetReplicatedObject(ServerSubObject->NetRefHandle)) == nullptr);
@@ -2917,6 +3351,9 @@ UE_NET_TEST_FIXTURE(FReplicationSystemServerClientTestFixture, TestDynamicObject
 	// Send and deliver packet to all clients
 	Server->UpdateAndSend(Clients);
 
+	// Send and deliver packet to all clients
+	Server->UpdateAndSend(Clients);
+
 	// We expect object to no longer exist and no longer be pending end replication
 	UE_NET_ASSERT_TRUE(Cast<UTestReplicatedIrisObject>(Clients[0]->GetReplicationBridge()->GetReplicatedObject(ServerObject->NetRefHandle)) == nullptr);
 	UE_NET_ASSERT_TRUE(Cast<UTestReplicatedIrisObject>(Clients[1]->GetReplicationBridge()->GetReplicatedObject(ServerObject->NetRefHandle)) == nullptr);
@@ -3023,8 +3460,6 @@ UE_NET_TEST_FIXTURE(FReplicationSystemServerClientTestFixture, TestStablyNamedOb
 /** Test combinations of start / stop / start for stable object using flush. */
 UE_NET_TEST_FIXTURE(FReplicationSystemServerClientTestFixture, TestStableStartFlushStartStopStartReplication)
 {
-	LOG_SCOPE_VERBOSITY_OVERRIDE(LogIris, ELogVerbosity::Verbose);
-
 	const TCHAR* StableObjectName = TEXT("StableObjectName");
 
 	UReplicationSystem* ReplicationSystem = Server->ReplicationSystem;
@@ -3098,8 +3533,6 @@ UE_NET_TEST_FIXTURE(FReplicationSystemServerClientTestFixture, TestStableStartFl
 
 UE_NET_TEST_FIXTURE(FReplicationSystemServerClientTestFixture, TestStablyNamedObjectMultiStartStopStartReplicationWithTwoClients)
 {
-	LOG_SCOPE_VERBOSITY_OVERRIDE(LogIris, ELogVerbosity::Verbose);
-
 	const TCHAR* StableObjectName = TEXT("StableObjectName");
 
 	UReplicationSystem* ReplicationSystem = Server->ReplicationSystem;
diff --git a/Engine/Source/Runtime/Net/Iris/Private/Iris/ReplicationSystem/NetRefHandleManager.cpp b/Engine/Source/Runtime/Net/Iris/Private/Iris/ReplicationSystem/NetRefHandleManager.cpp
--- a/Engine/Source/Runtime/Net/Iris/Private/Iris/ReplicationSystem/NetRefHandleManager.cpp
+++ b/Engine/Source/Runtime/Net/Iris/Private/Iris/ReplicationSystem/NetRefHandleManager.cpp
@@ -40,6 +40,13 @@ static FAutoConsoleVariableRef CVarAllowCreateNetObjectForNetRefHandlePendingDes
 	TEXT("If set to true we allow creation of a new NetObject for a NetRefHandle pending destroy that might still be replicating data to client. This is currently not fully supported and can lead to issues on the client.")
 );
 
+static int32 WarnIfPendingDestroyUpdateCount = 256;
+static FAutoConsoleVariableRef CVarWarnIfPendingDestroyUpdateCount(
+	TEXT("net.Iris.WarnIfPendingDestroyUpdateCount"),
+	WarnIfPendingDestroyUpdateCount,
+	TEXT("Output warning if an NetObject is stuck in pending destroy for more than WarnIfPendingDestroyUpdateCount updates.")
+);
+
 static_assert(GetDetachReasonBitsNeeded() <= 8U, "EInternalDetachReason is a uint8");
 
 const TCHAR* LexToString(EInternalDetachReason DetachReason)
@@ -652,8 +659,8 @@ void FNetRefHandleManager::DestroyNetObjectByIndex(FInternalNetRefIndex Internal
 	{
 		FReplicatedObjectData& Data = ReplicatedObjectData[InternalIndex];
 
-		// If object already is marked for flush, we only need to remove it from scope.
-		if (!Data.bFlush)
+		// If object already is marked for flush or pending destroy, we only need to remove it from scope.
+		if (!(Data.bFlush || Data.bNetObjectIsPendingDestroy))
 		{
 			// Just to catch misuse by bad code.
 			FInternalNetRefIndex FoundInternalIndex = RefHandleToInternalIndex.FindAndRemoveChecked(Data.RefHandle);
@@ -668,10 +675,13 @@ void FNetRefHandleManager::DestroyNetObjectByIndex(FInternalNetRefIndex Internal
 			NetHandleToInternalIndex.Remove(Data.NetHandle);
 
 			// We always defer the actual destroy
-			PendingDestroyInfos.Emplace(InternalIndex, Data.RefHandle);
+			PendingDestroyInfos.Emplace(InternalIndex, Data.RefHandle, UpdatePendingDestroyCounter);
 
 			// Before we can start replicating this RefHandle again, we need to first finish replication of previous incarnations of the object.
 			InternalIndicesReplicatingNetRefHandle.FindOrAdd(Data.RefHandle).Add(InternalIndex);
+
+			// Track that this object is marked for destroy so we can avoid calling it again
+			Data.bNetObjectIsPendingDestroy = true;
 		}
 
 		if (bKeepInScopeAndMarkForFlush)
@@ -717,6 +727,18 @@ void FNetRefHandleManager::DestroyObjectsPendingDestroy()
 
 			PendingDestroyInfos.RemoveAtSwap(It);
 		}
+		else if ((WarnIfPendingDestroyUpdateCount > 0) && uint32(UpdatePendingDestroyCounter - PendingDestroyInfos[It].PendingDestroyUpdateCount) >= uint32(WarnIfPendingDestroyUpdateCount))
+		{
+			UE_LOG(LogIris, Warning, TEXT("FNetRefHandleManager::DestroyObjectsPendingDestroy: Object %s IsReferenced %u"), *PrintObjectFromIndex(InternalIndex), ReplicatedObjectRefCount[InternalIndex]);
+			if (ReplicatedObjectRefCount[InternalIndex] == 0)
+			{
+				for (FInternalNetRefIndex SubObjectIndex : GetSubObjects(InternalIndex))
+				{
+					UE_LOG(LogIris, Warning, TEXT("FNetRefHandleManager::DestroyObjectsPendingDestroy: Due to SubObjectsObject %s IsReferenced %u"), *PrintObjectFromIndex(SubObjectIndex), ReplicatedObjectRefCount[SubObjectIndex]);
+				}
+			}
+			PendingDestroyInfos[It].PendingDestroyUpdateCount = UpdatePendingDestroyCounter;	
+		}
 	}
 	CSV_CUSTOM_STAT(IrisCommon, PendingDestroyInternalIndicesCount, (float)PendingDestroyInfos.Num(), ECsvCustomStatOp::Set);
 
@@ -724,6 +746,9 @@ void FNetRefHandleManager::DestroyObjectsPendingDestroy()
 	{
 		OnInternalNetRefIndicesFreed.Broadcast(MakeConstArrayView(FreedInternalIndices));
 	}
+
+	// We use a counter instead of time to avoid issues when debugging
+	++UpdatePendingDestroyCounter;
 }
 
 bool FNetRefHandleManager::AddSubObject(FNetRefHandle RootObjectHandle, FNetRefHandle SubObjectHandle, FNetRefHandle RelativeOtherSubObjectHandle, EAddSubObjectFlags Flags)
@@ -1427,19 +1452,19 @@ FString FNetRefHandleManager::PrintObjectFromIndex(FInternalNetRefIndex ObjectIn
 			return FString::Printf(TEXT("UnassignedObject (InternalIndex: %u)"), ObjectIndex);
 		}
 
-		const FNetRefHandle NetRefHandle = GetNetRefHandleFromInternalIndex(ObjectIndex);
 		const FReplicatedObjectData& ObjectData = GetReplicatedObjectDataNoCheck(ObjectIndex);
+		const FNetRefHandle NetRefHandle = ObjectData.RefHandle;
 
 		if (ObjectData.SubObjectRootIndex == InvalidInternalIndex)
 		{
-			return FString::Printf(TEXT("RootObject %s (InternalIndex: %u) (%s)"), *GetNameSafe(ReplicatedInstances[ObjectIndex]), ObjectIndex, *NetRefHandle.ToString());
+			return FString::Printf(TEXT("RootObject %s (InternalIndex: %u) (%s) Protocol: (%s)"), *GetNameSafe(ReplicatedInstances[ObjectIndex]), ObjectIndex, *NetRefHandle.ToString(), ObjectData.Protocol->DebugName->Name);
 		}
 		else
 		{
 			const FNetRefHandle RootNetRefHandle = GetNetRefHandleFromInternalIndex(ObjectData.SubObjectRootIndex);
-			return FString::Printf(TEXT("SubObject %s (InternalIndex: %u) (%s) tied to RootObject %s (InternalIndex: %u) (%s)"), 
+			return FString::Printf(TEXT("SubObject %s (InternalIndex: %u) (%s) tied to RootObject %s (InternalIndex: %u) (%s) Protocol: (%s)"), 
 								*GetNameSafe(ReplicatedInstances[ObjectIndex]), ObjectIndex, *NetRefHandle.ToString(),
-								*GetNameSafe(ReplicatedInstances[ObjectData.SubObjectRootIndex]), ObjectData.SubObjectRootIndex, *RootNetRefHandle.ToString());
+								*GetNameSafe(ReplicatedInstances[ObjectData.SubObjectRootIndex]), ObjectData.SubObjectRootIndex, *RootNetRefHandle.ToString(), ObjectData.Protocol->DebugName->Name);
 		}
 	}
 	else
diff --git a/Engine/Source/Runtime/Net/Iris/Private/Iris/ReplicationSystem/NetRefHandleManager.h b/Engine/Source/Runtime/Net/Iris/Private/Iris/ReplicationSystem/NetRefHandleManager.h
--- a/Engine/Source/Runtime/Net/Iris/Private/Iris/ReplicationSystem/NetRefHandleManager.h
+++ b/Engine/Source/Runtime/Net/Iris/Private/Iris/ReplicationSystem/NetRefHandleManager.h
@@ -131,6 +131,7 @@ public:
 		uint16 bPendingEndReplication : 1U = 0;
 		uint16 bHasCachedCreationInfo : 1U = 0;
 		uint16 bIsDebugObject : 1U = 0;
+		uint16 bNetObjectIsPendingDestroy : 1U = 0;
 
 		/** The factory responsible for instantiating this object */
 		FNetObjectFactoryId NetFactoryId = InvalidNetObjectFactoryId;
@@ -233,6 +234,7 @@ public:
 	{
 		FInternalNetRefIndex InternalIndex;
 		FNetRefHandle RefHandle;
+		uint32 PendingDestroyUpdateCount = 0U;
 	};
 
 	TConstArrayView<FPendingDestroyInfo> GetObjectsPendingDestroy() const { return MakeArrayView(PendingDestroyInfos); }
@@ -635,6 +637,9 @@ private:
 	// Protects parallel access to NextHandleId and RefHandleToInternalIndex when allocating NetRefHandles
 	FTransactionallySafeCriticalSection HandleSerialCriticalSection;
 #endif // UE_SUPPORT_PARALLEL_IRIS
+
+	// We have a running counter of update for debugging purposes
+	uint32 UpdatePendingDestroyCounter = 0U;
 };
 
 const FNetRefHandleManager::FReplicatedObjectData& FNetRefHandleManager::GetReplicatedObjectData(FInternalNetRefIndex InternalIndex) const
diff --git a/Engine/Source/Runtime/Net/Iris/Private/Iris/ReplicationSystem/ReplicationBridge.cpp b/Engine/Source/Runtime/Net/Iris/Private/Iris/ReplicationSystem/ReplicationBridge.cpp
--- a/Engine/Source/Runtime/Net/Iris/Private/Iris/ReplicationSystem/ReplicationBridge.cpp
+++ b/Engine/Source/Runtime/Net/Iris/Private/Iris/ReplicationSystem/ReplicationBridge.cpp
@@ -240,7 +240,7 @@ void UReplicationBridge::InternalDetachInstanceFromNetRefHandle(UE::Net::Private
 
 	if (FReplicationInstanceProtocol* InstanceProtocol = const_cast<FReplicationInstanceProtocol*>(NetRefHandleManager->DetachInstanceProtocol(InternalObjectIndex)))
 	{
-		UE_LOG_REPLICATIONBRIDGE(Verbose, TEXT("InternalDetachInstanceToNetHandle Detached: %s"), *NetRefHandleManager->PrintObjectFromIndex(InternalObjectIndex));
+		UE_LOG_REPLICATIONBRIDGE(Verbose, TEXT("InternalDetachInstanceFromNetHandle Detached: %s"), *NetRefHandleManager->PrintObjectFromIndex(InternalObjectIndex));
 
 		if (bUnBindInstanceProtocol && EnumHasAnyFlags(InstanceProtocol->InstanceTraits, EReplicationInstanceProtocolTraits::IsBound))
 		{
@@ -312,11 +312,11 @@ void UReplicationBridge::InternalFinishDestroyLocalObject(UE::Net::Private::FInt
 	// If the object are a member of any groups we need to remove it to make sure that we update filtering
 	GetReplicationSystem()->GetReplicationSystemInternal()->RemoveFromAllGroups(InternalReplicationIndex);
 
+	// Destroy NetObject and remove it from all scoping.
+	InternalDestroyNetObject(InternalReplicationIndex);
+
 	// If we have any attached SubObjects, finish destroy for them as well.
 	InternalFinishDestroyLocalSubObjects(InternalReplicationIndex, EndReplicationFlags);
-
-	// Finally Destroy NetObject and remove it from all scoping.
-	InternalDestroyNetObject(InternalReplicationIndex);
 }
 
 void UReplicationBridge::DestroyLocalNetHandle(FNetRefHandle Handle, EEndReplicationFlags EndReplicationFlags)
@@ -373,17 +373,11 @@ void UReplicationBridge::InternalStartDestroyLocalSubObjects(UE::Net::Private::F
 	{
 		for (FInternalNetRefIndex SubObjectInternalIndex : NetRefHandleManager->GetChildSubObjects(OwnerInternalIndex))
 		{
-			// Tag subobject for destroy. The check against the scope is needed since the subobjects array might contain subobjects already pending destroy.
-			if (NetRefHandleManager->IsScopableIndex(SubObjectInternalIndex))
+			FNetRefHandleManager::FReplicatedObjectData& SubObjectData = NetRefHandleManager->GetReplicatedObjectDataNoCheck(SubObjectInternalIndex);
+			if (!(SubObjectData.bFlush || SubObjectData.bNetObjectIsPendingDestroy))
 			{
-				FNetRefHandleManager::FReplicatedObjectData& SubObjectData = NetRefHandleManager->GetReplicatedObjectDataNoCheck(SubObjectInternalIndex);
-
 				SubObjectData.bPendingEndReplication = 1U;
-				if (!SubObjectData.bFlush)
-				{
-					UE_LOG_REPLICATIONBRIDGE(Verbose, TEXT("InternalStartDestroySubObjects %s"), *NetRefHandleManager->PrintObjectFromIndex(SubObjectInternalIndex));	
-					InternalStartDestroyLocalObject(SubObjectInternalIndex, Flags);				
-				}
+				InternalStartDestroyLocalObject(SubObjectInternalIndex, Flags);				
 			}
 		}
 	}
@@ -398,14 +392,9 @@ void UReplicationBridge::InternalFinishDestroyLocalSubObjects(UE::Net::Private::
 	{
 		for (FInternalNetRefIndex SubObjectInternalIndex : NetRefHandleManager->GetChildSubObjects(OwnerInternalIndex))
 		{
-			// Tag subobject for destroy. The check against the scope is needed since the subobjects array might contain subobjects already pending destroy.
-			if (NetRefHandleManager->IsScopableIndex(SubObjectInternalIndex))
-			{
-				UE_LOG_REPLICATIONBRIDGE(Verbose, TEXT("InternalFinishDestroySubObjects %s"), *NetRefHandleManager->PrintObjectFromIndex(SubObjectInternalIndex));
-				FNetRefHandleManager::FReplicatedObjectData& SubObjectData = NetRefHandleManager->GetReplicatedObjectDataNoCheck(SubObjectInternalIndex);
-				SubObjectData.bPendingEndReplication = 1U;
-				InternalFinishDestroyLocalObject(SubObjectInternalIndex, Flags);
-			}
+			FNetRefHandleManager::FReplicatedObjectData& SubObjectData = NetRefHandleManager->GetReplicatedObjectDataNoCheck(SubObjectInternalIndex);
+			SubObjectData.bPendingEndReplication = 1U;
+			InternalFinishDestroyLocalObject(SubObjectInternalIndex, Flags);
 		}
 	}
 }
diff --git a/Engine/Source/Runtime/Net/Iris/Private/Iris/ReplicationSystem/ReplicationWriter.cpp b/Engine/Source/Runtime/Net/Iris/Private/Iris/ReplicationSystem/ReplicationWriter.cpp
--- a/Engine/Source/Runtime/Net/Iris/Private/Iris/ReplicationSystem/ReplicationWriter.cpp
+++ b/Engine/Source/Runtime/Net/Iris/Private/Iris/ReplicationSystem/ReplicationWriter.cpp
@@ -719,7 +719,18 @@ uint32 FReplicationWriter::GetFlushStatus(uint32 InternalIndex, const FReplicati
 
 	if (FlushFlagsToTest == EFlushFlags::FlushFlags_None)
 	{
-		return FlushFlags;
+		return EFlushFlags::FlushFlags_None;
+	}
+
+	if (Info.IsSubObject)
+	{
+		// If owner is PendingDestroy it means that this subobject has been filtered out and should not trigger a flush.
+		const FNetRefHandleManager::FReplicatedObjectData& ObjectData = NetRefHandleManager->GetReplicatedObjectDataNoCheck(InternalIndex);
+		const FReplicationInfo& OwnerInfo = ReplicatedObjects[ObjectData.SubObjectRootIndex];
+		if (OwnerInfo.GetState() == EReplicatedObjectState::PendingDestroy)
+		{
+			return EFlushFlags::FlushFlags_None;
+		}
 	}
 
 	if (!!(FlushFlagsToTest & EFlushFlags::FlushFlags_FlushState) && (Info.HasDirtyChangeMask || HasInFlightStateChanges(InternalIndex, Info) || IsObjectPartOfActiveHugeObject(InternalIndex)))
@@ -740,8 +751,11 @@ uint32 FReplicationWriter::GetFlushStatus(uint32 InternalIndex, const FReplicati
 
 	if (!Info.IsSubObject && FlushFlags != FlushFlagsToTest)
 	{
-		// Check status of SubObjects as well.
-		for (uint32 SubObjectIndex : NetRefHandleManager->GetSubObjects(InternalIndex))
+		// Check status of SubObjects that are not filtered out as well.
+		FReplicationConditionals::FSubObjectsToReplicateArray SubObjectsToReplicate;
+		ReplicationConditionals->GetSubObjectsToReplicate(Parameters.ConnectionId, InternalIndex, SubObjectsToReplicate);
+
+		for (uint32 SubObjectIndex : SubObjectsToReplicate)
 		{
 			const FReplicationInfo& SubObjectInfo = GetReplicationInfo(SubObjectIndex);
 			FlushFlags |= GetFlushStatus(SubObjectIndex, SubObjectInfo, FlushFlagsToTest);
@@ -1002,11 +1016,6 @@ void FReplicationWriter::UpdateDirtyGlobalLifetimeConditionals(TArrayView<FInter
 
 	for (FInternalNetRefIndex InternalObjectIndex : ObjectsWithDirtyConditionals)
 	{
-		// Better safe than sorry, we do not want to dirty something going out of scope.
-		if (!ObjectsInScope.GetBit(InternalObjectIndex))
-		{
-			continue;
-		}
 		FReplicationInfo& Info = ReplicatedObjects[InternalObjectIndex];
 		if (Info.GetState() != EReplicatedObjectState::Invalid && Info.GetState() < EReplicatedObjectState::PendingDestroy)
 		{
@@ -1017,11 +1026,15 @@ void FReplicationWriter::UpdateDirtyGlobalLifetimeConditionals(TArrayView<FInter
 				const FInternalNetRefIndex RootInternalObjectIndex = ObjectData.SubObjectRootIndex;
 				FReplicationInfo& OwnerInfo = ReplicatedObjects[RootInternalObjectIndex];
 
-				if (ObjectsInScope.GetBit(RootInternalObjectIndex) && (OwnerInfo.GetState() != EReplicatedObjectState::Invalid && OwnerInfo.GetState() < EReplicatedObjectState::PendingDestroy))
+				if ((OwnerInfo.GetState() != EReplicatedObjectState::Invalid && OwnerInfo.GetState() < EReplicatedObjectState::PendingDestroy))
 				{
 					UE_LOG_REPWRITER(Verbose, TEXT("UpdateDirtyGlobalLifetimeConditionals for - %s"), ToCStr(NetRefHandleManager->PrintObjectFromIndex(RootInternalObjectIndex)));
 						
-					MarkObjectDirty(RootInternalObjectIndex, "UpdateDirtyGlobalLifetimeConditionals");
+					// Better safe than sorry, we do not want to dirty something going out of scope.
+					if (ObjectsInScope.GetBit(RootInternalObjectIndex))
+					{
+						MarkObjectDirty(RootInternalObjectIndex, "UpdateDirtyGlobalLifetimeConditionals");
+					}
 					OwnerInfo.HasDirtyConditionals = 1U;
 					OwnerInfo.HasDirtySubObjects = 1U;
 				}
@@ -1030,7 +1043,11 @@ void FReplicationWriter::UpdateDirtyGlobalLifetimeConditionals(TArrayView<FInter
 			{
 				UE_LOG_REPWRITER(Verbose, TEXT("UpdateDirtyGlobalLifetimeConditionals for - %s"), ToCStr(NetRefHandleManager->PrintObjectFromIndex(InternalObjectIndex)));
 	
-				MarkObjectDirty(InternalObjectIndex, "UpdateDirtyGlobalLifetimeConditionals2");
+				// Better safe than sorry, we do not want to dirty something going out of scope.
+				if (ObjectsInScope.GetBit(InternalObjectIndex))
+				{
+					MarkObjectDirty(InternalObjectIndex, "UpdateDirtyGlobalLifetimeConditionals2");
+				}
 				Info.HasDirtyConditionals = 1U;
 				Info.HasDirtySubObjects = 1U;
 			}
@@ -1125,15 +1142,6 @@ void FReplicationWriter::InternalUpdateDirtyChangeMasks(const FChangeMaskCache&
 	//UE_LOG_REPWRITER(VeryVerbose, TEXT("FReplicationWriter::UpdateDirtyChangeMasks() Updated %u Objects for ConnectionId:%u, ReplicationSystemId: %u."), CachedChangeMasks.Indices.Num(), Parameters.ConnectionId, Parameters.ReplicationSystem->GetId());	
 }
 
-void FReplicationWriter::NotifyDestroyedObjectPendingEndReplication(FInternalNetRefIndex ObjectInternalIndex)
-{
-	const FReplicationInfo& ReplicationInfo = GetReplicationInfo(ObjectInternalIndex);
-	if (ReplicationInfo.GetState() == EReplicatedObjectState::PendingCreate)
-	{
-		StopReplication(ObjectInternalIndex);
-	}
-}
-
 const FNetBitArray& FReplicationWriter::GetObjectsRequiringPriorityUpdate() const
 {
 	return ObjectsWithDirtyChanges;
@@ -1312,6 +1320,16 @@ void FReplicationWriter::HandleDeliveredRecord(const FReplicationRecord::FRecord
 					SetState(InternalIndex, EReplicatedObjectState::PendingTearOff);
 					SetState(InternalIndex, EReplicatedObjectState::Destroyed);
 					StopReplication(InternalIndex);
+
+					// Cleanup any filtered out subobjects that are still in PendingCreate
+					for (uint32 SubObjectIndex : NetRefHandleManager->GetSubObjects(InternalIndex))
+					{
+						const FReplicationInfo& SubObjectInfo = GetReplicationInfo(SubObjectIndex);
+						if (SubObjectInfo.GetState() == EReplicatedObjectState::PendingCreate && SubObjectInfo.IsFilteredOutSubObjectPendingCreate)
+						{
+							StopReplication(SubObjectIndex);
+						}
+					}
 				}
 				else
 				{
@@ -1342,6 +1360,19 @@ void FReplicationWriter::HandleDeliveredRecord(const FReplicationRecord::FRecord
 			// It is now safe to stop tracking this object
 			StopReplication(InternalIndex);
 
+			if (RecordInfo.WroteTearOff)
+			{
+				// Cleanup any filtered out subobjects that are still in PendingCreate
+				for (uint32 SubObjectIndex : NetRefHandleManager->GetSubObjects(InternalIndex))
+				{
+					const FReplicationInfo& SubObjectInfo = GetReplicationInfo(SubObjectIndex);
+					if (SubObjectInfo.GetState() == EReplicatedObjectState::PendingCreate && SubObjectInfo.IsFilteredOutSubObjectPendingCreate)
+					{
+						StopReplication(SubObjectIndex);
+					}
+				}
+			}
+
 			if (CurrentState == EReplicatedObjectState::CancelPendingDestroy)
 			{
 				StartReplication(InternalIndex);
@@ -2086,18 +2117,26 @@ uint32 FReplicationWriter::WriteRootObjectsPendingDestroy(FNetSerializationConte
 		EffectiveMaxDestroyObjectCount = static_cast<uint32>(MaxDestroyObjectsPerFrame) - WriteContext.WrittenDestroyObjectCount;
 	}
 
-	auto AreSubObjectsReadyToBeDestroyed = [this](uint32 InternalIndex)
+	auto AreSubObjectsReadyToBeDestroyed = [this](uint32 InternalIndex, const FReplicationConditionals::FSubObjectsToReplicateArray& FilteredSubObjects, TArrayView<const FInternalNetRefIndex>& SubObjects)
 	{
-		for (uint32 SubObjectIndex : NetRefHandleManager->GetSubObjects(InternalIndex))
+		for (uint32 SubObjectIndex : SubObjects)
 		{
 			FReplicationInfo& SubObjectInfo = GetReplicationInfo(SubObjectIndex);
 			const EReplicatedObjectState SubObjectState = SubObjectInfo.GetState();
+
 			// A subobject may have stopped replicating immediately when going out of scope if it was never replicated to begin with.
 			if (SubObjectState == EReplicatedObjectState::Invalid)
 			{
 				continue;
 			}
 
+			// Filtered out subobjects still in PendingCreate can be stopped.
+			if (SubObjectState == EReplicatedObjectState::PendingCreate && !FilteredSubObjects.Contains(SubObjectIndex))
+			{
+				StopReplication(SubObjectIndex);
+				continue;
+			}
+
 			if (SubObjectState == EReplicatedObjectState::WaitOnFlush)
 			{
 				if (GetFlushStatus(SubObjectIndex, SubObjectInfo, SubObjectInfo.FlushFlags) == EFlushFlags::FlushFlags_None)
@@ -2145,6 +2184,9 @@ uint32 FReplicationWriter::WriteRootObjectsPendingDestroy(FNetSerializationConte
 			continue;
 		}
 
+		// Lookup once.
+		TArrayView<const FInternalNetRefIndex> AllSubObjects = NetRefHandleManager->GetSubObjects(InternalIndex);
+
 		if (Info.GetState() == EReplicatedObjectState::WaitOnFlush)
 		{
 			if (GetFlushStatus(InternalIndex, Info, Info.FlushFlags) != EFlushFlags::FlushFlags_None)
@@ -2156,9 +2198,9 @@ uint32 FReplicationWriter::WriteRootObjectsPendingDestroy(FNetSerializationConte
 			if (Info.TearOff)
 			{
 				// Check that all subobjects are ready to be torn off too.
-				const bool bAllSubObjectsAreReadyToBeTornOff = [this, InternalIndex]()
+				const bool bAllSubObjectsAreReadyToBeTornOff = [this, InternalIndex, &AllSubObjects]()
 				{
-					for (const FInternalNetRefIndex SubObjectIndex : NetRefHandleManager->GetSubObjects(InternalIndex))
+					for (const FInternalNetRefIndex SubObjectIndex : AllSubObjects)
 					{
 						FReplicationInfo& SubObjectInfo = ReplicatedObjects[SubObjectIndex];
 						if (SubObjectInfo.GetState() == EReplicatedObjectState::WaitOnFlush)
@@ -2194,7 +2236,11 @@ uint32 FReplicationWriter::WriteRootObjectsPendingDestroy(FNetSerializationConte
 			}
 		}
 
-		if (!AreSubObjectsReadyToBeDestroyed(InternalIndex))
+
+		FReplicationConditionals::FSubObjectsToReplicateArray SubObjectsToReplicate;
+		ReplicationConditionals->GetSubObjectsToReplicate(Parameters.ConnectionId, InternalIndex, SubObjectsToReplicate);
+
+		if (!AreSubObjectsReadyToBeDestroyed(InternalIndex, SubObjectsToReplicate, AllSubObjects))
 		{
 			bWroteAllDestroyedObjects = false;
 			continue;
@@ -2236,7 +2282,6 @@ uint32 FReplicationWriter::WriteRootObjectsPendingDestroy(FNetSerializationConte
 			}
 		}
 
-
 		EInternalDetachReason Reason = ObjectData.InternalDetachReason;
 
 		// New path for IsUsingNewDetachAPI() == true
@@ -2278,7 +2323,7 @@ uint32 FReplicationWriter::WriteRootObjectsPendingDestroy(FNetSerializationConte
 		CreateObjectRecord(nullptr, Info, ObjectInfo, ObjectRecord);
 
 		// Fill in subobject record
-		for (const uint32 SubObjectIndex : NetRefHandleManager->GetSubObjects(InternalIndex))
+		for (const uint32 SubObjectIndex : AllSubObjects)
 		{
 			FReplicationInfo& SubObjectInfo = GetReplicationInfo(SubObjectIndex);
 			if (SubObjectInfo.GetState() == EReplicatedObjectState::Invalid)
@@ -2660,16 +2705,14 @@ FReplicationWriter::EWriteObjectStatus FReplicationWriter::WriteObjectAndSubObje
 		{
 			FReplicationInfo& SubObjectInfo = GetReplicationInfo(SubObjectIndex);
 
-			// Better safe than sorry, we do not want to dirty something going out of scope.
-			if (!ObjectsInScope.GetBit(SubObjectIndex))
-			{
-				continue;
-			}
-
 			// Need to be a bit careful what we explicitly dirty here, as we might have subobjects waiting for creation confirmation.
 			if ((SubObjectInfo.GetState() != EReplicatedObjectState::Invalid && SubObjectInfo.GetState() < EReplicatedObjectState::PendingDestroy) && CanSendObject(SubObjectIndex))
 			{
-				MarkObjectDirty(SubObjectIndex, "UpdateGlobalConditional");				
+				// Better safe than sorry, we do not want to dirty something going out of scope.
+				if (ObjectsInScope.GetBit(SubObjectIndex))
+				{
+					MarkObjectDirty(SubObjectIndex, "UpdateGlobalConditional");
+				}
 				SubObjectInfo.HasDirtyConditionals = 1;
 			}
 
@@ -2679,7 +2722,9 @@ FReplicationWriter::EWriteObjectStatus FReplicationWriter::WriteObjectAndSubObje
 	}
 	Info.HasDirtyConditionals = 0;
 
-	const bool bAllowStateAndAttachmentSending = !Info.SubObjectPendingDestroy;
+	// For replicated subobject and tear off of created SubObjects that are filtered out by SubObject conditionals we want to replicate the destroy/tearoff but skip the data.
+	// Note: IsFilteredOutSubObjectTearOff will only be set to ensure that a SubObject which has been replicated before being filtered out is property released.
+	const bool bAllowStateAndAttachmentSending = !(Info.SubObjectPendingDestroy || Info.IsFilteredOutSubObjectTearOff);
 	const bool bIsObjectIndexForAttachment = IsObjectIndexForOOBAttachment(InternalIndex);
 	const bool bHasState = (bIsInitialState || Info.HasDirtyChangeMask) && !!WriteObjectTraits.bWriteState && bAllowStateAndAttachmentSending;
 	const bool bHasAttachments = (Info.HasAttachments || bIsObjectIndexForAttachment) && bAllowStateAndAttachmentSending;
@@ -3025,13 +3070,27 @@ FReplicationWriter::EWriteObjectStatus FReplicationWriter::WriteObjectAndSubObje
 		FReplicationConditionals::FSubObjectsToReplicateArray SubObjectsToReplicate;
 		ReplicationConditionals->GetSubObjectsToReplicate(Parameters.ConnectionId, InternalIndex, SubObjectsToReplicate);
 
-		// Patch up list with subobjects pending destroy.
+		// We need to do some additional logic for filtered out subobjects, pending destroy or tearoff
 		for (uint32 SubObjectIndex : NetRefHandleManager->GetSubObjects(InternalIndex))
 		{
-			const FReplicationInfo& SubObjectInfo = GetReplicationInfo(SubObjectIndex);
-			if (SubObjectInfo.GetState() == EReplicatedObjectState::SubObjectPendingDestroy && !SubObjectsToReplicate.Contains(SubObjectIndex))
+			FReplicationInfo& SubObjectInfo = GetReplicationInfo(SubObjectIndex);
+			if (SubObjectInfo.GetState() == EReplicatedObjectState::SubObjectPendingDestroy || (SubObjectInfo.IsCreationConfirmed && SubObjectInfo.TearOff))
 			{
-				SubObjectsToReplicate.Add(SubObjectIndex);
+				if (!SubObjectsToReplicate.Contains(SubObjectIndex))
+				{
+					SubObjectsToReplicate.Add(SubObjectIndex);
+					if (SubObjectInfo.TearOff)
+					{
+						SubObjectInfo.IsFilteredOutSubObjectTearOff = 1U;
+					}
+				}
+			}
+			else if (SubObjectInfo.GetState() == EReplicatedObjectState::PendingCreate)
+			{
+				if (!SubObjectsToReplicate.Contains(SubObjectIndex))
+				{
+					SubObjectInfo.IsFilteredOutSubObjectPendingCreate = 1U;
+				}
 			}
 		}
 		
@@ -3907,10 +3966,12 @@ uint32 FReplicationWriter::HandleObjectBatchSuccess(const FBatchInfo& BatchInfo,
 		// We did write the initial state, change the state to WaitOnCreateConfirmation
 		if (BatchObjectInfo.bIsInitialState)
 		{
+			Info.IsFilteredOutSubObjectPendingCreate = 0U;
 			SetState(BatchObjectInfo.InternalIndex, EReplicatedObjectState::WaitOnCreateConfirmation);
 		}
 		else if (Info.TearOff)
 		{
+			Info.IsFilteredOutSubObjectTearOff = 0U;
 			if (BatchObjectInfo.bSentTearOff)
 			{
 				SetState(BatchObjectInfo.InternalIndex, EReplicatedObjectState::PendingTearOff);
diff --git a/Engine/Source/Runtime/Net/Iris/Private/Iris/ReplicationSystem/ReplicationWriter.h b/Engine/Source/Runtime/Net/Iris/Private/Iris/ReplicationSystem/ReplicationWriter.h
--- a/Engine/Source/Runtime/Net/Iris/Private/Iris/ReplicationSystem/ReplicationWriter.h
+++ b/Engine/Source/Runtime/Net/Iris/Private/Iris/ReplicationSystem/ReplicationWriter.h
@@ -144,6 +144,8 @@ public:
 				uint64 HasDirtyConditionals : 1;						// If this flag is set, we must update conditionals.
 				mutable uint64 HasCannotSendInfo : 1;					// If this flag is set, the object has been prevented from being sent at least once.
 				uint64 IsDebugObject : 1;								// If this flag is set, the object is a debugobject and is always sent using DebugObject attachments.
+				uint64 IsFilteredOutSubObjectPendingCreate : 1;			// If this flag is set, the object is a subobject that was filtered out while sending pending create which requires some extra attention
+				uint64 IsFilteredOutSubObjectTearOff : 1;				// If this flag is set, the object is a subobject confirmed as created that was filtered out while marked as tearoff which requires some extra attention
 			};
 		};
 
@@ -178,9 +180,6 @@ public:
 	// Force update DirtyChangeMasks and mark objects for flush and/or tearoff depending on flags
 	void ForceUpdateDirtyChangeMasks(const FChangeMaskCache& CachedChangeMasks, EFlushFlags ExtraFlushFlags, bool bMarkForTearOff) { InternalUpdateDirtyChangeMasks(CachedChangeMasks, ExtraFlushFlags, bMarkForTearOff); }
 
-	// Called if an object first being teared-off/flushed and then explicitly destroyed before it has been removed from scope
-	void NotifyDestroyedObjectPendingEndReplication(FInternalNetRefIndex ObjectInternalIndex);
-
 	// Called to propagate changes to global lifetime conditionals
 	void UpdateDirtyGlobalLifetimeConditionals(TArrayView<FInternalNetRefIndex> ObjectsWithDirtyConditionals);
 
