Optimize UpdateCreationDependencies in Iris filtering step by only updating objects in dependency trees when an object in the tree has become relevant or filtered.

#tests EngineTest Net unit tests, SceneGraph replication unit tests, ReplicationSystemTest unit tests, and Server Load Test

diff --git a/Engine/Plugins/Runtime/ReplicationSystemTestPlugin/Source/Private/Tests/ReplicationSystem/Filtering/TestFiltering.cpp b/Engine/Plugins/Runtime/ReplicationSystemTestPlugin/Source/Private/Tests/ReplicationSystem/Filtering/TestFiltering.cpp
--- a/Engine/Plugins/Runtime/ReplicationSystemTestPlugin/Source/Private/Tests/ReplicationSystem/Filtering/TestFiltering.cpp
+++ b/Engine/Plugins/Runtime/ReplicationSystemTestPlugin/Source/Private/Tests/ReplicationSystem/Filtering/TestFiltering.cpp
@@ -3031,6 +3031,97 @@ UE_NET_TEST_FIXTURE(FTestFilteringFixture, DependentObjectWithHysteresisAddedToE
 	}
 }
 
+UE_NET_TEST_FIXTURE(FTestFilteringFixture, CreationDependentObjectRespectsParentHysteresis)
+{
+	UReplicatedTestObjectBridge* Bridge = Server->GetReplicationBridge();
+
+	FReplicationSystemTestClient* Client = CreateClient();
+
+	UReplicatedTestObject* Parent = Server->CreateObject();
+	UReplicatedTestObject* Child = Server->CreateObject();
+	Bridge->AddCreationDependencyLink(Parent->NetRefHandle, Child->NetRefHandle);
+
+	// Create objects on client.
+	Server->UpdateAndSend({Client});
+	UE_NET_ASSERT_TRUE(Client->IsResolvableNetRefHandle(Parent->NetRefHandle));
+	UE_NET_ASSERT_TRUE(Client->IsResolvableNetRefHandle(Child->NetRefHandle));
+
+	// Filter out the parent, but set hysteresis to five frames.
+	Server->ReplicationSystem->SetFilter(Parent->NetRefHandle, NotRoutedFilterHandle, "FiveFrames");
+
+	// Parent and child should both stay in scope for 5 frames.
+	for (unsigned FrameIt = 0; FrameIt < 5U; ++FrameIt)
+	{
+		Server->UpdateAndSend({Client});
+		UE_NET_ASSERT_TRUE(Client->IsResolvableNetRefHandle(Parent->NetRefHandle));
+		UE_NET_ASSERT_TRUE(Client->IsResolvableNetRefHandle(Child->NetRefHandle));
+	}
+
+	// Parent and child should both go out of scope after 5 frames.
+	Server->UpdateAndSend({Client});
+	UE_NET_ASSERT_FALSE(Client->IsResolvableNetRefHandle(Parent->NetRefHandle));
+	UE_NET_ASSERT_FALSE(Client->IsResolvableNetRefHandle(Child->NetRefHandle));
+}
+
+UE_NET_TEST_FIXTURE(FTestFilteringFixture, CreationDependentObjectHysteresisIgnoredWhenParentIsFilteredOut)
+{
+	UReplicatedTestObjectBridge* Bridge = Server->GetReplicationBridge();
+
+	FReplicationSystemTestClient* Client = CreateClient();
+
+	UReplicatedTestObject* Parent = Server->CreateObject();
+	UReplicatedTestObject* Child = Server->CreateObject();
+	Bridge->AddCreationDependencyLink(Parent->NetRefHandle, Child->NetRefHandle);
+
+	// Create objects on client.
+	Server->UpdateAndSend({Client});
+	UE_NET_ASSERT_TRUE(Client->IsResolvableNetRefHandle(Parent->NetRefHandle));
+	UE_NET_ASSERT_TRUE(Client->IsResolvableNetRefHandle(Child->NetRefHandle));
+
+	// Filter out parent, filter out child with hysteresis.
+	Server->ReplicationSystem->SetFilter(Parent->NetRefHandle, NotRoutedFilterHandle);
+	Server->ReplicationSystem->SetFilter(Child->NetRefHandle, NotRoutedFilterHandle, "FiveFrames");
+
+	// Creation dependency should override child hysteresis, so both objects shouldn't exist on client.
+	Server->UpdateAndSend({Client});
+	UE_NET_ASSERT_FALSE(Client->IsResolvableNetRefHandle(Parent->NetRefHandle));
+	UE_NET_ASSERT_FALSE(Client->IsResolvableNetRefHandle(Child->NetRefHandle));
+}
+
+UE_NET_TEST_FIXTURE(FTestFilteringFixture, CreationDependentObjectInHysteresisFilteredOutWhenParentHysteresisEnds)
+{
+	UReplicatedTestObjectBridge* Bridge = Server->GetReplicationBridge();
+
+	FReplicationSystemTestClient* Client = CreateClient();
+
+	UReplicatedTestObject* Parent = Server->CreateObject();
+	UReplicatedTestObject* Child = Server->CreateObject();
+	Bridge->AddCreationDependencyLink(Parent->NetRefHandle, Child->NetRefHandle);
+
+	// Create objects on client.
+	Server->UpdateAndSend({Client});
+	UE_NET_ASSERT_TRUE(Client->IsResolvableNetRefHandle(Parent->NetRefHandle));
+	UE_NET_ASSERT_TRUE(Client->IsResolvableNetRefHandle(Child->NetRefHandle));
+
+	// Filter out parent after one frame, filter out child after five frames.
+	Server->ReplicationSystem->SetFilter(Parent->NetRefHandle, NotRoutedFilterHandle, "OneFrame");
+	Server->ReplicationSystem->SetFilter(Child->NetRefHandle, NotRoutedFilterHandle, "FiveFrames");
+
+	// Both objects should still exist after one frame.
+	Server->UpdateAndSend({Client});
+	UE_NET_ASSERT_TRUE(Client->IsResolvableNetRefHandle(Parent->NetRefHandle));
+	UE_NET_ASSERT_TRUE(Client->IsResolvableNetRefHandle(Child->NetRefHandle));
+
+	// Parent hysteresis ends after one frame, so child should be pulled out of hysteresis.
+	// Parent and child should remain out of scope for the duration of what would have been the child's hysteresis.
+	for (unsigned FrameIt = 0; FrameIt < 5U; ++FrameIt)
+	{
+		Server->UpdateAndSend( { Client });
+		UE_NET_ASSERT_FALSE(Client->IsResolvableNetRefHandle(Parent->NetRefHandle));
+		UE_NET_ASSERT_FALSE(Client->IsResolvableNetRefHandle(Child->NetRefHandle));
+	}
+}
+
 UE_NET_TEST_FIXTURE(FTestFilteringFixture, OwningConnectionIsSetProperly)
 {
 	FReplicationSystemTestClient* Client = CreateClient();
diff --git a/Engine/Plugins/Runtime/ReplicationSystemTestPlugin/Source/Private/Tests/ReplicationSystem/TestDependentObjects.cpp b/Engine/Plugins/Runtime/ReplicationSystemTestPlugin/Source/Private/Tests/ReplicationSystem/TestDependentObjects.cpp
--- a/Engine/Plugins/Runtime/ReplicationSystemTestPlugin/Source/Private/Tests/ReplicationSystem/TestDependentObjects.cpp
+++ b/Engine/Plugins/Runtime/ReplicationSystemTestPlugin/Source/Private/Tests/ReplicationSystem/TestDependentObjects.cpp
@@ -1113,6 +1113,148 @@ UE_NET_TEST_FIXTURE(FTestFilteredDependentObjectFixture, TestBasicCreationDepend
 	UE_NET_ASSERT_FALSE(Client->IsResolvableNetRefHandle(ChildObject->NetRefHandle));
 }
 
+UE_NET_TEST_FIXTURE(FTestFilteredDependentObjectFixture, TestCreationDependencyRemoveLink)
+{
+	UReplicationSystem* ReplicationSystem = Server->ReplicationSystem;
+	UReplicatedTestObjectBridge* Bridge = Server->GetReplicationBridge();
+
+	// Add a client
+	FReplicationSystemTestClient* Client = CreateClient();
+
+	// Spawn server objects
+	UReplicatedTestObject* ParentObject = Server->CreateObject();
+	UReplicatedTestObject* ChildObject = Server->CreateObject();
+	Bridge->AddCreationDependencyLink(ParentObject->NetRefHandle, ChildObject->NetRefHandle);
+
+	// Send and deliver packets
+	Server->UpdateAndSend({ Client });
+
+	// Objects should now exist on client
+	UE_NET_ASSERT_TRUE(Client->IsResolvableNetRefHandle(ParentObject->NetRefHandle));
+	UE_NET_ASSERT_TRUE(Client->IsResolvableNetRefHandle(ChildObject->NetRefHandle));
+
+	// Filter out parent
+	ReplicationSystem->SetFilter(ParentObject->NetRefHandle, NotRoutedFilterHandle);
+
+	// Send and deliver packets
+	Server->UpdateAndSend({ Client });
+
+	// Objects should not exist on client
+	UE_NET_ASSERT_FALSE(Client->IsResolvableNetRefHandle(ParentObject->NetRefHandle));
+	UE_NET_ASSERT_FALSE(Client->IsResolvableNetRefHandle(ChildObject->NetRefHandle));
+
+	// Remove dependency link
+	Bridge->RemoveCreationDependencyLink(ParentObject->NetRefHandle, ChildObject->NetRefHandle);
+
+	// Send and deliver packets
+	Server->UpdateAndSend({ Client });
+
+	// Child should come back into scope.
+	UE_NET_ASSERT_FALSE(Client->IsResolvableNetRefHandle(ParentObject->NetRefHandle));
+	UE_NET_ASSERT_TRUE(Client->IsResolvableNetRefHandle(ChildObject->NetRefHandle));
+
+	// Filter in parent
+	ReplicationSystem->SetFilter(ParentObject->NetRefHandle, InvalidNetObjectFilterHandle);
+
+	// Send and deliver packets
+	Server->UpdateAndSend({ Client });
+
+	// Parent should come back into scope.
+	UE_NET_ASSERT_TRUE(Client->IsResolvableNetRefHandle(ParentObject->NetRefHandle));
+	UE_NET_ASSERT_TRUE(Client->IsResolvableNetRefHandle(ChildObject->NetRefHandle));
+}
+
+UE_NET_TEST_FIXTURE(FTestFilteredDependentObjectFixture, TestCreationDependencyTwoParentsRemoveOneLink)
+{
+	UReplicationSystem* ReplicationSystem = Server->ReplicationSystem;
+	UReplicatedTestObjectBridge* Bridge = Server->GetReplicationBridge();
+
+	// Add a client
+	FReplicationSystemTestClient* Client = CreateClient();
+
+	// Spawn server objects
+	UReplicatedTestObject* ParentObject = Server->CreateObject();
+	UReplicatedTestObject* FilteredParentObject = Server->CreateObject();
+	UReplicatedTestObject* ChildObject = Server->CreateObject();
+	Bridge->AddCreationDependencyLink(ParentObject->NetRefHandle, ChildObject->NetRefHandle);
+	Bridge->AddCreationDependencyLink(FilteredParentObject->NetRefHandle, ChildObject->NetRefHandle);
+
+	// Send and deliver packets
+	Server->UpdateAndSend({ Client });
+
+	// Objects should now exist on client
+	UE_NET_ASSERT_TRUE(Client->IsResolvableNetRefHandle(ParentObject->NetRefHandle));
+	UE_NET_ASSERT_TRUE(Client->IsResolvableNetRefHandle(FilteredParentObject->NetRefHandle));
+	UE_NET_ASSERT_TRUE(Client->IsResolvableNetRefHandle(ChildObject->NetRefHandle));
+
+	// Filter out parent
+	ReplicationSystem->SetFilter(FilteredParentObject->NetRefHandle, NotRoutedFilterHandle);
+
+	// Send and deliver packets
+	Server->UpdateAndSend({ Client });
+
+	// Parent should exist, FilteredParent and Child should not.
+	UE_NET_ASSERT_TRUE(Client->IsResolvableNetRefHandle(ParentObject->NetRefHandle));
+	UE_NET_ASSERT_FALSE(Client->IsResolvableNetRefHandle(FilteredParentObject->NetRefHandle));
+	UE_NET_ASSERT_FALSE(Client->IsResolvableNetRefHandle(ChildObject->NetRefHandle));
+
+	// Remove dependency link
+	Bridge->RemoveCreationDependencyLink(FilteredParentObject->NetRefHandle, ChildObject->NetRefHandle);
+
+	// Send and deliver packets
+	Server->UpdateAndSend({ Client });
+
+	// Child should come back into scope.
+	UE_NET_ASSERT_TRUE(Client->IsResolvableNetRefHandle(ParentObject->NetRefHandle));
+	UE_NET_ASSERT_FALSE(Client->IsResolvableNetRefHandle(FilteredParentObject->NetRefHandle));
+	UE_NET_ASSERT_TRUE(Client->IsResolvableNetRefHandle(ChildObject->NetRefHandle));
+}
+
+UE_NET_TEST_FIXTURE(FTestFilteredDependentObjectFixture, TestCreationDependencyParentDestroyed)
+{
+	UReplicationSystem* ReplicationSystem = Server->ReplicationSystem;
+	UReplicatedTestObjectBridge* Bridge = Server->GetReplicationBridge();
+
+	// Add a client
+	FReplicationSystemTestClient* Client = CreateClient();
+
+	// Spawn server objects
+	UReplicatedTestObject* ParentObject = Server->CreateObject();
+	UReplicatedTestObject* ChildObject = Server->CreateObject();
+	Bridge->AddCreationDependencyLink(ParentObject->NetRefHandle, ChildObject->NetRefHandle);
+
+	// Send and deliver packets
+	Server->UpdateAndSend({ Client });
+
+	// Objects should now exist on client
+	UE_NET_ASSERT_TRUE(Client->IsResolvableNetRefHandle(ParentObject->NetRefHandle));
+	UE_NET_ASSERT_TRUE(Client->IsResolvableNetRefHandle(ChildObject->NetRefHandle));
+
+	// Filter out parent
+	ReplicationSystem->SetFilter(ParentObject->NetRefHandle, NotRoutedFilterHandle);
+
+	// Send and deliver packets
+	Server->UpdateAndSend({ Client });
+
+	// Both should be filtered
+	UE_NET_ASSERT_FALSE(Client->IsResolvableNetRefHandle(ParentObject->NetRefHandle));
+	UE_NET_ASSERT_FALSE(Client->IsResolvableNetRefHandle(ChildObject->NetRefHandle));
+
+	// Destroy parent
+	Server->DestroyObject(ParentObject);
+
+	// Send and deliver packets
+	// TODO: We update and send here twice because the base implementation of creation dependency handling
+	//       doesn't bring the child back into scope until two updates have passed. The second update can be
+	//       removed here once we've set bCVarRepFilterUpdateCreationDependenciesOnRelevancyChange=true by default
+	//       or once we've fixed the base implementation to filter in the child object immediately.
+	Server->UpdateAndSend({ Client });
+	Server->UpdateAndSend({ Client });
+
+	// Child should come back into scope.
+	UE_NET_ASSERT_TRUE(Client->IsResolvableNetRefHandle(ChildObject->NetRefHandle));
+}
+
 UE_NET_TEST_FIXTURE(FTestFilteredDependentObjectFixture, TestNestedCreationDependencyFiltering)
 {
 	UReplicationSystem* ReplicationSystem = Server->ReplicationSystem;
@@ -1217,29 +1359,100 @@ UE_NET_TEST_FIXTURE(FTestFilteredDependentObjectFixture, TestComplexCreationDepe
 		UE_NET_ASSERT_TRUE(Client->IsResolvableNetRefHandle(BottomObjects[ObjectIndex]->NetRefHandle));
 	}
 
-	// Filter out TopObjects[2]
-	ReplicationSystem->SetFilter(TopObjects[2]->NetRefHandle, NotRoutedFilterHandle);
+	{
+		// Filter out MiddleObjects[1]
+		ReplicationSystem->SetFilter(MiddleObjects[1]->NetRefHandle, NotRoutedFilterHandle);
 
-	// Send and deliver packets
-	Server->UpdateAndSend({ Client });
+		// Send and deliver packets
+		Server->UpdateAndSend({ Client });
 
-	// TopObjects[2] should not exist on the client
-	UE_NET_ASSERT_TRUE(Client->IsResolvableNetRefHandle(TopObjects[0]->NetRefHandle));
-	UE_NET_ASSERT_TRUE(Client->IsResolvableNetRefHandle(TopObjects[1]->NetRefHandle));
-	UE_NET_ASSERT_FALSE(Client->IsResolvableNetRefHandle(TopObjects[2]->NetRefHandle));
-	UE_NET_ASSERT_TRUE(Client->IsResolvableNetRefHandle(TopObjects[3]->NetRefHandle));
-
-	// MiddleObjects[1-2] should not exist on the client
-	UE_NET_ASSERT_TRUE(Client->IsResolvableNetRefHandle(MiddleObjects[0]->NetRefHandle));
-	UE_NET_ASSERT_FALSE(Client->IsResolvableNetRefHandle(MiddleObjects[1]->NetRefHandle));
-	UE_NET_ASSERT_FALSE(Client->IsResolvableNetRefHandle(MiddleObjects[2]->NetRefHandle));
-	UE_NET_ASSERT_TRUE(Client->IsResolvableNetRefHandle(MiddleObjects[3]->NetRefHandle));
-
-	// BottomObjects[0-2] should not exist on the client
-	UE_NET_ASSERT_FALSE(Client->IsResolvableNetRefHandle(BottomObjects[0]->NetRefHandle));
-	UE_NET_ASSERT_FALSE(Client->IsResolvableNetRefHandle(BottomObjects[1]->NetRefHandle));
-	UE_NET_ASSERT_FALSE(Client->IsResolvableNetRefHandle(BottomObjects[2]->NetRefHandle));
-	UE_NET_ASSERT_TRUE(Client->IsResolvableNetRefHandle(BottomObjects[3]->NetRefHandle));
+		// TopObjects should all exist on the client
+		UE_NET_ASSERT_TRUE(Client->IsResolvableNetRefHandle(TopObjects[0]->NetRefHandle));
+		UE_NET_ASSERT_TRUE(Client->IsResolvableNetRefHandle(TopObjects[1]->NetRefHandle));
+		UE_NET_ASSERT_TRUE(Client->IsResolvableNetRefHandle(TopObjects[2]->NetRefHandle));
+		UE_NET_ASSERT_TRUE(Client->IsResolvableNetRefHandle(TopObjects[3]->NetRefHandle));
+
+		// MiddleObjects[1] should not exist on the client
+		UE_NET_ASSERT_TRUE(Client->IsResolvableNetRefHandle(MiddleObjects[0]->NetRefHandle));
+		UE_NET_ASSERT_FALSE(Client->IsResolvableNetRefHandle(MiddleObjects[1]->NetRefHandle));
+		UE_NET_ASSERT_TRUE(Client->IsResolvableNetRefHandle(MiddleObjects[2]->NetRefHandle));
+		UE_NET_ASSERT_TRUE(Client->IsResolvableNetRefHandle(MiddleObjects[3]->NetRefHandle));
+
+		// BottomObjects[0-1] should not exist on the client
+		UE_NET_ASSERT_FALSE(Client->IsResolvableNetRefHandle(BottomObjects[0]->NetRefHandle));
+		UE_NET_ASSERT_FALSE(Client->IsResolvableNetRefHandle(BottomObjects[1]->NetRefHandle));
+		UE_NET_ASSERT_TRUE(Client->IsResolvableNetRefHandle(BottomObjects[2]->NetRefHandle));
+		UE_NET_ASSERT_TRUE(Client->IsResolvableNetRefHandle(BottomObjects[3]->NetRefHandle));
+	}
+
+	{
+		// Filter out TopObjects[2]
+		ReplicationSystem->SetFilter(TopObjects[2]->NetRefHandle, NotRoutedFilterHandle);
+
+		// Send and deliver packets
+		Server->UpdateAndSend({ Client });
+
+		// TopObjects[2] should not exist on the client
+		UE_NET_ASSERT_TRUE(Client->IsResolvableNetRefHandle(TopObjects[0]->NetRefHandle));
+		UE_NET_ASSERT_TRUE(Client->IsResolvableNetRefHandle(TopObjects[1]->NetRefHandle));
+		UE_NET_ASSERT_FALSE(Client->IsResolvableNetRefHandle(TopObjects[2]->NetRefHandle));
+		UE_NET_ASSERT_TRUE(Client->IsResolvableNetRefHandle(TopObjects[3]->NetRefHandle));
+
+		// MiddleObjects[1-2] should not exist on the client
+		UE_NET_ASSERT_TRUE(Client->IsResolvableNetRefHandle(MiddleObjects[0]->NetRefHandle));
+		UE_NET_ASSERT_FALSE(Client->IsResolvableNetRefHandle(MiddleObjects[1]->NetRefHandle));
+		UE_NET_ASSERT_FALSE(Client->IsResolvableNetRefHandle(MiddleObjects[2]->NetRefHandle));
+		UE_NET_ASSERT_TRUE(Client->IsResolvableNetRefHandle(MiddleObjects[3]->NetRefHandle));
+
+		// BottomObjects[0-2] should not exist on the client
+		UE_NET_ASSERT_FALSE(Client->IsResolvableNetRefHandle(BottomObjects[0]->NetRefHandle));
+		UE_NET_ASSERT_FALSE(Client->IsResolvableNetRefHandle(BottomObjects[1]->NetRefHandle));
+		UE_NET_ASSERT_FALSE(Client->IsResolvableNetRefHandle(BottomObjects[2]->NetRefHandle));
+		UE_NET_ASSERT_TRUE(Client->IsResolvableNetRefHandle(BottomObjects[3]->NetRefHandle));
+	}
+
+	{
+		// Filter in MiddleObjects[1]
+		// Nothing should change because one of its parents is filtered out
+		ReplicationSystem->SetFilter(MiddleObjects[1]->NetRefHandle, InvalidNetObjectFilterHandle);
+
+		// Send and deliver packets
+		Server->UpdateAndSend({ Client });
+
+		// TopObjects[2] should not exist on the client
+		UE_NET_ASSERT_TRUE(Client->IsResolvableNetRefHandle(TopObjects[0]->NetRefHandle));
+		UE_NET_ASSERT_TRUE(Client->IsResolvableNetRefHandle(TopObjects[1]->NetRefHandle));
+		UE_NET_ASSERT_FALSE(Client->IsResolvableNetRefHandle(TopObjects[2]->NetRefHandle));
+		UE_NET_ASSERT_TRUE(Client->IsResolvableNetRefHandle(TopObjects[3]->NetRefHandle));
+
+		// MiddleObjects[1-2] should not exist on the client
+		UE_NET_ASSERT_TRUE(Client->IsResolvableNetRefHandle(MiddleObjects[0]->NetRefHandle));
+		UE_NET_ASSERT_FALSE(Client->IsResolvableNetRefHandle(MiddleObjects[1]->NetRefHandle));
+		UE_NET_ASSERT_FALSE(Client->IsResolvableNetRefHandle(MiddleObjects[2]->NetRefHandle));
+		UE_NET_ASSERT_TRUE(Client->IsResolvableNetRefHandle(MiddleObjects[3]->NetRefHandle));
+
+		// BottomObjects[0-2] should not exist on the client
+		UE_NET_ASSERT_FALSE(Client->IsResolvableNetRefHandle(BottomObjects[0]->NetRefHandle));
+		UE_NET_ASSERT_FALSE(Client->IsResolvableNetRefHandle(BottomObjects[1]->NetRefHandle));
+		UE_NET_ASSERT_FALSE(Client->IsResolvableNetRefHandle(BottomObjects[2]->NetRefHandle));
+		UE_NET_ASSERT_TRUE(Client->IsResolvableNetRefHandle(BottomObjects[3]->NetRefHandle));
+	}
+
+	{
+		// Filter in TopObjects[2]
+		ReplicationSystem->SetFilter(TopObjects[2]->NetRefHandle, InvalidNetObjectFilterHandle);
+
+		// Send and deliver packets
+		Server->UpdateAndSend({ Client });
+
+		// Objects should all exist on client
+		for (uint32 ObjectIndex = 0U; ObjectIndex < ArrayObjectCount; ++ObjectIndex)
+		{
+			UE_NET_ASSERT_TRUE(Client->IsResolvableNetRefHandle(TopObjects[ObjectIndex]->NetRefHandle));
+			UE_NET_ASSERT_TRUE(Client->IsResolvableNetRefHandle(MiddleObjects[ObjectIndex]->NetRefHandle));
+			UE_NET_ASSERT_TRUE(Client->IsResolvableNetRefHandle(BottomObjects[ObjectIndex]->NetRefHandle));
+		}
+	}
 }
 
 // Tests creation dependencies with multiple parents, multiple children, and children created first.
@@ -1302,29 +1515,100 @@ UE_NET_TEST_FIXTURE(FTestFilteredDependentObjectFixture, TestComplexCreationDepe
 		UE_NET_ASSERT_TRUE(Client->IsResolvableNetRefHandle(BottomObjects[ObjectIndex]->NetRefHandle));
 	}
 
-	// Filter out TopObjects[2]
-	ReplicationSystem->SetFilter(TopObjects[2]->NetRefHandle, NotRoutedFilterHandle);
+	{
+		// Filter out MiddleObjects[1]
+		ReplicationSystem->SetFilter(MiddleObjects[1]->NetRefHandle, NotRoutedFilterHandle);
 
-	// Send and deliver packets
-	Server->UpdateAndSend({ Client });
+		// Send and deliver packets
+		Server->UpdateAndSend({ Client });
+
+		// TopObjects should all exist on the client
+		UE_NET_ASSERT_TRUE(Client->IsResolvableNetRefHandle(TopObjects[0]->NetRefHandle));
+		UE_NET_ASSERT_TRUE(Client->IsResolvableNetRefHandle(TopObjects[1]->NetRefHandle));
+		UE_NET_ASSERT_TRUE(Client->IsResolvableNetRefHandle(TopObjects[2]->NetRefHandle));
+		UE_NET_ASSERT_TRUE(Client->IsResolvableNetRefHandle(TopObjects[3]->NetRefHandle));
+
+		// MiddleObjects[1] should not exist on the client
+		UE_NET_ASSERT_TRUE(Client->IsResolvableNetRefHandle(MiddleObjects[0]->NetRefHandle));
+		UE_NET_ASSERT_FALSE(Client->IsResolvableNetRefHandle(MiddleObjects[1]->NetRefHandle));
+		UE_NET_ASSERT_TRUE(Client->IsResolvableNetRefHandle(MiddleObjects[2]->NetRefHandle));
+		UE_NET_ASSERT_TRUE(Client->IsResolvableNetRefHandle(MiddleObjects[3]->NetRefHandle));
+
+		// BottomObjects[0-1] should not exist on the client
+		UE_NET_ASSERT_FALSE(Client->IsResolvableNetRefHandle(BottomObjects[0]->NetRefHandle));
+		UE_NET_ASSERT_FALSE(Client->IsResolvableNetRefHandle(BottomObjects[1]->NetRefHandle));
+		UE_NET_ASSERT_TRUE(Client->IsResolvableNetRefHandle(BottomObjects[2]->NetRefHandle));
+		UE_NET_ASSERT_TRUE(Client->IsResolvableNetRefHandle(BottomObjects[3]->NetRefHandle));
+	}
+
+	{
+		// Filter out TopObjects[2]
+		ReplicationSystem->SetFilter(TopObjects[2]->NetRefHandle, NotRoutedFilterHandle);
+
+		// Send and deliver packets
+		Server->UpdateAndSend({ Client });
+
+		// TopObjects[2] should not exist on the client
+		UE_NET_ASSERT_TRUE(Client->IsResolvableNetRefHandle(TopObjects[0]->NetRefHandle));
+		UE_NET_ASSERT_TRUE(Client->IsResolvableNetRefHandle(TopObjects[1]->NetRefHandle));
+		UE_NET_ASSERT_FALSE(Client->IsResolvableNetRefHandle(TopObjects[2]->NetRefHandle));
+		UE_NET_ASSERT_TRUE(Client->IsResolvableNetRefHandle(TopObjects[3]->NetRefHandle));
+
+		// MiddleObjects[1-2] should not exist on the client
+		UE_NET_ASSERT_TRUE(Client->IsResolvableNetRefHandle(MiddleObjects[0]->NetRefHandle));
+		UE_NET_ASSERT_FALSE(Client->IsResolvableNetRefHandle(MiddleObjects[1]->NetRefHandle));
+		UE_NET_ASSERT_FALSE(Client->IsResolvableNetRefHandle(MiddleObjects[2]->NetRefHandle));
+		UE_NET_ASSERT_TRUE(Client->IsResolvableNetRefHandle(MiddleObjects[3]->NetRefHandle));
+
+		// BottomObjects[0-2] should not exist on the client
+		UE_NET_ASSERT_FALSE(Client->IsResolvableNetRefHandle(BottomObjects[0]->NetRefHandle));
+		UE_NET_ASSERT_FALSE(Client->IsResolvableNetRefHandle(BottomObjects[1]->NetRefHandle));
+		UE_NET_ASSERT_FALSE(Client->IsResolvableNetRefHandle(BottomObjects[2]->NetRefHandle));
+		UE_NET_ASSERT_TRUE(Client->IsResolvableNetRefHandle(BottomObjects[3]->NetRefHandle));
+	}
+
+	{
+		// Filter in MiddleObjects[1]
+		// Nothing should change because one of its parents is filtered out
+		ReplicationSystem->SetFilter(MiddleObjects[1]->NetRefHandle, InvalidNetObjectFilterHandle);
+
+		// Send and deliver packets
+		Server->UpdateAndSend({ Client });
+
+		// TopObjects[2] should not exist on the client
+		UE_NET_ASSERT_TRUE(Client->IsResolvableNetRefHandle(TopObjects[0]->NetRefHandle));
+		UE_NET_ASSERT_TRUE(Client->IsResolvableNetRefHandle(TopObjects[1]->NetRefHandle));
+		UE_NET_ASSERT_FALSE(Client->IsResolvableNetRefHandle(TopObjects[2]->NetRefHandle));
+		UE_NET_ASSERT_TRUE(Client->IsResolvableNetRefHandle(TopObjects[3]->NetRefHandle));
+
+		// MiddleObjects[1-2] should not exist on the client
+		UE_NET_ASSERT_TRUE(Client->IsResolvableNetRefHandle(MiddleObjects[0]->NetRefHandle));
+		UE_NET_ASSERT_FALSE(Client->IsResolvableNetRefHandle(MiddleObjects[1]->NetRefHandle));
+		UE_NET_ASSERT_FALSE(Client->IsResolvableNetRefHandle(MiddleObjects[2]->NetRefHandle));
+		UE_NET_ASSERT_TRUE(Client->IsResolvableNetRefHandle(MiddleObjects[3]->NetRefHandle));
+
+		// BottomObjects[0-2] should not exist on the client
+		UE_NET_ASSERT_FALSE(Client->IsResolvableNetRefHandle(BottomObjects[0]->NetRefHandle));
+		UE_NET_ASSERT_FALSE(Client->IsResolvableNetRefHandle(BottomObjects[1]->NetRefHandle));
+		UE_NET_ASSERT_FALSE(Client->IsResolvableNetRefHandle(BottomObjects[2]->NetRefHandle));
+		UE_NET_ASSERT_TRUE(Client->IsResolvableNetRefHandle(BottomObjects[3]->NetRefHandle));
+	}
+
+	{
+		// Filter in TopObjects[2]
+		ReplicationSystem->SetFilter(TopObjects[2]->NetRefHandle, InvalidNetObjectFilterHandle);
+
+		// Send and deliver packets
+		Server->UpdateAndSend({ Client });
 
-	// TopObjects[2] should not exist on the client
-	UE_NET_ASSERT_TRUE(Client->IsResolvableNetRefHandle(TopObjects[0]->NetRefHandle));
-	UE_NET_ASSERT_TRUE(Client->IsResolvableNetRefHandle(TopObjects[1]->NetRefHandle));
-	UE_NET_ASSERT_FALSE(Client->IsResolvableNetRefHandle(TopObjects[2]->NetRefHandle));
-	UE_NET_ASSERT_TRUE(Client->IsResolvableNetRefHandle(TopObjects[3]->NetRefHandle));
-
-	// MiddleObjects[1-2] should not exist on the client
-	UE_NET_ASSERT_TRUE(Client->IsResolvableNetRefHandle(MiddleObjects[0]->NetRefHandle));
-	UE_NET_ASSERT_FALSE(Client->IsResolvableNetRefHandle(MiddleObjects[1]->NetRefHandle));
-	UE_NET_ASSERT_FALSE(Client->IsResolvableNetRefHandle(MiddleObjects[2]->NetRefHandle));
-	UE_NET_ASSERT_TRUE(Client->IsResolvableNetRefHandle(MiddleObjects[3]->NetRefHandle));
-
-	// BottomObjects[0-2] should not exist on the client
-	UE_NET_ASSERT_FALSE(Client->IsResolvableNetRefHandle(BottomObjects[0]->NetRefHandle));
-	UE_NET_ASSERT_FALSE(Client->IsResolvableNetRefHandle(BottomObjects[1]->NetRefHandle));
-	UE_NET_ASSERT_FALSE(Client->IsResolvableNetRefHandle(BottomObjects[2]->NetRefHandle));
-	UE_NET_ASSERT_TRUE(Client->IsResolvableNetRefHandle(BottomObjects[3]->NetRefHandle));
+		// Objects should all exist on client
+		for (uint32 ObjectIndex = 0U; ObjectIndex < ArrayObjectCount; ++ObjectIndex)
+		{
+			UE_NET_ASSERT_TRUE(Client->IsResolvableNetRefHandle(TopObjects[ObjectIndex]->NetRefHandle));
+			UE_NET_ASSERT_TRUE(Client->IsResolvableNetRefHandle(MiddleObjects[ObjectIndex]->NetRefHandle));
+			UE_NET_ASSERT_TRUE(Client->IsResolvableNetRefHandle(BottomObjects[ObjectIndex]->NetRefHandle));
+		}
+	}
 }
 
 // Tests creation dependencies with multiple parents, multiple children, and mixed parent/child object creation order.
@@ -1387,29 +1671,100 @@ UE_NET_TEST_FIXTURE(FTestFilteredDependentObjectFixture, TestComplexCreationDepe
 		UE_NET_ASSERT_TRUE(Client->IsResolvableNetRefHandle(BottomObjects[ObjectIndex]->NetRefHandle));
 	}
 
-	// Filter out TopObjects[2]
-	ReplicationSystem->SetFilter(TopObjects[2]->NetRefHandle, NotRoutedFilterHandle);
+	{
+		// Filter out MiddleObjects[1]
+		ReplicationSystem->SetFilter(MiddleObjects[1]->NetRefHandle, NotRoutedFilterHandle);
 
-	// Send and deliver packets
-	Server->UpdateAndSend({ Client });
+		// Send and deliver packets
+		Server->UpdateAndSend({ Client });
+
+		// TopObjects should all exist on the client
+		UE_NET_ASSERT_TRUE(Client->IsResolvableNetRefHandle(TopObjects[0]->NetRefHandle));
+		UE_NET_ASSERT_TRUE(Client->IsResolvableNetRefHandle(TopObjects[1]->NetRefHandle));
+		UE_NET_ASSERT_TRUE(Client->IsResolvableNetRefHandle(TopObjects[2]->NetRefHandle));
+		UE_NET_ASSERT_TRUE(Client->IsResolvableNetRefHandle(TopObjects[3]->NetRefHandle));
+
+		// MiddleObjects[1] should not exist on the client
+		UE_NET_ASSERT_TRUE(Client->IsResolvableNetRefHandle(MiddleObjects[0]->NetRefHandle));
+		UE_NET_ASSERT_FALSE(Client->IsResolvableNetRefHandle(MiddleObjects[1]->NetRefHandle));
+		UE_NET_ASSERT_TRUE(Client->IsResolvableNetRefHandle(MiddleObjects[2]->NetRefHandle));
+		UE_NET_ASSERT_TRUE(Client->IsResolvableNetRefHandle(MiddleObjects[3]->NetRefHandle));
+
+		// BottomObjects[0-1] should not exist on the client
+		UE_NET_ASSERT_FALSE(Client->IsResolvableNetRefHandle(BottomObjects[0]->NetRefHandle));
+		UE_NET_ASSERT_FALSE(Client->IsResolvableNetRefHandle(BottomObjects[1]->NetRefHandle));
+		UE_NET_ASSERT_TRUE(Client->IsResolvableNetRefHandle(BottomObjects[2]->NetRefHandle));
+		UE_NET_ASSERT_TRUE(Client->IsResolvableNetRefHandle(BottomObjects[3]->NetRefHandle));
+	}
+
+	{
+		// Filter out TopObjects[2]
+		ReplicationSystem->SetFilter(TopObjects[2]->NetRefHandle, NotRoutedFilterHandle);
+
+		// Send and deliver packets
+		Server->UpdateAndSend({ Client });
+
+		// TopObjects[2] should not exist on the client
+		UE_NET_ASSERT_TRUE(Client->IsResolvableNetRefHandle(TopObjects[0]->NetRefHandle));
+		UE_NET_ASSERT_TRUE(Client->IsResolvableNetRefHandle(TopObjects[1]->NetRefHandle));
+		UE_NET_ASSERT_FALSE(Client->IsResolvableNetRefHandle(TopObjects[2]->NetRefHandle));
+		UE_NET_ASSERT_TRUE(Client->IsResolvableNetRefHandle(TopObjects[3]->NetRefHandle));
+
+		// MiddleObjects[1-2] should not exist on the client
+		UE_NET_ASSERT_TRUE(Client->IsResolvableNetRefHandle(MiddleObjects[0]->NetRefHandle));
+		UE_NET_ASSERT_FALSE(Client->IsResolvableNetRefHandle(MiddleObjects[1]->NetRefHandle));
+		UE_NET_ASSERT_FALSE(Client->IsResolvableNetRefHandle(MiddleObjects[2]->NetRefHandle));
+		UE_NET_ASSERT_TRUE(Client->IsResolvableNetRefHandle(MiddleObjects[3]->NetRefHandle));
+
+		// BottomObjects[0-2] should not exist on the client
+		UE_NET_ASSERT_FALSE(Client->IsResolvableNetRefHandle(BottomObjects[0]->NetRefHandle));
+		UE_NET_ASSERT_FALSE(Client->IsResolvableNetRefHandle(BottomObjects[1]->NetRefHandle));
+		UE_NET_ASSERT_FALSE(Client->IsResolvableNetRefHandle(BottomObjects[2]->NetRefHandle));
+		UE_NET_ASSERT_TRUE(Client->IsResolvableNetRefHandle(BottomObjects[3]->NetRefHandle));
+	}
+
+	{
+		// Filter in MiddleObjects[1]
+		// Nothing should change because one of its parents is filtered out
+		ReplicationSystem->SetFilter(MiddleObjects[1]->NetRefHandle, InvalidNetObjectFilterHandle);
 
-	// TopObjects[2] should not exist on the client
-	UE_NET_ASSERT_TRUE(Client->IsResolvableNetRefHandle(TopObjects[0]->NetRefHandle));
-	UE_NET_ASSERT_TRUE(Client->IsResolvableNetRefHandle(TopObjects[1]->NetRefHandle));
-	UE_NET_ASSERT_FALSE(Client->IsResolvableNetRefHandle(TopObjects[2]->NetRefHandle));
-	UE_NET_ASSERT_TRUE(Client->IsResolvableNetRefHandle(TopObjects[3]->NetRefHandle));
-
-	// MiddleObjects[1-2] should not exist on the client
-	UE_NET_ASSERT_TRUE(Client->IsResolvableNetRefHandle(MiddleObjects[0]->NetRefHandle));
-	UE_NET_ASSERT_FALSE(Client->IsResolvableNetRefHandle(MiddleObjects[1]->NetRefHandle));
-	UE_NET_ASSERT_FALSE(Client->IsResolvableNetRefHandle(MiddleObjects[2]->NetRefHandle));
-	UE_NET_ASSERT_TRUE(Client->IsResolvableNetRefHandle(MiddleObjects[3]->NetRefHandle));
-
-	// BottomObjects[0-2] should not exist on the client
-	UE_NET_ASSERT_FALSE(Client->IsResolvableNetRefHandle(BottomObjects[0]->NetRefHandle));
-	UE_NET_ASSERT_FALSE(Client->IsResolvableNetRefHandle(BottomObjects[1]->NetRefHandle));
-	UE_NET_ASSERT_FALSE(Client->IsResolvableNetRefHandle(BottomObjects[2]->NetRefHandle));
-	UE_NET_ASSERT_TRUE(Client->IsResolvableNetRefHandle(BottomObjects[3]->NetRefHandle));
+		// Send and deliver packets
+		Server->UpdateAndSend({ Client });
+
+		// TopObjects[2] should not exist on the client
+		UE_NET_ASSERT_TRUE(Client->IsResolvableNetRefHandle(TopObjects[0]->NetRefHandle));
+		UE_NET_ASSERT_TRUE(Client->IsResolvableNetRefHandle(TopObjects[1]->NetRefHandle));
+		UE_NET_ASSERT_FALSE(Client->IsResolvableNetRefHandle(TopObjects[2]->NetRefHandle));
+		UE_NET_ASSERT_TRUE(Client->IsResolvableNetRefHandle(TopObjects[3]->NetRefHandle));
+
+		// MiddleObjects[1-2] should not exist on the client
+		UE_NET_ASSERT_TRUE(Client->IsResolvableNetRefHandle(MiddleObjects[0]->NetRefHandle));
+		UE_NET_ASSERT_FALSE(Client->IsResolvableNetRefHandle(MiddleObjects[1]->NetRefHandle));
+		UE_NET_ASSERT_FALSE(Client->IsResolvableNetRefHandle(MiddleObjects[2]->NetRefHandle));
+		UE_NET_ASSERT_TRUE(Client->IsResolvableNetRefHandle(MiddleObjects[3]->NetRefHandle));
+
+		// BottomObjects[0-2] should not exist on the client
+		UE_NET_ASSERT_FALSE(Client->IsResolvableNetRefHandle(BottomObjects[0]->NetRefHandle));
+		UE_NET_ASSERT_FALSE(Client->IsResolvableNetRefHandle(BottomObjects[1]->NetRefHandle));
+		UE_NET_ASSERT_FALSE(Client->IsResolvableNetRefHandle(BottomObjects[2]->NetRefHandle));
+		UE_NET_ASSERT_TRUE(Client->IsResolvableNetRefHandle(BottomObjects[3]->NetRefHandle));
+	}
+
+	{
+		// Filter in TopObjects[2]
+		ReplicationSystem->SetFilter(TopObjects[2]->NetRefHandle, InvalidNetObjectFilterHandle);
+
+		// Send and deliver packets
+		Server->UpdateAndSend({ Client });
+
+		// Objects should all exist on client
+		for (uint32 ObjectIndex = 0U; ObjectIndex < ArrayObjectCount; ++ObjectIndex)
+		{
+			UE_NET_ASSERT_TRUE(Client->IsResolvableNetRefHandle(TopObjects[ObjectIndex]->NetRefHandle));
+			UE_NET_ASSERT_TRUE(Client->IsResolvableNetRefHandle(MiddleObjects[ObjectIndex]->NetRefHandle));
+			UE_NET_ASSERT_TRUE(Client->IsResolvableNetRefHandle(BottomObjects[ObjectIndex]->NetRefHandle));
+		}
+	}
 }
 
 UE_NET_TEST_FIXTURE(FTestFilteredDependentObjectFixture, TestCircularDependencyIsDenied)
diff --git a/Engine/Source/Runtime/Net/Iris/Private/Iris/ReplicationSystem/Filtering/ReplicationFiltering.cpp b/Engine/Source/Runtime/Net/Iris/Private/Iris/ReplicationSystem/Filtering/ReplicationFiltering.cpp
--- a/Engine/Source/Runtime/Net/Iris/Private/Iris/ReplicationSystem/Filtering/ReplicationFiltering.cpp
+++ b/Engine/Source/Runtime/Net/Iris/Private/Iris/ReplicationSystem/Filtering/ReplicationFiltering.cpp
@@ -31,6 +31,9 @@ static FAutoConsoleVariableRef CVarRepFilterCullNonRelevant(TEXT("Net.Iris.CullN
 bool bCVarRepFilterValidateNoSubObjectInScopeWithFilteredOutRootObject = false;
 static FAutoConsoleVariableRef CVarRepFilterValidateNoSubObjectInScopeWithFilteredOutRootObject(TEXT("Net.Iris.Filtering.ValidateNobSubObjectInScopeWithFilteredOutRootObject"), bCVarRepFilterValidateNoSubObjectInScopeWithFilteredOutRootObject, TEXT("Validate there are no subobjects in scope with a filtered out root object."), ECVF_Default);
 
+bool bCVarRepFilterUpdateCreationDependenciesOnRelevancyChange = false;
+static FAutoConsoleVariableRef CVarRepFilterUpdateCreationDependenciesOnRelevancyChange(TEXT("Net.Iris.Filtering.UpdateCreationDependenciesOnRelevancyChange"), bCVarRepFilterUpdateCreationDependenciesOnRelevancyChange, TEXT("When enabled, only update filtering of objects in creation dependency trees when they change relevancy."), ECVF_Default);
+
 FName GetStaticFilterName(FNetObjectFilterHandle Filter)
 {
 	switch (Filter)
@@ -841,6 +844,9 @@ void FReplicationFiltering::SetPerConnectionListsSize(FPerConnectionInfo& Connec
 	// The combined result of scoped objects, connection filtering and group exclusion filtering.
 	ConnectionInfo.ObjectsInScopeBeforeDynamicFiltering.SetNumBits(NewMaxInternalIndex);
 
+	// Creation dependents that have been filtered out because their parents are filtered.
+	ConnectionInfo.FilteredOutCreationDependents.SetNumBits(NewMaxInternalIndex);
+
 	// The final result of all filtering.
 	ConnectionInfo.ObjectsInScope.SetNumBits(NewMaxInternalIndex);
 
@@ -1272,6 +1278,85 @@ void FReplicationFiltering::PreUpdateDynamicFiltering()
 	}
 }
 
+void FReplicationFiltering::FilterOutCreationDependents(uint32 ParentIndex, FPerConnectionInfo& ConnectionInfo, const FNetBitArrayView ObjectsWithCreationDependents) const
+{
+	if (ConnectionInfo.FilteredOutCreationDependents.GetBit(ParentIndex))
+	{
+		// The parent and its descendents have already been filtered out.
+		return;
+	}
+
+	TArray<FInternalNetRefIndex, TInlineAllocator<64>> ObjectsToProcess;
+	ObjectsToProcess.Add(ParentIndex);
+	while (ObjectsToProcess.Num() > 0)
+	{
+		const FInternalNetRefIndex CurrentIndex = ObjectsToProcess.Pop(EAllowShrinking::No);
+		if (ConnectionInfo.FilteredOutCreationDependents.GetBit(CurrentIndex))
+		{
+			// This section of the dependency graph has already been filtered out, so we can skip processing it.
+			continue;
+		}
+
+		ConnectionInfo.FilteredOutCreationDependents.SetBit(CurrentIndex);
+
+		TConstArrayView<const FInternalNetRefIndex> Children = NetRefHandleManager->GetCreationDependents(CurrentIndex);
+		// Add the children to process them and their descendents.
+		ObjectsToProcess.Append(Children);
+	}
+}
+
+void FReplicationFiltering::FilterInCreationDependents(uint32 ParentIndex, FPerConnectionInfo& ConnectionInfo, const FNetBitArrayView ObjectsWithCreationDependents) const
+{
+	TArray<FInternalNetRefIndex, TInlineAllocator<64>> ObjectsToProcess;
+	ObjectsToProcess.Add(ParentIndex);
+	while (ObjectsToProcess.Num() > 0)
+	{
+		const FInternalNetRefIndex CurrentIndex = ObjectsToProcess.Pop(EAllowShrinking::No);
+
+		bool bFilteredOutByParent = false;
+		TConstArrayView<const FInternalNetRefIndex> Dependencies = NetRefHandleManager->GetCreationDependencies(CurrentIndex);
+		for (const FInternalNetRefIndex DependencyIndex : Dependencies)
+		{
+			if (ConnectionInfo.FilteredOutCreationDependents.GetBit(DependencyIndex))
+			{
+				// This object has a parent that has been filtered out, so we can't filter it in.
+				//
+				// Note: In the case that an object has N>1 parents that would be filtered in by one call to FilterInCreationDependents,
+				//       that object will be processed N times and only be filtered in the Nth time.
+				// Example: W starts filtered out. X, Y, and Z would be relevant, but creation dependency filters them out.
+				//          W is filtered in. Because of DFS, we process in the order W->Y->Z->X->Z. Z is only filtered in the 2nd time it is processed.
+				//   W
+				//  / \
+				// X   Y
+				//  \ /
+				//   Z
+				bFilteredOutByParent = true;
+				break;
+			}
+		}
+
+		if (bFilteredOutByParent)
+		{
+			// This object has a filtered parent, we can't filter them or their children in.
+			continue;
+		}
+
+		ConnectionInfo.FilteredOutCreationDependents.ClearBit(CurrentIndex);
+
+		TConstArrayView<const FInternalNetRefIndex> Children = NetRefHandleManager->GetCreationDependents(CurrentIndex);
+		for (const FInternalNetRefIndex ChildIndex : Children)
+		{
+			if (ConnectionInfo.DynamicFilteredOutObjectsHysteresisAdjusted.GetBit(ChildIndex))
+			{
+				// This child is filtered out, so don't process their subtree.
+				continue;
+			}
+
+			ObjectsToProcess.Add(ChildIndex);
+		}
+	}
+}
+
 void FReplicationFiltering::UpdateCreationDependentParent(uint32 ChildIndex, const FNetBitArrayView ObjectsWithCreationDependencies, FNetBitArrayView OutConnectionObjectsInScope, FNetBitArrayView ProcessedObjects, bool bRecursive) const
 {
 	const bool bStartedRelevant = OutConnectionObjectsInScope.IsBitSet(ChildIndex);
@@ -1356,6 +1441,8 @@ void FReplicationFiltering::UpdateDynamicFiltering()
 	const FNetBitArrayView DependentObjectsView = NetRefHandleManager->GetDependentObjectInternalIndices();
 	const uint32* DependentObjectsData = DependentObjectsView.GetDataChecked(WordCountForObjectBitArrays);
 
+	const FNetBitArrayView ObjectsWithCreationDependents = NetRefHandleManager->GetObjectsWithCreationDependents();
+
 	const FNetBitArrayView ObjectsWithCreationDependencies = NetRefHandleManager->GetObjectsWithCreationDependencies();
 	const uint32* ObjectsWithCreationDependenciesData = ObjectsWithCreationDependencies.GetDataChecked(WordCountForObjectBitArrays);
 	
@@ -1593,6 +1680,73 @@ void FReplicationFiltering::UpdateDynamicFiltering()
 
 		ensure(!bCVarRepFilterValidateNoSubObjectInScopeWithFilteredOutRootObject || !HasSubObjectInScopeWithFilteredOutRootObject(DynamicFilteredOutObjectsHysteresisAdjusted));
 
+		if (bCVarRepFilterUpdateCreationDependenciesOnRelevancyChange)
+		{
+			IRIS_PROFILER_SCOPE(FReplicationFiltering_UpdateCreationDependencies);
+
+			// This will be run on any object that was filtered out in the last frame.
+			auto ProcessFilteredOutObject = [&](const uint32 ObjectIndex)
+			{
+				if (ObjectsWithCreationDependents.IsBitSet(ObjectIndex))
+				{
+					// Filter out this object and all of its descendents
+					FilterOutCreationDependents(ObjectIndex, ConnectionInfo, ObjectsWithCreationDependents);
+				}
+				else if (ObjectsWithCreationDependencies.IsBitSet(ObjectIndex))
+				{
+					// This object is still a part of a dependency graph, so mark as filtered out.
+					ConnectionInfo.FilteredOutCreationDependents.SetBit(ObjectIndex);
+				}
+				else
+				{
+					// This object isn't a part of a dependency graph, so it shouldn't be marked.
+					ConnectionInfo.FilteredOutCreationDependents.ClearBit(ObjectIndex);
+				}
+			};
+
+			// For parents that have been filtered out this frame, filter out their dependents.
+			if (HysteresisState.Mode == EHysteresisProcessingMode::Enabled)
+			{
+				// Only filter creation dependents of objects in FilteredOutObjects if they didn't enter hysteresis.
+				for (const uint32 ObjectIndex : FilteredOutObjects)
+				{
+					if (!ConnectionInfo.HysteresisUpdater.IsObjectUpdated(ObjectIndex))
+					{
+						ProcessFilteredOutObject(ObjectIndex);
+					}
+				}
+
+				// FilteredOutByHysteresisObjects will contain all of the hysteresis objects whose frame timer just expired.
+				for (const uint32 ObjectIndex : FilteredOutByHysteresisObjects)
+				{
+					ProcessFilteredOutObject(ObjectIndex);
+				}
+			}
+			else
+			{
+				// When hysteresis is disabled, FilteredOutObjects will actually contain all of the objects that have been filtered out.
+				for (const uint32 ObjectIndex : FilteredOutObjects)
+				{
+					ProcessFilteredOutObject(ObjectIndex);
+				}
+			}
+
+			// For parents that have been filtered in this frame, filter in their dependents that aren't otherwise filtered out.
+			for (const uint32 ObjectIndex : FilteredInObjects)
+			{
+				if (ObjectsWithCreationDependents.IsBitSet(ObjectIndex) || ObjectsWithCreationDependencies.IsBitSet(ObjectIndex))
+				{
+					// Even if this object doesn't have creation dependents, it might need to be filtered in if a dependency link to a filtered out parent was removed.
+					FilterInCreationDependents(ObjectIndex, ConnectionInfo, ObjectsWithCreationDependents);
+				}
+				else
+				{
+					// This object doesn't have dependents or dependencies.
+					ConnectionInfo.FilteredOutCreationDependents.ClearBit(ObjectIndex);
+				}
+			}
+		}
+
 		// Update the entire scope for the connection.
 		{
 			IRIS_PROFILER_SCOPE(FReplicationFiltering_UpdateConnectionScope);
@@ -1600,11 +1754,13 @@ void FReplicationFiltering::UpdateDynamicFiltering()
 			uint32* ObjectsInScopeData = ConnectionInfo.ObjectsInScope.GetDataChecked(WordCountForObjectBitArrays);
 			const uint32* ObjectsInScopeBeforeDynamicFilteringData = ConnectionInfo.ObjectsInScopeBeforeDynamicFiltering.GetDataChecked(WordCountForObjectBitArrays);
 			const uint32* DynamicFilteredOutObjectsHysteresisAdjustedData = DynamicFilteredOutObjectsHysteresisAdjusted.GetDataChecked(WordCountForObjectBitArrays);
+			const uint32* FilteredOutCreationDependentsData = ConnectionInfo.FilteredOutCreationDependents.GetDataChecked(WordCountForObjectBitArrays);
 			for (SIZE_T WordIt = 0, WordEndIt = WordCountForObjectBitArrays; WordIt != WordEndIt; ++WordIt)
 			{
 				const uint32 ObjectsInScopeBeforeWord = ObjectsInScopeBeforeDynamicFilteringData[WordIt];
 				const uint32 DynamicFilteredOutWord = DynamicFilteredOutObjectsHysteresisAdjustedData[WordIt];
-				ObjectsInScopeData[WordIt] = ObjectsInScopeBeforeWord & ~DynamicFilteredOutWord;
+				const uint32 FilteredOutCreationDependentsWord = FilteredOutCreationDependentsData[WordIt];
+				ObjectsInScopeData[WordIt] = ObjectsInScopeBeforeWord & ~DynamicFilteredOutWord & ~FilteredOutCreationDependentsWord;
 			}
 
 			// Unconditionally filter out filtered out dependent objects from ObjectsInScope as GetDependentObjectFilterStatus, called later, requires it. Subobjects are dealt with later.
@@ -1615,6 +1771,7 @@ void FReplicationFiltering::UpdateDynamicFiltering()
 		}
 
 		// Update dependent objects that can only be relevant if their parent also is
+		if (!bCVarRepFilterUpdateCreationDependenciesOnRelevancyChange)
 		{
 			IRIS_PROFILER_SCOPE(FReplicationFiltering_UpdateCreationDependencies);
 			if (ObjectsWithCreationDependencies.IsAnyBitSet())
@@ -2583,6 +2740,26 @@ void FReplicationFiltering::NotifyRemovedDependentObject(FInternalNetRefIndex Ob
 	ObjectsRequiringDynamicFilterUpdate.SetBit(ObjectIndex);
 }
 
+void FReplicationFiltering::NotifyAddedCreationDependency(FInternalNetRefIndex ParentIndex, FInternalNetRefIndex ChildIndex)
+{
+	if (bCVarRepFilterUpdateCreationDependenciesOnRelevancyChange)
+	{
+		// This causes the parent and child to be reprocessed in the creation dependency update.
+		ObjectsRequiringDynamicFilterUpdate.SetBit(ParentIndex);
+		ObjectsRequiringDynamicFilterUpdate.SetBit(ChildIndex);
+	}
+}
+
+void FReplicationFiltering::NotifyRemovedCreationDependency(FInternalNetRefIndex ParentIndex, FInternalNetRefIndex ChildIndex)
+{
+	if (bCVarRepFilterUpdateCreationDependenciesOnRelevancyChange)
+	{
+		// This causes the parent and child to be reprocessed in the creation dependency update.
+		ObjectsRequiringDynamicFilterUpdate.SetBit(ParentIndex);
+		ObjectsRequiringDynamicFilterUpdate.SetBit(ChildIndex);
+	}
+}
+
 // Determine whether a dynamically filtered out dependent object should be in scope or not. Note that this function is expected to only be called for dependent objects NOT in scope. The Visited set is expected to be empty when calling this function.
 ENetFilterStatus FReplicationFiltering::GetDependentObjectFilterStatus(uint32 ConnectionId, FInternalNetRefIndex DependentObjectIndex, TSet<FInternalNetRefIndex>& Visited)
 {
@@ -2977,6 +3154,7 @@ void FReplicationFiltering::FPerConnectionInfo::Deinit()
 	GroupExcludedObjects.Empty();
 	ObjectsInScopeBeforeDynamicFiltering.Empty();
 	GroupIncludedObjects.Empty();
+	FilteredOutCreationDependents.Empty();
 	ObjectsInScope.Empty();
 	DynamicFilteredOutObjects.Empty();
 	InProgressDynamicFilteredOutObjects.Empty();
diff --git a/Engine/Source/Runtime/Net/Iris/Private/Iris/ReplicationSystem/Filtering/ReplicationFiltering.h b/Engine/Source/Runtime/Net/Iris/Private/Iris/ReplicationSystem/Filtering/ReplicationFiltering.h
--- a/Engine/Source/Runtime/Net/Iris/Private/Iris/ReplicationSystem/Filtering/ReplicationFiltering.h
+++ b/Engine/Source/Runtime/Net/Iris/Private/Iris/ReplicationSystem/Filtering/ReplicationFiltering.h
@@ -136,6 +136,9 @@ public:
 
 	void NotifyAddedDependentObject(FInternalNetRefIndex ObjectIndex);
 	void NotifyRemovedDependentObject(FInternalNetRefIndex ObjectIndex);
+
+	void NotifyAddedCreationDependency(FInternalNetRefIndex ParentIndex, FInternalNetRefIndex ChildIndex);
+	void NotifyRemovedCreationDependency(FInternalNetRefIndex ParentIndex, FInternalNetRefIndex ChildIndex);
 	
 	// SubObjectFilter status
 	void AddSubObjectFilter(FNetObjectGroupHandle GroupHandle);
@@ -161,6 +164,8 @@ private:
 		FNetBitArray ObjectsInScopeBeforeDynamicFiltering;
 		// Group inclusion filtering is assumed to happen seldom. Applied after dynamic filtering, but before dependent objects.
 		FNetBitArray GroupIncludedObjects;
+		// Creation dependents that have been filtered because a dependency has been filtered.
+		FNetBitArray FilteredOutCreationDependents;
 		// Objects in scope after all kinds of filtering, including dynamic filtering, has been applied
 		FNetBitArray ObjectsInScope;
 		// Which objects are filtered out after dynamic filters have been processed.
@@ -238,6 +243,12 @@ private:
 
 	bool HasDynamicFilters() const;
 
+	// Mark creation dependents as filtered out.
+	void FilterOutCreationDependents(uint32 ParentIndex, FPerConnectionInfo& ConnectionInfo, const FNetBitArrayView ObjectsWithCreationDependents) const;
+
+	// Mark creation dependents as filtered in, if they aren't filtered out by other dependencies.
+	void FilterInCreationDependents(uint32 ParentIndex, FPerConnectionInfo& ConnectionInfo, const FNetBitArrayView ObjectsWithCreationDependents) const;
+
 	void UpdateCreationDependentParent(uint32 ChildIndex, const FNetBitArrayView ObjectsWithCreationDependencies, FNetBitArrayView OutConnectionObjectsInScope, FNetBitArrayView ProcessedObjects, bool bIsRecursive) const;
 
 	// Helper to update and reset group exclusion filter effects if objects are removed from a filter or after a filter status change, returns true if the group filter was changed
diff --git a/Engine/Source/Runtime/Net/Iris/Private/Iris/ReplicationSystem/NetDependencyData.cpp b/Engine/Source/Runtime/Net/Iris/Private/Iris/ReplicationSystem/NetDependencyData.cpp
--- a/Engine/Source/Runtime/Net/Iris/Private/Iris/ReplicationSystem/NetDependencyData.cpp
+++ b/Engine/Source/Runtime/Net/Iris/Private/Iris/ReplicationSystem/NetDependencyData.cpp
@@ -35,6 +35,11 @@ void FNetDependencyData::FreeStoredDependencyDataForObject(FInternalNetRefIndex
 			DependentObjectInfosStorage.RemoveAt(Entry->DependentObjectsInfoArrayIndex);
 		}
 
+		if (Entry->CreationDependentArrayIndex != FDependencyInfo::InvalidCacheIndex)
+		{
+			CreationDependentInfosStorage.RemoveAt(Entry->CreationDependentArrayIndex);
+		}
+
 		if (Entry->CreationDependencyArrayIndex != FDependencyInfo::InvalidCacheIndex)
 		{
 			CreationDependencyInfosStorage.RemoveAt(Entry->CreationDependencyArrayIndex);
@@ -121,6 +126,36 @@ FNetDependencyData::FInternalNetRefIndexArray& FNetDependencyData::GetOrCreateIn
 	}
 }
 
+FNetDependencyData::FCreationDependentInfoArray& FNetDependencyData::GetOrCreateCreationDependentInfoArray(FInternalNetRefIndex ParentIndex)
+{
+	FDependencyInfo& Entry = GetOrCreateCacheEntry(ParentIndex);
+	
+	if (Entry.CreationDependentArrayIndex == FDependencyInfo::InvalidCacheIndex)
+	{
+		FSparseArrayAllocationInfo AllocInfo = CreationDependentInfosStorage.AddUninitialized();
+		Entry.CreationDependentArrayIndex = AllocInfo.Index;
+
+		FCreationDependentInfoArray* InfoArray = new (AllocInfo.Pointer) FCreationDependentInfoArray();
+		
+		return *InfoArray;
+	}
+	else
+	{
+		return CreationDependentInfosStorage[Entry.CreationDependentArrayIndex];
+	}
+}
+
+void FNetDependencyData::FreeCreationDependentInfoArray(FInternalNetRefIndex ParentIndex)
+{
+	if (FDependencyInfo* Entry = DependencyInfos.Find(ParentIndex))
+	{
+		if (Entry->CreationDependentArrayIndex != FDependencyInfo::InvalidCacheIndex)
+		{
+			CreationDependentInfosStorage.RemoveAt(Entry->CreationDependentArrayIndex);
+			Entry->CreationDependentArrayIndex = FDependencyInfo::InvalidCacheIndex;
+		}
+	}
+}
 
 FNetDependencyData::FCreationDependencyInfoArray& FNetDependencyData::GetOrCreateCreationDependencyInfoArray(FInternalNetRefIndex ChildIndex)
 {
diff --git a/Engine/Source/Runtime/Net/Iris/Private/Iris/ReplicationSystem/NetDependencyData.h b/Engine/Source/Runtime/Net/Iris/Private/Iris/ReplicationSystem/NetDependencyData.h
--- a/Engine/Source/Runtime/Net/Iris/Private/Iris/ReplicationSystem/NetDependencyData.h
+++ b/Engine/Source/Runtime/Net/Iris/Private/Iris/ReplicationSystem/NetDependencyData.h
@@ -38,6 +38,7 @@ public:
 	typedef TArray<FInternalNetRefIndex, TInlineAllocator<8>> FInternalNetRefIndexArray;
 	typedef TArray<FLifeTimeConditionStorage, TInlineAllocator<8>> FSubObjectConditionalsArray;
 	typedef TArray<FDependentObjectInfo, TInlineAllocator<8>> FDependentObjectInfoArray;
+	typedef TArray<FInternalNetRefIndex, TInlineAllocator<8>> FCreationDependentInfoArray;
 	typedef TArray<FInternalNetRefIndex, TInlineAllocator<2>> FCreationDependencyInfoArray;
 
 	enum EArrayType
@@ -113,6 +114,26 @@ public:
 		return true;
 	}
 
+	FCreationDependentInfoArray& GetOrCreateCreationDependentInfoArray(FInternalNetRefIndex ParentIndex);
+
+	void FreeCreationDependentInfoArray(FInternalNetRefIndex ParentIndex);
+
+	FCreationDependentInfoArray* GetCreationDependentInfoArray(FInternalNetRefIndex ParentIndex)
+	{
+		const FDependencyInfo* Entry = DependencyInfos.Find(ParentIndex);
+		const uint32 ArrayIndex = Entry ? Entry->CreationDependentArrayIndex : FDependencyInfo::InvalidCacheIndex;
+		if (ArrayIndex != FDependencyInfo::InvalidCacheIndex)
+		{
+			return &CreationDependentInfosStorage[ArrayIndex];
+		}
+		return nullptr;
+	}
+
+	const FCreationDependentInfoArray* GetCreationDependentInfoArray(FInternalNetRefIndex ParentIndex) const
+	{
+		return const_cast<FNetDependencyData*>(this)->GetCreationDependentInfoArray(ParentIndex);
+	}
+
 	/** Create or return the creation dependency list of an object */
 	FCreationDependencyInfoArray& GetOrCreateCreationDependencyInfoArray(FInternalNetRefIndex ChildIndex);
 	
@@ -132,13 +153,7 @@ public:
 
 	const FCreationDependencyInfoArray* GetCreationDependencyInfoArray(FInternalNetRefIndex ChildIndex) const
 	{
-		const FDependencyInfo* Entry = DependencyInfos.Find(ChildIndex);
-		const uint32 ArrayIndex = Entry ? Entry->CreationDependencyArrayIndex : FDependencyInfo::InvalidCacheIndex;
-		if (ArrayIndex != FDependencyInfo::InvalidCacheIndex)
-		{
-			return &CreationDependencyInfosStorage[ArrayIndex];
-		}
-		return nullptr;
+		return const_cast<FNetDependencyData*>(this)->GetCreationDependencyInfoArray(ChildIndex);
 	}
 
 	template<EArrayType TypeIndex>
@@ -178,6 +193,7 @@ private:
 		uint32 ArrayIndices[EArrayType::Count] = { InvalidCacheIndex, InvalidCacheIndex, InvalidCacheIndex, InvalidCacheIndex };
 		uint32 SubObjectConditionalArrayIndex = InvalidCacheIndex;
 		uint32 DependentObjectsInfoArrayIndex = InvalidCacheIndex;
+		uint32 CreationDependentArrayIndex = InvalidCacheIndex;
 		uint32 CreationDependencyArrayIndex = InvalidCacheIndex;
 	};
 
@@ -197,6 +213,9 @@ private:
 	// Storage for DependentObjects traits and info (bound on the Parent)
 	TSparseArray<FDependentObjectInfoArray> DependentObjectInfosStorage;
 
+	// Storage for creation dependent infos (bound on the Parent)
+	TSparseArray<FCreationDependentInfoArray> CreationDependentInfosStorage;
+
 	// Storage for creation dependency infos (bound on the Child)
 	TSparseArray<FCreationDependencyInfoArray> CreationDependencyInfosStorage;
 
diff --git a/Engine/Source/Runtime/Net/Iris/Private/Iris/ReplicationSystem/NetRefHandleManager.cpp b/Engine/Source/Runtime/Net/Iris/Private/Iris/ReplicationSystem/NetRefHandleManager.cpp
--- a/Engine/Source/Runtime/Net/Iris/Private/Iris/ReplicationSystem/NetRefHandleManager.cpp
+++ b/Engine/Source/Runtime/Net/Iris/Private/Iris/ReplicationSystem/NetRefHandleManager.cpp
@@ -137,6 +137,7 @@ void FNetRefHandleManager::Init(const FInitParams& InitParams)
 		InitNetBitArray(&AssignedInternalIndices);
 		InitNetBitArray(&SubObjectInternalIndices);
 		InitNetBitArray(&DependentObjectInternalIndices);
+		InitNetBitArray(&ObjectsWithCreationDependents);
 		InitNetBitArray(&ObjectsWithCreationDependencies);
 		InitNetBitArray(&ObjectsWithDependentObjectsInternalIndices);
 		InitNetBitArray(&DestroyedStartupObjectInternalIndices);
@@ -1112,60 +1113,101 @@ bool FNetRefHandleManager::AddCreationDependency(FNetRefHandle Parent, FNetRefHa
 	FInternalNetRefIndex ParentIndex = GetInternalIndex(Parent);
 	FInternalNetRefIndex ChildIndex = GetInternalIndex(Child);
 
+	return InternalAddCreationDependency(ParentIndex, ChildIndex);
+}
+
+void FNetRefHandleManager::RemoveCreationDependency(FNetRefHandle Parent, FNetRefHandle Child)
+{
+	FInternalNetRefIndex ParentIndex = GetInternalIndex(Parent);
+	FInternalNetRefIndex ChildIndex = GetInternalIndex(Child);
+
+	InternalRemoveCreationDependency(ParentIndex, ChildIndex);
+}
+
+bool FNetRefHandleManager::InternalAddCreationDependency(FInternalNetRefIndex ParentIndex, FInternalNetRefIndex ChildIndex)
+{
 	if (ParentIndex == InvalidInternalIndex || ChildIndex == InvalidInternalIndex)
 	{
-		ensureMsgf(false, TEXT("AddCreationDependency failed due to invalid reference. Parent: %s | Child: %s"), *PrintObjectFromNetRefHandle(Parent), *PrintObjectFromNetRefHandle(Child));
+		ensureMsgf(false, TEXT("AddCreationDependency failed due to invalid reference. Parent: %s | Child: %s"), *PrintObjectFromIndex(ParentIndex), *PrintObjectFromIndex(ChildIndex));
 		return false;
 	}
 
 	if (IsSubObject(ParentIndex) || IsSubObject(ChildIndex))
 	{
-		UE_LOG(LogIris, Error, TEXT("AddCreationDependency only works on RootObjects. Parent: %s | Child: %s"), *PrintObjectFromNetRefHandle(Parent), *PrintObjectFromNetRefHandle(Child));
+		UE_LOG(LogIris, Error, TEXT("AddCreationDependency only works on RootObjects. Parent: %s | Child: %s"), *PrintObjectFromIndex(ParentIndex), *PrintObjectFromIndex(ChildIndex));
 		ensureMsgf(false, TEXT("AddCreationDependency not supported on subobjects"));
 		return false;
 	}
 
-	FNetDependencyData::FCreationDependencyInfoArray& CreationDependencyInfo = SubObjects.GetOrCreateCreationDependencyInfoArray(ChildIndex);
-
-	// Make sure the parent isn't already in the list
-	if (CreationDependencyInfo.Find(ParentIndex) != INDEX_NONE)
+	// Add child to parent's list.
+	FNetDependencyData::FCreationDependentInfoArray& CreationDependentInfo = SubObjects.GetOrCreateCreationDependentInfoArray(ParentIndex);
+	// Only add the child to the list if it's not already present.
+	if (CreationDependentInfo.Find(ChildIndex) == INDEX_NONE)
 	{
-		return true;
+		CreationDependentInfo.Add(ChildIndex);
+		ObjectsWithCreationDependents.SetBit(ParentIndex);
 	}
 
-	CreationDependencyInfo.Add(ParentIndex);
-
-	ObjectsWithCreationDependencies.SetBit(ChildIndex);
+	// Add parent to child's list.
+	FNetDependencyData::FCreationDependencyInfoArray& CreationDependencyInfo = SubObjects.GetOrCreateCreationDependencyInfoArray(ChildIndex);
+	// Only add the parent to the list if it's not already present.
+	if (CreationDependencyInfo.Find(ParentIndex) == INDEX_NONE)
+	{
+		CreationDependencyInfo.Add(ParentIndex);
+		ObjectsWithCreationDependencies.SetBit(ChildIndex);
+	}
 
 	return true;
 }
 
-void FNetRefHandleManager::RemoveCreationDependency(FNetRefHandle Parent, FNetRefHandle Child)
+void FNetRefHandleManager::InternalRemoveCreationDependency(FInternalNetRefIndex ParentIndex, FInternalNetRefIndex ChildIndex)
 {
-	FInternalNetRefIndex ParentIndex = GetInternalIndex(Parent);
-	FInternalNetRefIndex ChildIndex = GetInternalIndex(Child);
-
 	if (ParentIndex == InvalidInternalIndex || ChildIndex == InvalidInternalIndex)
 	{
-		ensureMsgf(false, TEXT("RemoveCreationDependency failed due to invalid reference. Parent: %s | Child: %s"), *PrintObjectFromNetRefHandle(Parent), *PrintObjectFromNetRefHandle(Child));
+		ensureMsgf(false, TEXT("RemoveCreationDependency failed due to invalid reference. Parent: %s | Child: %s"), *PrintObjectFromIndex(ParentIndex), *PrintObjectFromIndex(ChildIndex));
 		return;
 	}
 
+	// Remove child from parent's list.
+	FNetDependencyData::FCreationDependentInfoArray* CreationDependentInfo = SubObjects.GetCreationDependentInfoArray(ParentIndex);
+	if (CreationDependentInfo)
+	{
+		CreationDependentInfo->Remove(ChildIndex);
+		if (CreationDependentInfo->IsEmpty())
+		{
+			ObjectsWithCreationDependents.ClearBit(ParentIndex);
+		}
+	}
+
+	// Remove parent from child's list.
 	FNetDependencyData::FCreationDependencyInfoArray* CreationDependencyInfo = SubObjects.GetCreationDependencyInfoArray(ChildIndex);
-	if (!CreationDependencyInfo)
+	if (CreationDependencyInfo)
 	{
-		// No creation dependencies exist
-		return;
+		CreationDependencyInfo->Remove(ParentIndex);
+
+		if (CreationDependencyInfo->IsEmpty())
+		{
+			ObjectsWithCreationDependencies.ClearBit(ChildIndex);
+			// We don't free the array because we usually add a new dependency right after or the object will stop replicating soon and the array will be freed there.
+		}
 	}
+}
 
-	CreationDependencyInfo->Remove(ParentIndex);
+TConstArrayView<const FInternalNetRefIndex> FNetRefHandleManager::GetCreationDependents(FInternalNetRefIndex ParentInternalIndex) const
+{
+	if (ObjectsWithCreationDependents.IsBitSet(ParentInternalIndex) == false)
+	{
+		return MakeArrayView<const FInternalNetRefIndex>(nullptr, 0);
+	}
 
-	if (CreationDependencyInfo->IsEmpty())
+	const FNetDependencyData::FCreationDependentInfoArray* CreationDependentArray = SubObjects.GetCreationDependentInfoArray(ParentInternalIndex);
+	
+	if (!CreationDependentArray)
 	{
-		ObjectsWithCreationDependencies.ClearBit(ChildIndex);
-		
-		//We don't free the array because we usually add a new dependency right after or the object will stop replicating soon and the array will be freed there.
+		return MakeArrayView<const FInternalNetRefIndex>(nullptr, 0);
 	}
+	
+	return MakeArrayView(*CreationDependentArray);
 }
 
 TConstArrayView<const FInternalNetRefIndex> FNetRefHandleManager::GetCreationDependencies(FInternalNetRefIndex ChildInternalIndex) const
@@ -1246,11 +1288,27 @@ void FNetRefHandleManager::InternalRemoveAllDependencies(FInternalNetRefIndex De
 		DependentObjectArray->Reset();		
 	}
 
+	if (FNetDependencyData::FCreationDependentInfoArray* CreationDependentsArray = SubObjects.GetCreationDependentInfoArray(DependentInternalIndex))
+	{
+		// Remove the creation dependencies between this object and its dependents.
+		for (int32 DependentsArrayIndex = CreationDependentsArray->Num() - 1; DependentsArrayIndex >= 0; --DependentsArrayIndex)
+		{
+			FInternalNetRefIndex ChildIndex = CreationDependentsArray->GetData()[DependentsArrayIndex];
+			InternalRemoveCreationDependency(DependentInternalIndex, ChildIndex);
+		}
+
+		CreationDependentsArray->Reset();
+	}
+
 	if (FNetDependencyData::FCreationDependencyInfoArray* CreationDependenciesArray = SubObjects.GetCreationDependencyInfoArray(DependentInternalIndex))
 	{
-		// TODO: This code currently assumes that the parent cannot be destroyed while the children depending on him are still replicating.
-		
-		// Need to add a parent->child link and decide if it's legal for us to remove the dependency automatically or not.
+		// Remove the creation dependencies between this object and its dependencies.
+		for (int32 DependenciesArrayIndex = CreationDependenciesArray->Num() - 1; DependenciesArrayIndex >= 0; --DependenciesArrayIndex)
+		{
+			FInternalNetRefIndex ParentIndex = CreationDependenciesArray->GetData()[DependenciesArrayIndex];
+			InternalRemoveCreationDependency(ParentIndex, DependentInternalIndex);
+		}
+
 		CreationDependenciesArray->Reset();
 	}
 
@@ -1261,6 +1319,7 @@ void FNetRefHandleManager::InternalRemoveAllDependencies(FInternalNetRefIndex De
 	DependentObjectData.bHasDependentObjects = false;
 	ObjectsWithDependentObjectsInternalIndices.ClearBit(DependentInternalIndex);
 	DependentObjectInternalIndices.ClearBit(DependentInternalIndex);
+	ObjectsWithCreationDependents.ClearBit(DependentInternalIndex);
 	ObjectsWithCreationDependencies.ClearBit(DependentInternalIndex);
 }
 
diff --git a/Engine/Source/Runtime/Net/Iris/Private/Iris/ReplicationSystem/NetRefHandleManager.h b/Engine/Source/Runtime/Net/Iris/Private/Iris/ReplicationSystem/NetRefHandleManager.h
--- a/Engine/Source/Runtime/Net/Iris/Private/Iris/ReplicationSystem/NetRefHandleManager.h
+++ b/Engine/Source/Runtime/Net/Iris/Private/Iris/ReplicationSystem/NetRefHandleManager.h
@@ -342,6 +342,18 @@ public:
 	bool AddCreationDependency(FNetRefHandle Parent, FNetRefHandle Child);
 	void RemoveCreationDependency(FNetRefHandle Parent, FNetRefHandle Child);
 
+private:
+	bool InternalAddCreationDependency(FInternalNetRefIndex ParentIndex, FInternalNetRefIndex ChildIndex);
+	void InternalRemoveCreationDependency(FInternalNetRefIndex ParentIndex, FInternalNetRefIndex ChildIndex);
+
+public:
+	const FNetBitArrayView GetObjectsWithCreationDependents() const
+	{
+		return MakeNetBitArrayView(ObjectsWithCreationDependents);
+	}
+
+	TConstArrayView<const FInternalNetRefIndex> GetCreationDependents(FInternalNetRefIndex ParentInternalIndex) const;
+
 	const FNetBitArrayView GetObjectsWithCreationDependencies() const
 	{
 		return MakeNetBitArrayView(ObjectsWithCreationDependencies);
@@ -566,7 +578,8 @@ private:
 	// Bitset containing all internal indices that are dependentObjects
 	FNetBitArray DependentObjectInternalIndices;
 
-	// Dependent objects that need their parent to be relevant too
+	// Objects who need to exist before their children can be created.
+	FNetBitArray ObjectsWithCreationDependents;
 
 	// Objects that cannot be created until the parents they depend on also exist
 	FNetBitArray ObjectsWithCreationDependencies;
diff --git a/Engine/Source/Runtime/Net/Iris/Private/Iris/ReplicationSystem/ObjectReplicationBridge.cpp b/Engine/Source/Runtime/Net/Iris/Private/Iris/ReplicationSystem/ObjectReplicationBridge.cpp
--- a/Engine/Source/Runtime/Net/Iris/Private/Iris/ReplicationSystem/ObjectReplicationBridge.cpp
+++ b/Engine/Source/Runtime/Net/Iris/Private/Iris/ReplicationSystem/ObjectReplicationBridge.cpp
@@ -1241,6 +1241,11 @@ void UObjectReplicationBridge::AddCreationDependencyLink(FNetRefHandle Parent, F
     
     if (bSuccess)
 	{
+		const FInternalNetRefIndex ParentIndex = LocalNetRefHandleManager.GetInternalIndex(Parent);
+		const FInternalNetRefIndex ChildIndex = LocalNetRefHandleManager.GetInternalIndex(Child);
+		FReplicationFiltering& Filtering = ReplicationSystemInternal->GetFiltering();
+		Filtering.NotifyAddedCreationDependency(ParentIndex, ChildIndex);
+
 		UE_LOG_OBJECTREPLICATIONBRIDGE(Verbose, TEXT("AddCreationDependencyLink Added creation dependency between child: %s and parent: %s"), *PrintObjectFromNetRefHandle(Child), *PrintObjectFromNetRefHandle(Parent));
 	}
 	else
@@ -1251,7 +1256,21 @@ void UObjectReplicationBridge::AddCreationDependencyLink(FNetRefHandle Parent, F
 
 void UObjectReplicationBridge::RemoveCreationDependencyLink(FNetRefHandle Parent, FNetRefHandle Child)
 {
-	GetReplicationSystem()->GetReplicationSystemInternal()->GetNetRefHandleManager().RemoveCreationDependency(Parent, Child);
+	using namespace UE::Net;
+	using namespace UE::Net::Private;
+
+	FReplicationSystemInternal* ReplicationSystemInternal = GetReplicationSystem()->GetReplicationSystemInternal();
+	FNetRefHandleManager& LocalNetRefHandleManager = ReplicationSystemInternal->GetNetRefHandleManager();
+
+	LocalNetRefHandleManager.RemoveCreationDependency(Parent, Child);
+
+	const FInternalNetRefIndex ParentIndex = LocalNetRefHandleManager.GetInternalIndex(Parent);
+	const FInternalNetRefIndex ChildIndex = LocalNetRefHandleManager.GetInternalIndex(Child);
+	if (ParentIndex != FNetRefHandleManager::InvalidInternalIndex && ChildIndex != FNetRefHandleManager::InvalidInternalIndex)
+	{
+		FReplicationFiltering& Filtering = ReplicationSystemInternal->GetFiltering();
+		Filtering.NotifyRemovedCreationDependency(ParentIndex, ChildIndex);
+	}
 }
 
 bool UObjectReplicationBridge::WriteNetRefHandleCreationInfo(FReplicationBridgeSerializationContext& Context, FNetRefHandle Handle, UE::Net::Private::FInternalNetRefIndex InternalObjectIndex)
