Iris
- Added optional parameter to control if a subobjects child subobjects replicate before or after their parent
- ActorComponents still defaults to replicate their SubObjects before the components to not cause any unexpected divergences
- Clarified documentation
- Deprecated ESubObjectInsertionOrder::ReplicateWith in favour of explicitly specifying root and parent when adding a subobject
- Fixed bug where using ESubObjectInsertionOrder::InsertAtStart would shift conditionals around and apply them to the wrong subobject.
- Added unit tests

#Synced-CL 50477682

diff --git a/Engine/Plugins/Runtime/ReplicationSystemTestPlugin/Source/Private/Tests/ReplicationSystem/Conditionals/TestConditionals.cpp b/Engine/Plugins/Runtime/ReplicationSystemTestPlugin/Source/Private/Tests/ReplicationSystem/Conditionals/TestConditionals.cpp
--- a/Engine/Plugins/Runtime/ReplicationSystemTestPlugin/Source/Private/Tests/ReplicationSystem/Conditionals/TestConditionals.cpp
+++ b/Engine/Plugins/Runtime/ReplicationSystemTestPlugin/Source/Private/Tests/ReplicationSystem/Conditionals/TestConditionals.cpp
@@ -891,7 +891,11 @@ UE_NET_TEST_FIXTURE(FTestConditionalsFixture, HierarchicalSubObjectIsNotReplicat
 	// Spawn object + subobjects on server
 	UTestReplicatedIrisObject* ServerObject = Server->CreateObject(0, 0);
 	UTestReplicatedIrisObject* ServerSubObject = Server->CreateSubObject(ServerObject->NetRefHandle, 0, 0);
-	UTestReplicatedIrisObject* ServerSubSubObject = Server->CreateSubObject<UTestReplicatedIrisObject>(ServerObject->NetRefHandle, ServerSubObject->NetRefHandle, ESubObjectInsertionOrder::ReplicateWith);
+	UTestReplicatedIrisObject* ServerSubSubObject = Server->CreateSubObjectWithParams<UTestReplicatedIrisObject>(
+		{ 
+			.RootObjectHandle = ServerObject->NetRefHandle, 
+			.ParentSubObjectHandle = ServerSubObject->NetRefHandle,
+		});
 
 	// Set condition to only replicate if root is the owning connection
 	Server->ReplicationBridge->SetSubObjectNetCondition(ServerSubObject->NetRefHandle, ELifetimeCondition::COND_SkipOwner);
@@ -958,6 +962,44 @@ UE_NET_TEST_FIXTURE(FTestConditionalsFixture, HierarchicalSubObjectIsNotReplicat
 
 }
 
+UE_NET_TEST_FIXTURE(FTestConditionalsFixture, SubObjectWithCustomOrderGetsCorrectConditional)
+{
+	// Add client
+	FReplicationSystemTestClient* Client = CreateClient();
+
+	// Spawn object + subobjects on server
+	UTestReplicatedIrisObject* ServerObject = Server->CreateObject(0, 0);
+	UTestReplicatedIrisObject* ServerSubObject0 = Server->CreateSubObject(ServerObject->NetRefHandle, 0, 0);
+	// Set condition on first subobject to not replicate
+	Server->ReplicationBridge->SetSubObjectNetCondition(ServerSubObject0->NetRefHandle, ELifetimeCondition::COND_Never);
+
+	// Add another subobject
+	UTestReplicatedIrisObject* ServerSubObject1 = Server->CreateSubObjectWithParams<UTestReplicatedIrisObject>(
+		{ 
+			.RootObjectHandle = ServerObject->NetRefHandle, 
+			.InsertionOrder = ESubObjectInsertionOrder::InsertAtStart
+		});
+
+	// Set some values
+	ServerSubObject0->IntA = 1; 
+	ServerSubObject1->IntA = 2; 
+
+	// Send and deliver packet
+	Server->UpdateAndSend({Client});
+	
+	// Object should have been created on the client
+	const UTestReplicatedIrisObject* ClientObject = Cast<UTestReplicatedIrisObject>(Client->GetReplicationBridge()->GetReplicatedObject(ServerObject->NetRefHandle));
+	UE_NET_ASSERT_NE(ClientObject, nullptr);
+
+	// SubObject0 should not have been created, 
+	const UTestReplicatedIrisObject* ClientSubObject0 = Cast<UTestReplicatedIrisObject>(Client->GetReplicationBridge()->GetReplicatedObject(ServerSubObject0->NetRefHandle));
+	UE_NET_ASSERT_EQ(ClientSubObject0, nullptr);
+
+	// SubObject1 should have been created, 
+	const UTestReplicatedIrisObject* ClientSubObject1 = Cast<UTestReplicatedIrisObject>(Client->GetReplicationBridge()->GetReplicatedObject(ServerSubObject1->NetRefHandle));
+	UE_NET_ASSERT_NE(ClientSubObject1, nullptr);	
+}
+
 UE_NET_TEST_FIXTURE(FTestConditionalsFixture, NoneIsReplicated)
 {
 	// Add client
diff --git a/Engine/Plugins/Runtime/ReplicationSystemTestPlugin/Source/Private/Tests/ReplicationSystem/ReplicatedTestObject.cpp b/Engine/Plugins/Runtime/ReplicationSystemTestPlugin/Source/Private/Tests/ReplicationSystem/ReplicatedTestObject.cpp
--- a/Engine/Plugins/Runtime/ReplicationSystemTestPlugin/Source/Private/Tests/ReplicationSystem/ReplicatedTestObject.cpp
+++ b/Engine/Plugins/Runtime/ReplicationSystemTestPlugin/Source/Private/Tests/ReplicationSystem/ReplicatedTestObject.cpp
@@ -82,12 +82,12 @@ UE::Net::FNetRefHandle UReplicatedTestObjectBridge::BeginReplication(UReplicated
 	return Handle;
 }
 
-UE::Net::FNetRefHandle UReplicatedTestObjectBridge::BeginReplication(FNetRefHandle OwnerHandle, UReplicatedTestObject* SubObjectInstance, FNetRefHandle InsertRelativeToSubObjectHandle, UE::Net::ESubObjectInsertionOrder InsertionOrder)
+UE::Net::FNetRefHandle UReplicatedTestObjectBridge::BeginReplication(FNetRefHandle OwnerHandle, UReplicatedTestObject* SubObjectInstance)
 {
 	check(OwnerHandle.IsValid());
 
 	// Create NetRefHandle for the registered fragments
-	const FSubObjectReplicationParams Params { .RootObjectHandle = OwnerHandle, .InsertRelativeToSubObjectHandle = InsertRelativeToSubObjectHandle, .InsertionOrder = InsertionOrder };
+	const FSubObjectReplicationParams Params { .RootObjectHandle = OwnerHandle };
 	FNetRefHandle Handle = Super::StartReplicatingSubObject(SubObjectInstance, Params, ReplicatedObjectFactoryId);
 
 	if (Handle.IsValid())
@@ -100,6 +100,23 @@ UE::Net::FNetRefHandle UReplicatedTestObjectBridge::BeginReplication(FNetRefHand
 	return Handle;
 }
 
+UE::Net::FNetRefHandle UReplicatedTestObjectBridge::BeginReplication(UReplicatedTestObject* SubObjectInstance, const FSubObjectReplicationParams& Params)
+{
+	check(Params.RootObjectHandle.IsValid());
+
+	// Create NetRefHandle for the registered fragments
+	FNetRefHandle Handle = Super::StartReplicatingSubObject(SubObjectInstance, Params, ReplicatedObjectFactoryId);
+
+	if (Handle.IsValid())
+	{
+		// This is optional but typically we want to cache at least the NetRefHandle in the game instance to avoid doing map lookups to find it
+		SubObjectInstance->NetRefHandle = Handle;
+		SubObjectInstance->bIsSubObject = true;
+	}
+	
+	return Handle;	
+}
+
 void UReplicatedTestObjectBridge::EndReplication(UReplicatedTestObject* Instance, EEndReplicationFlags Flags)
 {
 	StopReplicatingNetObject(Instance, Flags);
diff --git a/Engine/Plugins/Runtime/ReplicationSystemTestPlugin/Source/Private/Tests/ReplicationSystem/ReplicatedTestObject.h b/Engine/Plugins/Runtime/ReplicationSystemTestPlugin/Source/Private/Tests/ReplicationSystem/ReplicatedTestObject.h
--- a/Engine/Plugins/Runtime/ReplicationSystemTestPlugin/Source/Private/Tests/ReplicationSystem/ReplicatedTestObject.h
+++ b/Engine/Plugins/Runtime/ReplicationSystemTestPlugin/Source/Private/Tests/ReplicationSystem/ReplicatedTestObject.h
@@ -628,7 +628,8 @@ public:
 	// In this example we have methods that directly uses UReplicatedTestObject;
 	FNetRefHandle BeginReplication(UReplicatedTestObject* Instance);
 	FNetRefHandle BeginReplication(UReplicatedTestObject* Instance, const UObjectReplicationBridge::FRootObjectReplicationParams& Params);
-	FNetRefHandle BeginReplication(FNetRefHandle OwnerHandle, UReplicatedTestObject* Instance, FNetRefHandle InsertRelativeToSubObjectHandle = FNetRefHandle::GetInvalid(), UE::Net::ESubObjectInsertionOrder InsertionOrder = UE::Net::ESubObjectInsertionOrder::None);
+	FNetRefHandle BeginReplication(FNetRefHandle OwnerHandle, UReplicatedTestObject* Instance);
+	FNetRefHandle BeginReplication(UReplicatedTestObject* SubObjectInstance, const FSubObjectReplicationParams& Params);
 
 	void EndReplication(UReplicatedTestObject* Instance, EEndReplicationFlags Flags=EEndReplicationFlags::Destroy);
 
diff --git a/Engine/Plugins/Runtime/ReplicationSystemTestPlugin/Source/Private/Tests/ReplicationSystem/ReplicationSystemServerClientTestFixture.h b/Engine/Plugins/Runtime/ReplicationSystemTestPlugin/Source/Private/Tests/ReplicationSystem/ReplicationSystemServerClientTestFixture.h
--- a/Engine/Plugins/Runtime/ReplicationSystemTestPlugin/Source/Private/Tests/ReplicationSystem/ReplicationSystemServerClientTestFixture.h
+++ b/Engine/Plugins/Runtime/ReplicationSystemTestPlugin/Source/Private/Tests/ReplicationSystem/ReplicationSystemServerClientTestFixture.h
@@ -176,17 +176,33 @@ public:
 	}
 
 	template<typename T>
-	T* CreateSubObject(FNetRefHandle OwnerHandle, FNetRefHandle InsertRelativeToSubObjectHandle = FNetRefHandle::GetInvalid(), UE::Net::ESubObjectInsertionOrder InsertionOrder = UE::Net::ESubObjectInsertionOrder::None)
+	T* CreateSubObject(FNetRefHandle OwnerHandle)
 	{
-		T* CreatedObject = NewObject<T>();
-		if (Cast<UReplicatedTestObject>(CreatedObject))
+		T* CreatedSubObject = NewObject<T>();
+		if (Cast<UReplicatedTestObject>(CreatedSubObject))
 		{
-			CreatedObjects.Add(TStrongObjectPtr<UObject>(CreatedObject));
+			CreatedObjects.Add(TStrongObjectPtr<UObject>(CreatedSubObject));
 
 			// Add it to the bridge for replication
-			ReplicationBridge->BeginReplication(OwnerHandle, CreatedObject, InsertRelativeToSubObjectHandle, InsertionOrder);
+			ReplicationBridge->BeginReplication(OwnerHandle, CreatedSubObject);
 	
-			return CreatedObject;
+			return CreatedSubObject;
+		}
+		return nullptr;
+	}
+
+	template<typename T>
+	T* CreateSubObjectWithParams(const UObjectReplicationBridge::FSubObjectReplicationParams& SubObjectParams)
+	{
+		T* CreatedSubObject = NewObject<T>();
+		if (Cast<UReplicatedTestObject>(CreatedSubObject))
+		{
+			CreatedObjects.Add(TStrongObjectPtr<UObject>(CreatedSubObject));
+
+			// Add it to the bridge for replication
+			ReplicationBridge->BeginReplication(CreatedSubObject, SubObjectParams);
+	
+			return CreatedSubObject;
 		}
 		return nullptr;
 	}
diff --git a/Engine/Plugins/Runtime/ReplicationSystemTestPlugin/Source/Private/Tests/ReplicationSystem/TestReplicationSystem.cpp b/Engine/Plugins/Runtime/ReplicationSystemTestPlugin/Source/Private/Tests/ReplicationSystem/TestReplicationSystem.cpp
--- a/Engine/Plugins/Runtime/ReplicationSystemTestPlugin/Source/Private/Tests/ReplicationSystem/TestReplicationSystem.cpp
+++ b/Engine/Plugins/Runtime/ReplicationSystemTestPlugin/Source/Private/Tests/ReplicationSystem/TestReplicationSystem.cpp
@@ -556,7 +556,7 @@ UE_NET_TEST_FIXTURE(FReplicationSystemServerClientTestFixture, TestSubObjectDefa
 	UE_NET_ASSERT_GT(ClientSubObject2->LastRepOrderCounter, ClientSubObject1->LastRepOrderCounter);
 }
 
-UE_NET_TEST_FIXTURE(FReplicationSystemServerClientTestFixture, TestSubObjectSpecifiedReplicationOrder)
+UE_NET_TEST_FIXTURE(FReplicationSystemServerClientTestFixture, TestSubObjectReplicationOrderWithChildSubObjectsReplicatingBeforeParent)
 {
 	// Add a client
 	FReplicationSystemTestClient* Client = CreateClient();
@@ -564,19 +564,30 @@ UE_NET_TEST_FIXTURE(FReplicationSystemServerClientTestFixture, TestSubObjectSpec
 	// Spawn object on server
 	UReplicatedTestObject* ServerObject = Server->CreateObject(0, 0);
 
-	// Spawn some subobjects
-	UReplicatedSubObjectOrderObject* ServerSubObject0 = Server->CreateSubObject<UReplicatedSubObjectOrderObject>(ServerObject->NetRefHandle);
-	// Specify Subobject1 to replicate with SubObject0, which means that it will replicate before Subobjet0 is replicated
-	UReplicatedSubObjectOrderObject* ServerSubObject1 = Server->CreateSubObject<UReplicatedSubObjectOrderObject>(ServerObject->NetRefHandle, ServerSubObject0->NetRefHandle, UE::Net::ESubObjectInsertionOrder::ReplicateWith);
-	// Specify SubObect 2 to replicate with no specific order (it will be added to the owner and thus replicate last)
-	UReplicatedSubObjectOrderObject* ServerSubObject2 = Server->CreateSubObject<UReplicatedSubObjectOrderObject>(ServerObject->NetRefHandle, ServerSubObject1->NetRefHandle, UE::Net::ESubObjectInsertionOrder::None);
+	// Spawn subobject set to replicate its child subobjects before the parent
+	UReplicatedSubObjectOrderObject* ServerSubObject0 = Server->CreateSubObjectWithParams<UReplicatedSubObjectOrderObject>(
+		{
+			.RootObjectHandle = ServerObject->NetRefHandle,
+			.ChildSubObjectsReplicationOrder = EChildSubObjectsReplicationOrder::BeforeParent,
+		});
+
+	// Specify Subobject1 to replicate with SubObject0 as its parent, and since SubObject0 is set to replicate children before itself we expect this to be reflected in the order.
+	UReplicatedSubObjectOrderObject* ServerSubObject1 = Server->CreateSubObjectWithParams<UReplicatedSubObjectOrderObject>(
+		{
+			.RootObjectHandle = ServerObject->NetRefHandle,
+			.ParentSubObjectHandle = ServerSubObject0->NetRefHandle
+		});
+
+	// Create SubObject 2 to replicate with no specific order (it will be added Last in the root list and thus replicate last.)
+	UReplicatedSubObjectOrderObject* ServerSubObject2 = Server->CreateSubObjectWithParams<UReplicatedSubObjectOrderObject>(
+		{
+			.RootObjectHandle = ServerObject->NetRefHandle,
+		});
 
 	UReplicatedSubObjectOrderObject::RepOrderCounter = 0U;
 
 	// Send and deliver packet
-	Server->NetUpdate();
-	Server->SendAndDeliverTo(Client, true);
-	Server->PostSendUpdate();
+	Server->UpdateAndSend({Client});
 
 	// Verify that objects have replicated
 	UReplicatedSubObjectOrderObject* ClientSubObject0 = Cast<UReplicatedSubObjectOrderObject>(Client->GetReplicationBridge()->GetReplicatedObject(ServerSubObject0->NetRefHandle));
@@ -593,6 +604,52 @@ UE_NET_TEST_FIXTURE(FReplicationSystemServerClientTestFixture, TestSubObjectSpec
 	UE_NET_ASSERT_EQ(ClientSubObject2->LastRepOrderCounter, 3U);
 }
 
+UE_NET_TEST_FIXTURE(FReplicationSystemServerClientTestFixture, TestSubObjectReplicationOrderWithChildSubObjectsReplicatingAfterParent)
+{
+	// Add a client
+	FReplicationSystemTestClient* Client = CreateClient();
+
+	// Spawn object on server
+	UReplicatedTestObject* ServerObject = Server->CreateObject(0, 0);
+
+	// Spawn subobject set to replicate its child subobjects after the parent
+	UReplicatedSubObjectOrderObject* ServerSubObject0 = Server->CreateSubObjectWithParams<UReplicatedSubObjectOrderObject>(
+		{
+			.RootObjectHandle = ServerObject->NetRefHandle,
+			.ChildSubObjectsReplicationOrder = EChildSubObjectsReplicationOrder::AfterParent,
+		});
+
+	// Specify Subobject1 to replicate with SubObject0 as its parent, and since SubObject0 is set to replicate children after itself we expect this to be reflected in the order.
+	UReplicatedSubObjectOrderObject* ServerSubObject1 = Server->CreateSubObjectWithParams<UReplicatedSubObjectOrderObject>(
+		{
+			.RootObjectHandle = ServerObject->NetRefHandle, 
+			.ParentSubObjectHandle = ServerSubObject0->NetRefHandle, 
+		});
+
+	// Specify SubObect 2 to replicate with no specific order (it will be added to the root list and thus replicate last)
+	UReplicatedSubObjectOrderObject* ServerSubObject2 = Server->CreateSubObject<UReplicatedSubObjectOrderObject>(ServerObject->NetRefHandle);
+
+	UReplicatedSubObjectOrderObject::RepOrderCounter = 0U;
+
+	// Send and deliver packet
+	Server->NetUpdate();
+	Server->SendAndDeliverTo(Client, true);
+	Server->PostSendUpdate();
+
+	// Verify that objects have replicated
+	UReplicatedSubObjectOrderObject* ClientSubObject0 = Cast<UReplicatedSubObjectOrderObject>(Client->GetReplicationBridge()->GetReplicatedObject(ServerSubObject0->NetRefHandle));
+	UReplicatedSubObjectOrderObject* ClientSubObject1 = Cast<UReplicatedSubObjectOrderObject>(Client->GetReplicationBridge()->GetReplicatedObject(ServerSubObject1->NetRefHandle));
+	UReplicatedSubObjectOrderObject* ClientSubObject2 = Cast<UReplicatedSubObjectOrderObject>(Client->GetReplicationBridge()->GetReplicatedObject(ServerSubObject2->NetRefHandle));
+
+	UE_NET_ASSERT_NE(ClientSubObject0, nullptr);
+	UE_NET_ASSERT_NE(ClientSubObject1, nullptr);
+	UE_NET_ASSERT_NE(ClientSubObject2, nullptr);
+
+	// Verify that they have replicated in expected order setup earlier
+	UE_NET_ASSERT_EQ(ClientSubObject0->LastRepOrderCounter, 1U);
+	UE_NET_ASSERT_EQ(ClientSubObject1->LastRepOrderCounter, 2U);
+	UE_NET_ASSERT_EQ(ClientSubObject2->LastRepOrderCounter, 3U);
+}
 
 UE_NET_TEST_FIXTURE(FReplicationSystemServerClientTestFixture, TestSubObjectInsertAtStart)
 {
@@ -603,13 +660,33 @@ UE_NET_TEST_FIXTURE(FReplicationSystemServerClientTestFixture, TestSubObjectInse
 	UReplicatedObjectTestSubObjectCreationOrder* ServerObject = Server->CreateObject<UReplicatedObjectTestSubObjectCreationOrder>();
 
 	// Spawn a subobject
-	UReplicatedSubObjectOrderObject* ServerSubObject0 = Server->CreateSubObject<UReplicatedSubObjectOrderObject>(ServerObject->NetRefHandle);
-	
-	// Spawn a subobject and make it replicate before SubObject0
-	UReplicatedSubObjectOrderObject* ServerSubObject1 = Server->CreateSubObject<UReplicatedSubObjectOrderObject>(ServerObject->NetRefHandle, ServerSubObject0->NetRefHandle, UE::Net::ESubObjectInsertionOrder::ReplicateWith);
+	UReplicatedSubObjectOrderObject* ServerSubObject0 = Server->CreateSubObjectWithParams<UReplicatedSubObjectOrderObject>(
+		{
+			.RootObjectHandle = ServerObject->NetRefHandle,
+			.ChildSubObjectsReplicationOrder = EChildSubObjectsReplicationOrder::AfterParent,
+		});
 	
-	// Spawn a subobject and make it replicate first
-	UReplicatedSubObjectOrderObject* ServerSubObject2 = Server->CreateSubObject<UReplicatedSubObjectOrderObject>(ServerObject->NetRefHandle, FNetRefHandle::GetInvalid(), UE::Net::ESubObjectInsertionOrder::InsertAtStart);
+	// Spawn a subobject and make it replicate as a child to SubObject0,
+	UReplicatedSubObjectOrderObject* ServerSubObject1 = Server->CreateSubObjectWithParams<UReplicatedSubObjectOrderObject>(
+		{
+			.RootObjectHandle = ServerObject->NetRefHandle,
+			.ParentSubObjectHandle = ServerSubObject0->NetRefHandle, 
+		});
+
+	// Spawn a subobject and make it replicate as a child to SubObbject0 but before ServerSubObject1
+	UReplicatedSubObjectOrderObject* ServerSubObject2 = Server->CreateSubObjectWithParams<UReplicatedSubObjectOrderObject>(
+		{
+			.RootObjectHandle = ServerObject->NetRefHandle,
+			.ParentSubObjectHandle = ServerSubObject0->NetRefHandle, 
+			.InsertionOrder = UE::Net::ESubObjectInsertionOrder::InsertAtStart
+		});
+
+	// Spawn a subobject and make it replicate first of all subobjects 
+	UReplicatedSubObjectOrderObject* ServerSubObject3 = Server->CreateSubObjectWithParams<UReplicatedSubObjectOrderObject>(
+		{
+			.RootObjectHandle = ServerObject->NetRefHandle, 
+			.InsertionOrder = UE::Net::ESubObjectInsertionOrder::InsertAtStart
+		});
 
 	UReplicatedSubObjectOrderObject::RepOrderCounter = 0U;
 
@@ -620,22 +697,26 @@ UE_NET_TEST_FIXTURE(FReplicationSystemServerClientTestFixture, TestSubObjectInse
 	UReplicatedSubObjectOrderObject* ClientSubObject0 = Cast<UReplicatedSubObjectOrderObject>(Client->GetReplicationBridge()->GetReplicatedObject(ServerSubObject0->NetRefHandle));
 	UReplicatedSubObjectOrderObject* ClientSubObject1 = Cast<UReplicatedSubObjectOrderObject>(Client->GetReplicationBridge()->GetReplicatedObject(ServerSubObject1->NetRefHandle));
 	UReplicatedSubObjectOrderObject* ClientSubObject2 = Cast<UReplicatedSubObjectOrderObject>(Client->GetReplicationBridge()->GetReplicatedObject(ServerSubObject2->NetRefHandle));
+	UReplicatedSubObjectOrderObject* ClientSubObject3 = Cast<UReplicatedSubObjectOrderObject>(Client->GetReplicationBridge()->GetReplicatedObject(ServerSubObject3->NetRefHandle));
 
 	UE_NET_ASSERT_NE(ClientSubObject0, nullptr);
 	UE_NET_ASSERT_NE(ClientSubObject1, nullptr);
 	UE_NET_ASSERT_NE(ClientSubObject2, nullptr);
 
 	// Verify that they were created in the expected order
-	UE_NET_ASSERT_EQ(ClientSubObject2->CreationOrder, 1);
-	UE_NET_ASSERT_EQ(ClientSubObject1->CreationOrder, 2);
-	UE_NET_ASSERT_EQ(ClientSubObject0->CreationOrder, 3);
+	UE_NET_ASSERT_EQ(ClientSubObject3->CreationOrder, 1);
+	UE_NET_ASSERT_EQ(ClientSubObject0->CreationOrder, 2);
+	UE_NET_ASSERT_EQ(ClientSubObject2->CreationOrder, 3);
+	UE_NET_ASSERT_EQ(ClientSubObject1->CreationOrder, 4);
 
 	// Verify that they have replicated in the expected order
-	UE_NET_ASSERT_EQ(ClientSubObject2->LastRepOrderCounter, 1U);
-	UE_NET_ASSERT_EQ(ClientSubObject1->LastRepOrderCounter, 2U);
-	UE_NET_ASSERT_EQ(ClientSubObject0->LastRepOrderCounter, 3U);
+	UE_NET_ASSERT_EQ(ClientSubObject3->LastRepOrderCounter, 1U);
+	UE_NET_ASSERT_EQ(ClientSubObject0->LastRepOrderCounter, 2U);
+	UE_NET_ASSERT_EQ(ClientSubObject2->LastRepOrderCounter, 3U);
+	UE_NET_ASSERT_EQ(ClientSubObject1->LastRepOrderCounter, 4U);
 }
 
+
 class FTestNetTokensFixture : public FReplicationSystemServerClientTestFixture
 {
 public:
diff --git a/Engine/Source/Runtime/Engine/Private/Net/Iris/ReplicationSystem/EngineReplicationBridge.cpp b/Engine/Source/Runtime/Engine/Private/Net/Iris/ReplicationSystem/EngineReplicationBridge.cpp
--- a/Engine/Source/Runtime/Engine/Private/Net/Iris/ReplicationSystem/EngineReplicationBridge.cpp
+++ b/Engine/Source/Runtime/Engine/Private/Net/Iris/ReplicationSystem/EngineReplicationBridge.cpp
@@ -88,6 +88,14 @@ static FAutoConsoleVariableRef CVarEnableDynamicNetUpdateFrequency(
 	TEXT("When true changes to AActor::NetUpdateFrequency will be updated in Iris after being registered for replication.")
 );
 
+// Due to how it is setup in legacy replication system default is to replicate child subobjects before the actorcomponents
+static bool bDefaultReplicateActorComponentSubobjectsAfterComponent = false;
+static FAutoConsoleVariableRef CVarDefaultReplicateActorComponentSubobjectsAfterComponent(
+		TEXT("net.Iris.DefaultReplicateActorComponentSubobjectsAfterComponent"),
+		bDefaultReplicateActorComponentSubobjectsAfterComponent,
+		TEXT("If set to true, ChildSubObjects of ActorComponents will replicate after component. Default is false.")
+);
+
 bool IsActorValidForIrisReplication(const AActor* Actor)
 {
 	return IsValid(Actor) && !Actor->IsActorBeingDestroyed() && !Actor->IsUnreachable();
@@ -561,7 +569,12 @@ UE::Net::FNetRefHandle UEngineReplicationBridge::StartReplicatingComponent(FNetR
 		}
 
 		// Start replicating the subobject with its owner.
-		const FSubObjectReplicationParams Params{ .RootObjectHandle = RootObjectHandle };
+		// TODO: Add control over this from ActorComponent and support it from legacy replication system as well. JIRA UE-237066
+		const FSubObjectReplicationParams Params
+		{ 
+			.RootObjectHandle = RootObjectHandle,
+			.ChildSubObjectsReplicationOrder = UE::Net::Private::bDefaultReplicateActorComponentSubobjectsAfterComponent ? EChildSubObjectsReplicationOrder::AfterParent : EChildSubObjectsReplicationOrder::BeforeParent
+		};
 		ReplicatedComponentHandle = Super::StartReplicatingSubObject(SubObject, Params, SubObjectFactoryId);
 	}
 
@@ -578,8 +591,7 @@ UE::Net::FNetRefHandle UEngineReplicationBridge::StartReplicatingComponent(FNetR
 		const FSubObjectReplicationParams Params
 		{
 			.RootObjectHandle = RootObjectHandle,
-			.InsertRelativeToSubObjectHandle = ReplicatedComponentHandle,
-			.InsertionOrder = ESubObjectInsertionOrder::ReplicateWith
+			.ParentSubObjectHandle = ReplicatedComponentHandle,
 		};
 	
 		for (const FSubObjectRegistry::FEntry& SubObjectInfo : RepComponentInfo->SubObjects.GetRegistryList())
diff --git a/Engine/Source/Runtime/Engine/Private/Net/Iris/ReplicationSystem/ReplicationSystemUtil.cpp b/Engine/Source/Runtime/Engine/Private/Net/Iris/ReplicationSystem/ReplicationSystemUtil.cpp
--- a/Engine/Source/Runtime/Engine/Private/Net/Iris/ReplicationSystem/ReplicationSystemUtil.cpp
+++ b/Engine/Source/Runtime/Engine/Private/Net/Iris/ReplicationSystem/ReplicationSystemUtil.cpp
@@ -454,8 +454,7 @@ void FReplicationSystemUtil::StartReplicatingActorComponentSubObject(UActorCompo
 							const UObjectReplicationBridge::FSubObjectReplicationParams Params
 							{ 
 								.RootObjectHandle = ActorRefHandle,
-								.InsertRelativeToSubObjectHandle = ActorComponentRefHandle,
-								.InsertionOrder = ESubObjectInsertionOrder::ReplicateWith
+								.ParentSubObjectHandle = ActorComponentRefHandle
 							};
 							const FNetRefHandle SubObjectRefHandle = Bridge->StartReplicatingSubObject(SubObject, Params);
 							if (SubObjectRefHandle.IsValid() && NetCondition != ELifetimeCondition::COND_None)
diff --git a/Engine/Source/Runtime/Net/Iris/Private/Iris/ReplicationSystem/Conditionals/ReplicationConditionals.cpp b/Engine/Source/Runtime/Net/Iris/Private/Iris/ReplicationSystem/Conditionals/ReplicationConditionals.cpp
--- a/Engine/Source/Runtime/Net/Iris/Private/Iris/ReplicationSystem/Conditionals/ReplicationConditionals.cpp
+++ b/Engine/Source/Runtime/Net/Iris/Private/Iris/ReplicationSystem/Conditionals/ReplicationConditionals.cpp
@@ -575,74 +575,118 @@ void FReplicationConditionals::Update()
 	}
 }
 
-void FReplicationConditionals::GetChildSubObjectsToReplicate(uint32 ReplicatingConnectionId, const FConditionalsMask& LifetimeConditionals,  const FInternalNetRefIndex ParentObjectIndex, FSubObjectsToReplicateArray& OutSubObjectsToReplicate) const
+void FReplicationConditionals::GetChildSubObjectsToReplicate(uint32 ReplicatingConnectionId, const FConditionalsMask& LifetimeConditionals,  const FInternalNetRefIndex ParentObjectIndex, const FChildSubObjectsInfo& SubObjectsInfo, FSubObjectsToReplicateArray& OutSubObjectsToReplicate) const
 {
-	// To mimic old system we use a weird replication order based on hierarchy, SubSubObjects are replicated before the parent
-	FChildSubObjectsInfo SubObjectsInfo;
-	if (NetRefHandleManager->GetChildSubObjects(ParentObjectIndex, SubObjectsInfo))
+	if (SubObjectsInfo.SubObjectLifeTimeConditions == nullptr)
 	{
-		if (SubObjectsInfo.SubObjectLifeTimeConditions == nullptr)
+		for (uint32 ArrayIndex = 0; ArrayIndex < SubObjectsInfo.NumSubObjects; ++ArrayIndex)
 		{
-			for (uint32 ArrayIndex = 0; ArrayIndex < SubObjectsInfo.NumSubObjects; ++ArrayIndex)
+			FInternalNetRefIndex SubObjectIndex = SubObjectsInfo.ChildSubObjects[ArrayIndex];
+
+			FChildSubObjectsInfo ChildSubObjectsInfo;
+			if (NetRefHandleManager->GetChildSubObjects(SubObjectIndex, ChildSubObjectsInfo))
+			{
+				// If SubObject has ChildSubObjects we need to consider replication order.
+				if (ChildSubObjectsInfo.bReplicateChildSubObjectsAfterParent)
+				{
+					OutSubObjectsToReplicate.Add(SubObjectIndex);
+					GetChildSubObjectsToReplicate(ReplicatingConnectionId, LifetimeConditionals, SubObjectIndex, ChildSubObjectsInfo, OutSubObjectsToReplicate);
+				}
+				else
+				{
+					GetChildSubObjectsToReplicate(ReplicatingConnectionId, LifetimeConditionals, SubObjectIndex, ChildSubObjectsInfo, OutSubObjectsToReplicate);
+					OutSubObjectsToReplicate.Add(SubObjectIndex);
+				}
+			}
+			else
 			{
-				FInternalNetRefIndex SubObjectIndex = SubObjectsInfo.ChildSubObjects[ArrayIndex];
-				GetChildSubObjectsToReplicate(ReplicatingConnectionId, LifetimeConditionals, SubObjectIndex, OutSubObjectsToReplicate);
 				OutSubObjectsToReplicate.Add(SubObjectIndex);
 			}
 		}
-		else
+	}
+	else
+	{
+		// Append child subobjects that fulfill the condition
+		for (uint32 ArrayIndex = 0; ArrayIndex < SubObjectsInfo.NumSubObjects; ++ArrayIndex)
 		{
-			// Append child subobjects that fulfill the condition
-			for (uint32 ArrayIndex = 0; ArrayIndex < SubObjectsInfo.NumSubObjects; ++ArrayIndex)
+			const FInternalNetRefIndex SubObjectIndex = SubObjectsInfo.ChildSubObjects[ArrayIndex];
+			const ELifetimeCondition LifeTimeCondition = (ELifetimeCondition)SubObjectsInfo.SubObjectLifeTimeConditions[ArrayIndex];
+			if (LifeTimeCondition == COND_NetGroup)
 			{
-				const FInternalNetRefIndex SubObjectIndex = SubObjectsInfo.ChildSubObjects[ArrayIndex];
-				const ELifetimeCondition LifeTimeCondition = (ELifetimeCondition)SubObjectsInfo.SubObjectLifeTimeConditions[ArrayIndex];
-				if (LifeTimeCondition == COND_NetGroup)
+				bool bShouldReplicateSubObject = false;
+				const TArrayView<const FNetObjectGroupHandle::FGroupIndexType> GroupIndexes = NetObjectGroups->GetGroupIndexesOfNetObject(SubObjectIndex);
+				for (const FNetObjectGroupHandle::FGroupIndexType GroupIndex : GroupIndexes)
 				{
-					bool bShouldReplicateSubObject = false;
-					//TArray<FNetObjectGroupHandle> GroupsMemberOf;
-					//NetObjectGroups->GetGroupHandlesOfNetObject(SubObjectIndex, GroupsMemberOf);
-					// 
-					 const TArrayView<const FNetObjectGroupHandle::FGroupIndexType> GroupIndexes = NetObjectGroups->GetGroupIndexesOfNetObject(SubObjectIndex);
-					for (const FNetObjectGroupHandle::FGroupIndexType GroupIndex : GroupIndexes)
-					{
-						const FNetObjectGroupHandle NetGroup = NetObjectGroups->GetHandleFromIndex(GroupIndex);
+					const FNetObjectGroupHandle NetGroup = NetObjectGroups->GetHandleFromIndex(GroupIndex);
 
-						if (NetGroup.IsNetGroupOwnerNetObjectGroup())
-						{
-							bShouldReplicateSubObject = LifetimeConditionals.IsConditionEnabled(COND_OwnerOnly);
-						}
-						else if (NetGroup.IsNetGroupReplayNetObjectGroup())
+					if (NetGroup.IsNetGroupOwnerNetObjectGroup())
+					{
+						bShouldReplicateSubObject = LifetimeConditionals.IsConditionEnabled(COND_OwnerOnly);
+					}
+					else if (NetGroup.IsNetGroupReplayNetObjectGroup())
+					{
+						bShouldReplicateSubObject = LifetimeConditionals.IsConditionEnabled(COND_ReplayOnly);
+					}
+					else
+					{
+						ENetFilterStatus ReplicationStatus = ENetFilterStatus::Disallow;
+						ensureMsgf(ReplicationFiltering->GetSubObjectFilterStatus(NetGroup, ReplicatingConnectionId, ReplicationStatus), TEXT("FReplicationConditionals::GetChildSubObjectsToReplicat Trying to filter with group %u that is not a SubObjectFilterGroup"), NetGroup.GetGroupIndex());
+						bShouldReplicateSubObject = ReplicationStatus != ENetFilterStatus::Disallow;
+					}
+						
+					if (bShouldReplicateSubObject)
+					{
+						FChildSubObjectsInfo ChildSubObjectsInfo;
+						if (NetRefHandleManager->GetChildSubObjects(SubObjectIndex, ChildSubObjectsInfo))
 						{
-							bShouldReplicateSubObject = LifetimeConditionals.IsConditionEnabled(COND_ReplayOnly);
+							// If SubObject has ChildSubObjects we need to consider replication order.
+							if (ChildSubObjectsInfo.bReplicateChildSubObjectsAfterParent)
+							{
+								OutSubObjectsToReplicate.Add(SubObjectIndex);
+								GetChildSubObjectsToReplicate(ReplicatingConnectionId, LifetimeConditionals, SubObjectIndex, ChildSubObjectsInfo, OutSubObjectsToReplicate);
+							}
+							else
+							{
+								GetChildSubObjectsToReplicate(ReplicatingConnectionId, LifetimeConditionals, SubObjectIndex, ChildSubObjectsInfo, OutSubObjectsToReplicate);
+								OutSubObjectsToReplicate.Add(SubObjectIndex);
+							}
 						}
 						else
 						{
-							ENetFilterStatus ReplicationStatus = ENetFilterStatus::Disallow;
-							ensureMsgf(ReplicationFiltering->GetSubObjectFilterStatus(NetGroup, ReplicatingConnectionId, ReplicationStatus), TEXT("FReplicationConditionals::GetChildSubObjectsToReplicat Trying to filter with group %u that is not a SubObjectFilterGroup"), NetGroup.GetGroupIndex());
-							bShouldReplicateSubObject = ReplicationStatus != ENetFilterStatus::Disallow;
-						}
-						
-						if (bShouldReplicateSubObject)
-						{
-							GetChildSubObjectsToReplicate(ReplicatingConnectionId, LifetimeConditionals, SubObjectIndex, OutSubObjectsToReplicate);
 							OutSubObjectsToReplicate.Add(SubObjectIndex);
-							break;
 						}
+						break;
 					}
-
-					UE_CLOG(!bShouldReplicateSubObject, LogIrisConditionals, VeryVerbose, TEXT("%s Filtered out by COND_NetGroup"), *NetRefHandleManager->PrintObjectFromIndex(SubObjectIndex));
 				}
-				else if (LifetimeConditionals.IsConditionEnabled(LifeTimeCondition))
+
+				UE_CLOG(!bShouldReplicateSubObject, LogIrisConditionals, VeryVerbose, TEXT("%s Filtered out by COND_NetGroup"), *NetRefHandleManager->PrintObjectFromIndex(SubObjectIndex));
+			}
+			else if (LifetimeConditionals.IsConditionEnabled(LifeTimeCondition))
+			{
+				FChildSubObjectsInfo ChildSubObjectsInfo;
+				if (NetRefHandleManager->GetChildSubObjects(SubObjectIndex, ChildSubObjectsInfo))
 				{
-					GetChildSubObjectsToReplicate(ReplicatingConnectionId, LifetimeConditionals, SubObjectIndex, OutSubObjectsToReplicate);
-					OutSubObjectsToReplicate.Add(SubObjectIndex);
+					// If SubObject has ChildSubObjects we need to consider replication order.
+					if (ChildSubObjectsInfo.bReplicateChildSubObjectsAfterParent)
+					{
+						OutSubObjectsToReplicate.Add(SubObjectIndex);
+						GetChildSubObjectsToReplicate(ReplicatingConnectionId, LifetimeConditionals, SubObjectIndex, ChildSubObjectsInfo, OutSubObjectsToReplicate);
+					}
+					else
+					{
+						GetChildSubObjectsToReplicate(ReplicatingConnectionId, LifetimeConditionals, SubObjectIndex, ChildSubObjectsInfo, OutSubObjectsToReplicate);
+						OutSubObjectsToReplicate.Add(SubObjectIndex);
+					}
 				}
 				else
 				{
-					UE_LOG(LogIrisConditionals, VeryVerbose, TEXT("%s Filtered out by %s"), *NetRefHandleManager->PrintObjectFromIndex(SubObjectIndex), *UEnum::GetValueAsString(LifeTimeCondition));
+					OutSubObjectsToReplicate.Add(SubObjectIndex);
 				}
 			}
+			else
+			{
+				UE_LOG(LogIrisConditionals, VeryVerbose, TEXT("%s Filtered out by %s"), *NetRefHandleManager->PrintObjectFromIndex(SubObjectIndex), *UEnum::GetValueAsString(LifeTimeCondition));
+			}
 		}
 	}
 }
@@ -651,12 +695,15 @@ void FReplicationConditionals::GetSubObjectsToReplicate(uint32 ReplicationConnec
 {
 	//IRIS_PROFILER_SCOPE_VERBOSE(FReplicationConditionals_GetSubObjectsToReplicate);
 
-	// For now, we do nothing to detect if a conditional has changed on the RootParent, we simply defer this until the next
-	// time the subobjects are marked as dirty. We might want to consider to explicitly mark object and subobjects as dirty when 
-	// the owning connections or conditionals such as bRepPhysics or Role is changed.
 	constexpr bool bInitialState = false;
 	const FConditionalsMask LifetimeConditionals = GetLifetimeConditionals(ReplicationConnectionId, RootObjectIndex, bInitialState);
-	GetChildSubObjectsToReplicate(ReplicationConnectionId, LifetimeConditionals, RootObjectIndex, OutSubObjectsToReplicate);
+
+	FChildSubObjectsInfo SubObjectsInfo;
+
+	if (NetRefHandleManager->GetChildSubObjects(RootObjectIndex, SubObjectsInfo))
+	{
+		GetChildSubObjectsToReplicate(ReplicationConnectionId, LifetimeConditionals, RootObjectIndex, SubObjectsInfo, OutSubObjectsToReplicate);
+	}
 }
 
 bool FReplicationConditionals::ApplyConditionalsToChangeMask(uint32 ReplicatingConnectionId, bool bIsInitialState, FInternalNetRefIndex ParentObjectIndex, FInternalNetRefIndex ObjectIndex, uint32* ChangeMaskData, const uint32* ConditionalChangeMaskData, const FReplicationProtocol* Protocol)
diff --git a/Engine/Source/Runtime/Net/Iris/Private/Iris/ReplicationSystem/Conditionals/ReplicationConditionals.h b/Engine/Source/Runtime/Net/Iris/Private/Iris/ReplicationSystem/Conditionals/ReplicationConditionals.h
--- a/Engine/Source/Runtime/Net/Iris/Private/Iris/ReplicationSystem/Conditionals/ReplicationConditionals.h
+++ b/Engine/Source/Runtime/Net/Iris/Private/Iris/ReplicationSystem/Conditionals/ReplicationConditionals.h
@@ -22,6 +22,7 @@ namespace UE::Net
 		class FNetObjectGroups;
 		class FReplicationFiltering;
 		class FNetObjectGroups;
+		struct FChildSubObjectsInfo;
 	}
 }
 
@@ -141,7 +142,7 @@ private:
 	void ClearPerObjectInfo(FInternalNetRefIndex ObjectIndex);
 	void ClearConnectionInfosForObject(const FNetBitArray& ValidConnections, FInternalNetRefIndex ObjectIndex);
 
-	void GetChildSubObjectsToReplicate(uint32 ReplicatingConnectionId, const FConditionalsMask& LifetimeConditionals, const FInternalNetRefIndex ParentObjectIndex, FSubObjectsToReplicateArray& OutSubObjectsToReplicate) const;
+	void GetChildSubObjectsToReplicate(uint32 ReplicatingConnectionId, const FConditionalsMask& LifetimeConditionals, const FInternalNetRefIndex ParentObjectIndex, const FChildSubObjectsInfo& SubObjectsInfo, FSubObjectsToReplicateArray& OutSubObjectsToReplicate) const;
 
 	ELifetimeCondition GetDynamicCondition(FInternalNetRefIndex ObjectIndex, uint16 RepIndex) const;
 	void SetDynamicCondition(FInternalNetRefIndex ObjectIndex, uint16 RepIndex, ELifetimeCondition Condition);
diff --git a/Engine/Source/Runtime/Net/Iris/Private/Iris/ReplicationSystem/NetDependencyData.h b/Engine/Source/Runtime/Net/Iris/Private/Iris/ReplicationSystem/NetDependencyData.h
--- a/Engine/Source/Runtime/Net/Iris/Private/Iris/ReplicationSystem/NetDependencyData.h
+++ b/Engine/Source/Runtime/Net/Iris/Private/Iris/ReplicationSystem/NetDependencyData.h
@@ -22,6 +22,7 @@ struct FChildSubObjectsInfo
 	const FInternalNetRefIndex* ChildSubObjects = nullptr;
 	const FLifeTimeConditionStorage* SubObjectLifeTimeConditions = nullptr;
 	uint32 NumSubObjects = 0U;
+	bool bReplicateChildSubObjectsAfterParent = true;
 };
 
 struct FDependentObjectInfo
diff --git a/Engine/Source/Runtime/Net/Iris/Private/Iris/ReplicationSystem/NetRefHandleManager.cpp b/Engine/Source/Runtime/Net/Iris/Private/Iris/ReplicationSystem/NetRefHandleManager.cpp
--- a/Engine/Source/Runtime/Net/Iris/Private/Iris/ReplicationSystem/NetRefHandleManager.cpp
+++ b/Engine/Source/Runtime/Net/Iris/Private/Iris/ReplicationSystem/NetRefHandleManager.cpp
@@ -144,7 +144,7 @@ void FNetRefHandleManager::Init(const FInitParams& InitParams)
 		InitNetBitArray(&WantToBeDormantInternalIndices);
 		InitNetBitArray(&ObjectsWithPreUpdate);
 		InitNetBitArray(&DormantObjectsPendingFlushNet);
-		InitNetBitArray(&ObjectsWithFullPushBasedDirtiness);	
+		InitNetBitArray(&ObjectsWithFullPushBasedDirtiness);
 	}
 
 	// Mark the invalid index as used
@@ -752,29 +752,34 @@ void FNetRefHandleManager::DestroyObjectsPendingDestroy()
 	++UpdatePendingDestroyCounter;
 }
 
-bool FNetRefHandleManager::AddSubObject(FNetRefHandle RootObjectHandle, FNetRefHandle SubObjectHandle, FNetRefHandle RelativeOtherSubObjectHandle, EAddSubObjectFlags Flags)
+bool FNetRefHandleManager::AddSubObject(FNetRefHandle ParentObjectHandle, FNetRefHandle SubObjectHandle, FNetRefHandle RelativeOtherSubObjectHandle, EAddSubObjectFlags Flags)
 {
-	check(RootObjectHandle != SubObjectHandle);
+	check(ParentObjectHandle != SubObjectHandle);
 
 	// validate objects
-	const FInternalNetRefIndex RootObjectInternalIndex = GetInternalIndex(RootObjectHandle);
+	const FInternalNetRefIndex ParentObjectInternalIndex = GetInternalIndex(ParentObjectHandle);
 	const FInternalNetRefIndex SubObjectInternalIndex = GetInternalIndex(SubObjectHandle);
 
-	const bool bIsValidOwner = ensure(RootObjectInternalIndex != InvalidInternalIndex);
+	const bool bIsValidParent = ensure(ParentObjectInternalIndex != InvalidInternalIndex);
 	const bool bIsValidSubObject = ensure(SubObjectInternalIndex != InvalidInternalIndex);
 
-	if (!bIsValidOwner || !bIsValidSubObject)
+	if (!bIsValidParent || !bIsValidSubObject)
 	{
 		return false;
 	}
 
+	const FInternalNetRefIndex RootObjectInternalIndex = GetRootObjectInternalIndexOfAnyObject(ParentObjectInternalIndex);
+
+#if UE_NET_TRACE_ENABLED
+	const FNetRefHandle RootObjectHandle = GetReplicatedObjectDataNoCheck(RootObjectInternalIndex).RefHandle;
 	UE_NET_TRACE_SUBOBJECT(RootObjectHandle, SubObjectHandle);
+#endif
 
-	const FInternalNetRefIndex RelativeOtherSubObjectInternalIndex = EnumHasAnyFlags(Flags, EAddSubObjectFlags::ReplicateWithSubObject | EAddSubObjectFlags::InsertAtStartOrAfterOther) ? GetInternalIndex(RelativeOtherSubObjectHandle) : InvalidInternalIndex;
-	return InternalAddSubObject(RootObjectInternalIndex, SubObjectInternalIndex, RelativeOtherSubObjectInternalIndex, Flags);
+	const FInternalNetRefIndex RelativeOtherSubObjectInternalIndex = EnumHasAnyFlags(Flags, EAddSubObjectFlags::InsertAtStartOrAfterOther) ? GetInternalIndex(RelativeOtherSubObjectHandle) : InvalidInternalIndex;
+	return InternalAddSubObject(RootObjectInternalIndex, ParentObjectInternalIndex, SubObjectInternalIndex, RelativeOtherSubObjectInternalIndex, Flags);
 }
 
-bool FNetRefHandleManager::InternalAddSubObject(FInternalNetRefIndex RootObjectInternalIndex, FInternalNetRefIndex SubObjectInternalIndex, FInternalNetRefIndex RelativeOtherSubObjectInternalIndex, EAddSubObjectFlags Flags)
+bool FNetRefHandleManager::InternalAddSubObject(FInternalNetRefIndex RootObjectInternalIndex, FInternalNetRefIndex OptionalParentSubObjectInternalIndex, FInternalNetRefIndex SubObjectInternalIndex, FInternalNetRefIndex RelativeOtherSubObjectInternalIndex, EAddSubObjectFlags Flags)
 {
 	using namespace UE::Net::Private;
 
@@ -784,27 +789,9 @@ bool FNetRefHandleManager::InternalAddSubObject(FInternalNetRefIndex RootObjectI
 		return false;
 	}
 
-	// Add the subobject to the RootObject
+	// Add the subobject to the RootObject, this array contains all subobjects associated with the root in the order they are added
 	FNetDependencyData::FInternalNetRefIndexArray& SubObjectArray = SubObjects.GetOrCreateInternalIndexArray<FNetDependencyData::SubObjects>(RootObjectInternalIndex);
-
-	if (EnumHasAnyFlags(Flags, EAddSubObjectFlags::InsertAtStart))
-	{
-		// Add at the start
-		SubObjectArray.Insert(SubObjectInternalIndex, 0);
-	}
-	else if (EnumHasAnyFlags(Flags, EAddSubObjectFlags::InsertAtStartOrAfterOther))
-	{
-		// Insert the object right after the other one, or at the start if it wasn't in the list
-		const int32 RelativeIndex = RelativeOtherSubObjectInternalIndex != FNetRefHandleManager::InvalidInternalIndex ? SubObjectArray.Find(RelativeOtherSubObjectInternalIndex) : INDEX_NONE;
-		const int32 InsertIndex = (RelativeIndex == INDEX_NONE) ? 0 : (RelativeIndex + 1);
-
-		SubObjectArray.Insert(SubObjectInternalIndex, InsertIndex);
-	}
-	else
-	{
-		// At at the end
-		SubObjectArray.Add(SubObjectInternalIndex);
-	}
+	SubObjectArray.Add(SubObjectInternalIndex);
 	
 	SubObjectData.SubObjectRootIndex = RootObjectInternalIndex;
 
@@ -813,17 +800,17 @@ bool FNetRefHandleManager::InternalAddSubObject(FInternalNetRefIndex RootObjectI
 
 	FInternalNetRefIndex ParentOfSubObjectIndex = RootObjectInternalIndex;
 	
-	// Make sure the non-rootobject Parent is a subobject of the same RootObject
-	if (EnumHasAnyFlags(Flags, EAddSubObjectFlags::ReplicateWithSubObject) && RelativeOtherSubObjectInternalIndex != InvalidInternalIndex)
+	if (OptionalParentSubObjectInternalIndex != InvalidInternalIndex && OptionalParentSubObjectInternalIndex != RootObjectInternalIndex)
 	{
-		const bool bIsValidOuter = SubObjectArray.Contains(RelativeOtherSubObjectInternalIndex);
-		if (ensureMsgf(bIsValidOuter, TEXT("RelativeOtherSubObjectHandle %s must be a Subobject of %s"), *PrintObjectFromIndex(RelativeOtherSubObjectInternalIndex), *PrintObjectFromIndex(RootObjectInternalIndex)))
+		// Make sure the non-rootobject Parent is a subobject of the same RootObject
+		const bool bIsValidOuter = SubObjectArray.Contains(OptionalParentSubObjectInternalIndex);
+		if (ensureMsgf(bIsValidOuter, TEXT("ParentSubObjectInternalIndex %s must be a Subobject of %s"), *PrintObjectFromIndex(OptionalParentSubObjectInternalIndex), *PrintObjectFromIndex(RootObjectInternalIndex)))
 		{
-			ParentOfSubObjectIndex = RelativeOtherSubObjectInternalIndex;
+			ParentOfSubObjectIndex = OptionalParentSubObjectInternalIndex;
 		}
 	}
 
-	// Add the subobject to his Parent's list
+	// Add the subobject to its Parent's list which is what we use to define replication order.
 	FNetDependencyData::FSubObjectConditionalsArray* SubObjectConditionalsArray = nullptr;
 	FNetDependencyData::FInternalNetRefIndexArray& ChildSubObjectArray = SubObjects.GetOrCreateInternalChildSubObjectsArray(ParentOfSubObjectIndex, SubObjectConditionalsArray);
 
@@ -831,6 +818,10 @@ bool FNetRefHandleManager::InternalAddSubObject(FInternalNetRefIndex RootObjectI
 	{
 		// Add at the start
 		ChildSubObjectArray.Insert(SubObjectInternalIndex, 0);
+		if (SubObjectConditionalsArray)
+		{
+			SubObjectConditionalsArray->Insert(ELifetimeCondition::COND_None, 0);
+		}
 	}
 	else if (EnumHasAnyFlags(Flags, EAddSubObjectFlags::InsertAtStartOrAfterOther))
 	{
@@ -839,20 +830,23 @@ bool FNetRefHandleManager::InternalAddSubObject(FInternalNetRefIndex RootObjectI
 		const int32 InsertIndex = (RelativeIndex == INDEX_NONE) ? 0 : (RelativeIndex + 1);
 
 		ChildSubObjectArray.Insert(SubObjectInternalIndex, InsertIndex);
+		if (SubObjectConditionalsArray)
+		{
+			SubObjectConditionalsArray->Insert(ELifetimeCondition::COND_None, InsertIndex);
+		}
 	}
 	else
 	{
 		// At at the end
 		ChildSubObjectArray.Add(SubObjectInternalIndex);
+		if (SubObjectConditionalsArray)
+		{
+			SubObjectConditionalsArray->Add(ELifetimeCondition::COND_None);
+		}
 	}
 	
-
-	if (SubObjectConditionalsArray)
-	{
-		SubObjectConditionalsArray->Add(ELifetimeCondition::COND_None);
-	}
-
 	SubObjectData.SubObjectParentIndex = ParentOfSubObjectIndex;
+	SubObjectData.bReplicateChildSubObjectsAfterParent = EnumHasAnyFlags(Flags, EAddSubObjectFlags::ReplicateChildSubObjectsAfterParent) ? 1U : 0U;
 
 	return true;
 }
@@ -885,7 +879,7 @@ void FNetRefHandleManager::InternalRemoveSubObject(FInternalNetRefIndex RootObje
 			if (SubObjects.GetInternalChildSubObjectAndConditionalArrays(SubObjectData.SubObjectParentIndex, ChildSubObjectArray, SubObjectConditionsArray))
 			{
 				const int32 ArrayIndex = ChildSubObjectArray->Find(SubObjectInternalIndex);
-				if (ensureMsgf(ArrayIndex != INDEX_NONE, TEXT("Subobject: %s not found in the list of his parent"), *PrintObjectFromIndex(SubObjectInternalIndex)))
+				if (ensureMsgf(ArrayIndex != INDEX_NONE, TEXT("Subobject: %s not found in the list of the parent"), *PrintObjectFromIndex(SubObjectInternalIndex)))
 				{
 					ChildSubObjectArray->RemoveAt(ArrayIndex);
 					if (SubObjectConditionsArray)
diff --git a/Engine/Source/Runtime/Net/Iris/Private/Iris/ReplicationSystem/NetRefHandleManager.h b/Engine/Source/Runtime/Net/Iris/Private/Iris/ReplicationSystem/NetRefHandleManager.h
--- a/Engine/Source/Runtime/Net/Iris/Private/Iris/ReplicationSystem/NetRefHandleManager.h
+++ b/Engine/Source/Runtime/Net/Iris/Private/Iris/ReplicationSystem/NetRefHandleManager.h
@@ -50,11 +50,12 @@ typedef uint32 FInternalNetRefIndex;
 enum class EAddSubObjectFlags : uint32
 {
 	None = 0U,
-	ReplicateWithSubObject = 1U,
-	/** Insert the subject at the start of the list */
-	InsertAtStart = ReplicateWithSubObject << 1U, 
-	/** Insert at the start of the list or right after the relative handle if one is set */
+	/** Insert the subject at the start of the parents list */
+	InsertAtStart = 1U, 
+	/** Insert at the start of the parents list or right after the relative handle if one is set */
 	InsertAtStartOrAfterOther = InsertAtStart << 1U,
+	/** If set, children of this SubObject will replicate after the parent */
+	ReplicateChildSubObjectsAfterParent  = InsertAtStartOrAfterOther << 1U,
 };
 ENUM_CLASS_FLAGS(EAddSubObjectFlags);
 
@@ -120,6 +121,8 @@ public:
 		/** Subobjects only: Internal index of the ParentObject of this subobject */
 		FInternalNetRefIndex SubObjectParentIndex = InvalidInternalIndex;
 
+		/** SubObjects only: Determine if child subobjects should replicate before or after parent */
+		uint16 bReplicateChildSubObjectsAfterParent : 1U = 0;
 		uint16 bShouldPropagateChangedStates : 1U = 0;
 		uint16 bTearOff : 1U = 0;
 		uint16 bFlush : 1U = 0;
@@ -306,8 +309,8 @@ public:
 	const FNetBitArray& GetSubObjectInternalIndices() const { return SubObjectInternalIndices; }
 	const FNetBitArrayView GetSubObjectInternalIndicesView() const { return MakeNetBitArrayView(SubObjectInternalIndices); }
 
-	bool AddSubObject(FNetRefHandle RootObjectHandle, FNetRefHandle SubObjectHandle, FNetRefHandle RelativeOtherSubObjectHandle, EAddSubObjectFlags Flags = EAddSubObjectFlags::None);
-	bool AddSubObject(FNetRefHandle RootObjectHandle, FNetRefHandle SubObjectHandle, EAddSubObjectFlags Flags = EAddSubObjectFlags::None);
+	bool AddSubObject(FNetRefHandle ParentObjectHandle, FNetRefHandle SubObjectHandle, FNetRefHandle RelativeOtherSubObjectHandle, EAddSubObjectFlags Flags = EAddSubObjectFlags::None);
+	bool AddSubObject(FNetRefHandle ParentObjectHandle, FNetRefHandle SubObjectHandle, EAddSubObjectFlags Flags = EAddSubObjectFlags::None);
 	void RemoveSubObject(FNetRefHandle SubObjectHandle);
 	
 	FNetRefHandle GetRootObjectOfAnyObject(FNetRefHandle NetRefHandle) const;
@@ -497,7 +500,7 @@ private:
 	// Get the next free internal index, returns InvalidInternalIndex if a free one cannot be found
 	FInternalNetRefIndex GetNextFreeInternalIndex() const;
 
-	bool InternalAddSubObject(FInternalNetRefIndex RootObjectInternalIndex, FInternalNetRefIndex SubObjectInternalIndex, FInternalNetRefIndex RelativeOtherSubObjectInternalIndex, EAddSubObjectFlags Flags);
+	bool InternalAddSubObject(FInternalNetRefIndex RootObjectInternalIndex, FInternalNetRefIndex ParentSubObjectInternalIndex, FInternalNetRefIndex SubObjectInternalIndex, FInternalNetRefIndex RelativeOtherSubObjectInternalIndex, EAddSubObjectFlags Flags);
 
 	inline void SetIsSubObject(FInternalNetRefIndex InternalIndex, bool IsSubObject) { SubObjectInternalIndices.SetBitValue(InternalIndex, IsSubObject); }
 	void InternalRemoveSubObject(FInternalNetRefIndex RootObjectInternalIndex, FInternalNetRefIndex SubObjectInternalIndex, bool bRemoveFromSubObjectArray = true);
@@ -756,7 +759,13 @@ TArrayView<const FInternalNetRefIndex> FNetRefHandleManager::GetSubObjects(FInte
 
 bool FNetRefHandleManager::GetChildSubObjects(FInternalNetRefIndex OwnerIndex, FChildSubObjectsInfo& OutInfo) const
 {
-	return SubObjects.GetChildSubObjects(OwnerIndex, OutInfo);
+	if (!SubObjects.GetChildSubObjects(OwnerIndex, OutInfo) || OutInfo.NumSubObjects == 0U)
+	{
+		return false;
+	}
+
+	OutInfo.bReplicateChildSubObjectsAfterParent = GetReplicatedObjectDataNoCheck(OwnerIndex).bReplicateChildSubObjectsAfterParent;
+	return true;
 }
 
 TArrayView<const FInternalNetRefIndex> FNetRefHandleManager::GetChildSubObjects(FInternalNetRefIndex OwnerIndex) const
diff --git a/Engine/Source/Runtime/Net/Iris/Private/Iris/ReplicationSystem/ObjectReplicationBridge.cpp b/Engine/Source/Runtime/Net/Iris/Private/Iris/ReplicationSystem/ObjectReplicationBridge.cpp
--- a/Engine/Source/Runtime/Net/Iris/Private/Iris/ReplicationSystem/ObjectReplicationBridge.cpp
+++ b/Engine/Source/Runtime/Net/Iris/Private/Iris/ReplicationSystem/ObjectReplicationBridge.cpp
@@ -919,13 +919,13 @@ UE::Net::FNetRefHandle UObjectReplicationBridge::StartReplicatingSubObject(UObje
 
 	const FInternalNetRefIndex RootObjectIndex = LocalNetRefHandleManager.GetInternalIndex(Params.RootObjectHandle);
 
-	if (!ensureMsgf(!Instance->IsTemplate(), TEXT("Iris cannot replicate subobject %s owned by %s because it's an %s"), 
+	if (!ensureMsgf(!Instance->IsTemplate(), TEXT("Iris cannot replicate subobject %s with root %s because it's an %s"), 
 		*GetNameSafe(Instance), *GetNameSafe(LocalNetRefHandleManager.GetReplicatedObjectInstance(RootObjectIndex)), Instance->HasAnyFlags(RF_ArchetypeObject)?TEXT("Archetype"):TEXT("DefaultObject")))
 	{
 		return FNetRefHandle::GetInvalid();
 	}
 
-	if (!ensureMsgf(IsReplicatedHandle(Params.RootObjectHandle), TEXT("Owner %s (%s) must be replicated for subobject %s to replicate."), 
+	if (!ensureMsgf(IsReplicatedHandle(Params.RootObjectHandle), TEXT("Root %s (%s) must be replicated for subobject %s to replicate."), 
 		*GetNameSafe(LocalNetRefHandleManager.GetReplicatedObjectInstance(RootObjectIndex)), *Params.RootObjectHandle.ToString(), *GetNameSafe(Instance)))
 	{
 		return FNetRefHandle::GetInvalid();
@@ -933,11 +933,21 @@ UE::Net::FNetRefHandle UObjectReplicationBridge::StartReplicatingSubObject(UObje
 
 	if (!LocalNetRefHandleManager.IsScopableIndex(RootObjectIndex))
 	{
-		UE_LOG_OBJECTREPLICATIONBRIDGE(Warning, TEXT("Owner %s (%s) must be in scope for subobject %s to replicate."), 
+		UE_LOG_OBJECTREPLICATIONBRIDGE(Warning, TEXT("Root %s (%s) must be in scope for subobject %s to replicate."), 
 			*GetNameSafe(LocalNetRefHandleManager.GetReplicatedObjectInstance(RootObjectIndex)), *Params.RootObjectHandle.ToString(), *GetNameSafe(Instance))
 		return FNetRefHandle::GetInvalid();
 	}
 
+	// Verify that ParentSubObjectHandle is valid if specified
+	if (Params.ParentSubObjectHandle.IsValid() && (Params.ParentSubObjectHandle != Params.RootObjectHandle))
+	{		
+		if (!ensureMsgf(NetRefHandleManager->GetRootObjectOfSubObject(Params.ParentSubObjectHandle) == Params.RootObjectHandle, TEXT("ParentSubObjectHandle %s must be a replicated subobject with same RootObjectHandle %s. Trying to add subobject: %s"), 
+			*LocalNetRefHandleManager.PrintObjectFromNetRefHandle(Params.ParentSubObjectHandle), *LocalNetRefHandleManager.PrintObjectFromNetRefHandle(Params.RootObjectHandle), *GetNameSafe(Instance)))
+		{
+			return FNetRefHandle::GetInvalid();
+		}
+	}
+
 	// Pre-registering subobjects is not supported
 	{
 		const FNetRefHandle PotentiallyPreRegisteredHandle = GetObjectReferenceCache()->GetObjectReferenceHandleFromObject(Instance);
@@ -972,9 +982,18 @@ UE::Net::FNetRefHandle UObjectReplicationBridge::StartReplicatingSubObject(UObje
 	}
 
 	// Add subobject
-	InternalAddSubObject(Params.RootObjectHandle, SubObjectRefHandle, Params.InsertRelativeToSubObjectHandle, Params.InsertionOrder);
-
-	UE_LOG_OBJECTREPLICATIONBRIDGE(Verbose, TEXT("StartReplicatingSubObject Added %s RelativeToSubObjectHandle %s"), *PrintObjectFromNetRefHandle(SubObjectRefHandle), *PrintObjectFromNetRefHandle(Params.InsertRelativeToSubObjectHandle));
+	{
+		if (Params.ParentSubObjectHandle.IsValid() && Params.ParentSubObjectHandle != Params.RootObjectHandle)
+		{
+			InternalAddSubObject(Params.ParentSubObjectHandle, SubObjectRefHandle, Params.InsertRelativeToSubObjectHandle, Params.InsertionOrder, Params.ChildSubObjectsReplicationOrder);
+			UE_LOG_OBJECTREPLICATIONBRIDGE(Verbose, TEXT("StartReplicatingChildSubObject Added %s RelativeToSubObjectHandle %s"), *PrintObjectFromNetRefHandle(SubObjectRefHandle), *PrintObjectFromNetRefHandle(Params.InsertRelativeToSubObjectHandle));
+		}
+		else
+		{
+			InternalAddSubObject(Params.RootObjectHandle, SubObjectRefHandle, Params.InsertRelativeToSubObjectHandle, Params.InsertionOrder, Params.ChildSubObjectsReplicationOrder);
+			UE_LOG_OBJECTREPLICATIONBRIDGE(Verbose, TEXT("StartReplicatingSubObject Added %s RelativeToSubObjectHandle %s"), *PrintObjectFromNetRefHandle(SubObjectRefHandle), *PrintObjectFromNetRefHandle(Params.InsertRelativeToSubObjectHandle));
+		}
+	}
 
 	// SubObjects should always poll with owner
 	SetPollWithObject(Params.RootObjectHandle, SubObjectRefHandle);
diff --git a/Engine/Source/Runtime/Net/Iris/Private/Iris/ReplicationSystem/ReplicationBridge.cpp b/Engine/Source/Runtime/Net/Iris/Private/Iris/ReplicationSystem/ReplicationBridge.cpp
--- a/Engine/Source/Runtime/Net/Iris/Private/Iris/ReplicationSystem/ReplicationBridge.cpp
+++ b/Engine/Source/Runtime/Net/Iris/Private/Iris/ReplicationSystem/ReplicationBridge.cpp
@@ -332,7 +332,7 @@ void UReplicationBridge::DestroyLocalNetHandle(FNetRefHandle Handle, EEndReplica
 	}
 }
 
-void UReplicationBridge::InternalAddSubObject(FNetRefHandle OwnerHandle, FNetRefHandle SubObjectHandle, FNetRefHandle InsertRelativeToSubObjectHandle, UE::Net::ESubObjectInsertionOrder InsertionOrder)
+void UReplicationBridge::InternalAddSubObject(FNetRefHandle ParentHandle, FNetRefHandle SubObjectHandle, FNetRefHandle InsertRelativeToSubObjectHandle, UE::Net::ESubObjectInsertionOrder InsertionOrder, UE::Net::EChildSubObjectsReplicationOrder ChildSubObjectsReplicationOrder)
 {
 	using namespace UE::Net;
 	using namespace UE::Net::Private;
@@ -343,9 +343,21 @@ void UReplicationBridge::InternalAddSubObject(FNetRefHandle OwnerHandle, FNetRef
 	{
 		case ESubObjectInsertionOrder::None: 
 			break;
+
+PRAGMA_DISABLE_DEPRECATION_WARNINGS
 		case ESubObjectInsertionOrder::ReplicateWith:
-			AddSubObjectFlags |= EAddSubObjectFlags::ReplicateWithSubObject;
+		{
+			// This flag is being deprecated.
+			ensureMsgf(false, TEXT("ESubObjectInsertionOrder::ReplicateWith is being deprecated in favor of explicitly specifying the parent using FSubObjectReplicationParams::ParentSubObjectHandle"));
+			if (InsertRelativeToSubObjectHandle.IsValid())
+			{
+				ParentHandle = InsertRelativeToSubObjectHandle;
+				InsertRelativeToSubObjectHandle = FNetRefHandle();
+			}
 			break;
+		}
+PRAGMA_ENABLE_DEPRECATION_WARNINGS
+
 		case ESubObjectInsertionOrder::InsertAtStart:
 			AddSubObjectFlags |= EAddSubObjectFlags::InsertAtStart;
 			break;
@@ -357,7 +369,12 @@ void UReplicationBridge::InternalAddSubObject(FNetRefHandle OwnerHandle, FNetRef
 			break;
 	}
 
-	if (NetRefHandleManager->AddSubObject(OwnerHandle, SubObjectHandle, InsertRelativeToSubObjectHandle, AddSubObjectFlags))
+	if (ChildSubObjectsReplicationOrder == UE::Net::EChildSubObjectsReplicationOrder::AfterParent)
+	{
+		AddSubObjectFlags |= EAddSubObjectFlags::ReplicateChildSubObjectsAfterParent;	
+	}
+
+	if (NetRefHandleManager->AddSubObject(ParentHandle, SubObjectHandle, InsertRelativeToSubObjectHandle, AddSubObjectFlags))
 	{
 		// If the subobject is new we need to update it immediately to pick it up for replication with its new parent
 		ForceNetUpdate(ReplicationSystem->GetId(), NetRefHandleManager->GetInternalIndex(SubObjectHandle));
diff --git a/Engine/Source/Runtime/Net/Iris/Public/Iris/ReplicationSystem/ObjectReplicationBridge.h b/Engine/Source/Runtime/Net/Iris/Public/Iris/ReplicationSystem/ObjectReplicationBridge.h
--- a/Engine/Source/Runtime/Net/Iris/Public/Iris/ReplicationSystem/ObjectReplicationBridge.h
+++ b/Engine/Source/Runtime/Net/Iris/Public/Iris/ReplicationSystem/ObjectReplicationBridge.h
@@ -154,6 +154,7 @@ public:
 	 * 
 	 * @param Instance The instance that needs to be replicated
 	 * @param Params Optional configuration parameters to specify how the root object will be replicated.
+	 * @param NetFactoryId NetFactory to use for this object.
 	 * @return The NetRefHandle associated to this object if the operation succeeded.
 	 */
 	IRISCORE_API FNetRefHandle StartReplicatingRootObject(UObject* Instance, const FRootObjectReplicationParams& Params, UE::Net::FNetObjectFactoryId NetFactoryId);
@@ -161,10 +162,9 @@ public:
 	/**
 	 * Start replicating a SubObject and return a valid NetRefHandle if successful.
 	 * 
-	 * @param OwnerHandle The owner of the subobject
-	 * @param Instance The instance that needs to be replicated
-	 * @param InsertRelativeToSubObjectHandle When valid the new subobject will be inserted in the subobject replication list next to the specified handle and the wanted insertion order.
-	 * @param InsertionOrder When none will always add the new subobject at the end of the list. 
+	 * @param Instance, the object that should be replicated as a SubObject
+	 * @param Params, See FSubObjectReplicationParams for details.
+	 * @param NetFactoryId NetFactory to use for this SubObject.
 	 * @return The NetRefHandle of the subobject if successful.
 	 */
 	IRISCORE_API FNetRefHandle StartReplicatingSubObject(UObject* Instance, const FSubObjectReplicationParams& Params, UE::Net::FNetObjectFactoryId NetFactoryId);
@@ -173,8 +173,7 @@ public:
 	 * Stop replicating any type of NetObject. @see UReplicationBridge::StopReplicatingNetObject for more details.
 	 * @param Instance The instance that won't be replicated anymore
 	 * @param EndReplicationFlags Optional settings to modify the function behavior. Defaults to destroy the  instance on remote clients.
-	 * 
-	 * */
+	 */
 	IRISCORE_API void StopReplicatingNetObject(UObject* Instance, EEndReplicationFlags EndReplicationFlags = EEndReplicationFlags::Destroy);
 
 	/** 
@@ -333,12 +332,23 @@ public:
 
 	struct FSubObjectReplicationParams
 	{
-		/** The root object that the subobject will be attached to*/
+		/** The root object that the subobject will be attached to */
 		FNetRefHandle RootObjectHandle;
 
-		/** Optional handle to another subobject indicating where in the subobject list we inserted */
+		/**
+		 * Optional handle to another parent subobject sharing the same Root, if set, the added subobject will replicate as a child to the specified 
+		 * parent and be subject to any conditionals applied to the parent. If not set, the RootObjectHandle will be used instead.
+		 */
+		FNetRefHandle ParentSubObjectHandle;
+
+		/** Optional handle to another subobject, used in combination with InsertionOrder to provide control over replication order for subobjects sharing the same Parent (Root or specified Parent). */
 		FNetRefHandle InsertRelativeToSubObjectHandle; 
 
+		/** Specify if child subobjects added to this SubObject should replicate before or after the parent. Note: Root always replicate before SubObjects.
+		 * This mostly exists to allow ActorComponents to keep old behavior where they always replicated SubObjects before the component itself.
+		 */
+		UE::Net::EChildSubObjectsReplicationOrder ChildSubObjectsReplicationOrder = UE::Net::EChildSubObjectsReplicationOrder::AfterParent;
+
 		/** Optional enum to select where in the list this subobject will be inserted */
 		UE::Net::ESubObjectInsertionOrder InsertionOrder = UE::Net::ESubObjectInsertionOrder::None;
 
diff --git a/Engine/Source/Runtime/Net/Iris/Public/Iris/ReplicationSystem/ReplicationBridge.h b/Engine/Source/Runtime/Net/Iris/Public/Iris/ReplicationSystem/ReplicationBridge.h
--- a/Engine/Source/Runtime/Net/Iris/Public/Iris/ReplicationSystem/ReplicationBridge.h
+++ b/Engine/Source/Runtime/Net/Iris/Public/Iris/ReplicationSystem/ReplicationBridge.h
@@ -98,14 +98,22 @@ namespace UE::Net
 	enum class ESubObjectInsertionOrder : uint8
 	{
 		None,
-		/** Insert the subobject so it will replicate before the other subobject. */
-		ReplicateWith,
-		/** Insert the subobject at the start of the list so it can be created and replicated first */
+		/** DEPRECATED: Specify ParentSubObject in SubObjectReplicationParams instead. See FSubObjectReplicationParams for details. */
+		ReplicateWith UE_DEPRECATED(5.8, "Specify ParentSubObject in SubObjectReplicationParams instead. See FSubObjectReplicationParams for details."),
+		/** Insert the subobject at the start of the parents list of subobjeccts so it can be created and replicated first */
 		InsertAtStart,
-		/** Insert at the start of the list or right after the relative handle if one is set */
+		/** Insert at the start of the parents list of subobjects or right after the relative handle if one is set */
 		InsertAtStartOrAfterOther,
 	};
 
+	enum class EChildSubObjectsReplicationOrder : uint8
+	{
+		/** Child subobjects will replicate before the parent */
+		BeforeParent,
+		/** Child subobjects will replicate after the parent */
+		AfterParent,
+	};
+
 } // end namespace UE::Net
 
 //------------------------------------------------------------------------
@@ -226,8 +234,12 @@ protected:
 	/** Destroy the handle and all internal book keeping associated with it. */
 	IRISCORE_API void InternalDestroyNetObject(UE::Net::Private::FInternalNetRefIndex InternalObjectIndex);
 	
-	/** Add SubObjectHandle as SubObject to OwnerHandle. */
-	IRISCORE_API void InternalAddSubObject(FNetRefHandle OwnerHandle, FNetRefHandle SubObjectHandle, FNetRefHandle InsertRelativeToSubObjectHandle, UE::Net::ESubObjectInsertionOrder InsertionOrder);
+	/**
+	 * Add SubObjectHandle as SubObject to ParentHandle. If the ParentHandle is a SubObject itself the SubObject will be share the same RootObject but will be subjects to the Parents conditionals.
+	 * InsertRelativeToSubObjectHandle and InsertionOrder applies to the order of child subobjects of the parent.
+	 * ChildSubObjectsReplicationOrder specified so if a parents childsubobjects replicate before or after the parent itself.
+	 */
+	IRISCORE_API void InternalAddSubObject(FNetRefHandle ParentHandle, FNetRefHandle SubObjectHandle, FNetRefHandle InsertRelativeToSubObjectHandle, UE::Net::ESubObjectInsertionOrder InsertionOrder, UE::Net::EChildSubObjectsReplicationOrder ChildSubObjectsReplicationOrder);
 
 	/** 
 	 * Restart/cancel pending async stop replication.
