Add subobject and group filter support in new creation dependency update

#Synced-CL 50763424

diff --git a/Engine/Plugins/Runtime/ReplicationSystemTestPlugin/Source/Private/Tests/ReplicationSystem/TestDependentObjects.cpp b/Engine/Plugins/Runtime/ReplicationSystemTestPlugin/Source/Private/Tests/ReplicationSystem/TestDependentObjects.cpp
--- a/Engine/Plugins/Runtime/ReplicationSystemTestPlugin/Source/Private/Tests/ReplicationSystem/TestDependentObjects.cpp
+++ b/Engine/Plugins/Runtime/ReplicationSystemTestPlugin/Source/Private/Tests/ReplicationSystem/TestDependentObjects.cpp
@@ -1767,6 +1767,386 @@ UE_NET_TEST_FIXTURE(FTestFilteredDependentObjectFixture, TestComplexCreationDepe
 	}
 }
 
+// Add a creation dependency link after the parent has already been filtered out.
+UE_NET_TEST_FIXTURE(FTestFilteredDependentObjectFixture, TestCreationDependencyAddedAfterParentFilteredOut)
+{
+	UReplicationSystem* ReplicationSystem = Server->ReplicationSystem;
+	UReplicatedTestObjectBridge* Bridge = Server->GetReplicationBridge();
+
+	// Add a client
+	FReplicationSystemTestClient* Client = CreateClient();
+
+	// Spawn server objects
+	UReplicatedTestObject* ParentObject = Server->CreateObject();
+	UReplicatedTestObject* ChildObject = Server->CreateObject();
+
+	// Filter out the parent before adding the creation dependency
+	ReplicationSystem->SetFilter(ParentObject->NetRefHandle, NotRoutedFilterHandle);
+
+	// Send and deliver packets
+	Server->UpdateAndSend({ Client });
+
+	// Parent should not exist, child should exist (no dependency yet)
+	UE_NET_ASSERT_FALSE(Client->IsResolvableNetRefHandle(ParentObject->NetRefHandle));
+	UE_NET_ASSERT_TRUE(Client->IsResolvableNetRefHandle(ChildObject->NetRefHandle));
+
+	// Add the creation dependency while the parent is filtered out
+	Bridge->AddCreationDependencyLink(ParentObject->NetRefHandle, ChildObject->NetRefHandle);
+
+	// Send and deliver packets
+	Server->UpdateAndSend({ Client });
+
+	// Both should be gone since the parent is filtered out and the child depends on it
+	UE_NET_ASSERT_FALSE(Client->IsResolvableNetRefHandle(ParentObject->NetRefHandle));
+	UE_NET_ASSERT_FALSE(Client->IsResolvableNetRefHandle(ChildObject->NetRefHandle));
+
+	// Filter parent back in
+	ReplicationSystem->SetFilter(ParentObject->NetRefHandle, InvalidNetObjectFilterHandle);
+
+	// Send and deliver packets
+	Server->UpdateAndSend({ Client });
+
+	// Both objects should now exist on the client
+	UE_NET_ASSERT_TRUE(Client->IsResolvableNetRefHandle(ParentObject->NetRefHandle));
+	UE_NET_ASSERT_TRUE(Client->IsResolvableNetRefHandle(ChildObject->NetRefHandle));
+}
+
+// Filter a creation dependency parent out via an exclusion group filter
+UE_NET_TEST_FIXTURE(FTestFilteredDependentObjectFixture, TestCreationDependencyFilteredByExclusionGroup)
+{
+	UReplicationSystem* ReplicationSystem = Server->ReplicationSystem;
+	UReplicatedTestObjectBridge* Bridge = Server->GetReplicationBridge();
+
+	// Add a client
+	FReplicationSystemTestClient* Client = CreateClient();
+
+	// Spawn server objects
+	UReplicatedTestObject* ParentObject = Server->CreateObject();
+	UReplicatedTestObject* ChildObject = Server->CreateObject();
+	Bridge->AddCreationDependencyLink(ParentObject->NetRefHandle, ChildObject->NetRefHandle);
+
+	// Send and deliver packets
+	Server->UpdateAndSend({ Client });
+
+	// Both objects should exist on client
+	UE_NET_ASSERT_TRUE(Client->IsResolvableNetRefHandle(ParentObject->NetRefHandle));
+	UE_NET_ASSERT_TRUE(Client->IsResolvableNetRefHandle(ChildObject->NetRefHandle));
+
+	// Add parent to an exclusion group and filter it out
+	FNetObjectGroupHandle ExclusionGroup = ReplicationSystem->CreateGroup(NAME_None);
+	ReplicationSystem->AddExclusionFilterGroup(ExclusionGroup);
+	ReplicationSystem->AddToGroup(ExclusionGroup, ParentObject->NetRefHandle);
+
+	// Send and deliver packets
+	Server->UpdateAndSend({ Client });
+
+	// Both objects should be filtered out
+	UE_NET_ASSERT_FALSE(Client->IsResolvableNetRefHandle(ParentObject->NetRefHandle));
+	UE_NET_ASSERT_FALSE(Client->IsResolvableNetRefHandle(ChildObject->NetRefHandle));
+
+	// Filter the parent back in
+	ReplicationSystem->SetGroupFilterStatus(ExclusionGroup, ENetFilterStatus::Allow);
+
+	// Send and deliver packets
+	Server->UpdateAndSend({ Client });
+
+	// Both objects should exist on client again
+	UE_NET_ASSERT_TRUE(Client->IsResolvableNetRefHandle(ParentObject->NetRefHandle));
+	UE_NET_ASSERT_TRUE(Client->IsResolvableNetRefHandle(ChildObject->NetRefHandle));
+}
+
+// Filter a creation dependency parent out via an owner filter
+UE_NET_TEST_FIXTURE(FTestFilteredDependentObjectFixture, TestCreationDependencyFilteredByOwnerFilter)
+{
+	UReplicationSystem* ReplicationSystem = Server->ReplicationSystem;
+	UReplicatedTestObjectBridge* Bridge = Server->GetReplicationBridge();
+
+	// Add a client
+	FReplicationSystemTestClient* Client = CreateClient();
+
+	// Spawn server objects
+	UReplicatedTestObject* ParentObject = Server->CreateObject();
+	UReplicatedTestObject* ChildObject = Server->CreateObject();
+	Bridge->AddCreationDependencyLink(ParentObject->NetRefHandle, ChildObject->NetRefHandle);
+
+	// Send and deliver packets
+	Server->UpdateAndSend({ Client });
+
+	// Both objects should exist on client
+	UE_NET_ASSERT_TRUE(Client->IsResolvableNetRefHandle(ParentObject->NetRefHandle));
+	UE_NET_ASSERT_TRUE(Client->IsResolvableNetRefHandle(ChildObject->NetRefHandle));
+
+	// Set owner filter on parent without setting an owning connection, which filters it out for all connections
+	ReplicationSystem->SetFilter(ParentObject->NetRefHandle, ToOwnerFilterHandle);
+
+	// Send and deliver packets
+	Server->UpdateAndSend({ Client });
+
+	// Both objects should be filtered out
+	UE_NET_ASSERT_FALSE(Client->IsResolvableNetRefHandle(ParentObject->NetRefHandle));
+	UE_NET_ASSERT_FALSE(Client->IsResolvableNetRefHandle(ChildObject->NetRefHandle));
+
+	// Set the client as the owning connection to filter the parent back in
+	ReplicationSystem->SetOwningNetConnection(ParentObject->NetRefHandle, Client->ConnectionIdOnServer);
+
+	// Send and deliver packets
+	Server->UpdateAndSend({ Client });
+
+	// Both objects should exist on client again
+	UE_NET_ASSERT_TRUE(Client->IsResolvableNetRefHandle(ParentObject->NetRefHandle));
+	UE_NET_ASSERT_TRUE(Client->IsResolvableNetRefHandle(ChildObject->NetRefHandle));
+}
+
+// Filter a creation dependency out via dynamic filter, then filter in with an inclusion group.
+UE_NET_TEST_FIXTURE(FTestFilteredDependentObjectFixture, TestCreationDependencyFilteredByDynamicFilterOverriddenByInclusionGroup)
+{
+	UReplicationSystem* ReplicationSystem = Server->ReplicationSystem;
+	UReplicatedTestObjectBridge* Bridge = Server->GetReplicationBridge();
+
+	// Add a client
+	FReplicationSystemTestClient* Client = CreateClient();
+
+	// Spawn server objects
+	UReplicatedTestObject* ParentObject = Server->CreateObject();
+	UReplicatedTestObject* ChildObject = Server->CreateObject();
+	Bridge->AddCreationDependencyLink(ParentObject->NetRefHandle, ChildObject->NetRefHandle);
+
+	// Send and deliver packets
+	Server->UpdateAndSend({ Client });
+
+	// Both objects should exist on client
+	UE_NET_ASSERT_TRUE(Client->IsResolvableNetRefHandle(ParentObject->NetRefHandle));
+	UE_NET_ASSERT_TRUE(Client->IsResolvableNetRefHandle(ChildObject->NetRefHandle));
+
+	// Filter out parent
+	ReplicationSystem->SetFilter(ParentObject->NetRefHandle, NotRoutedFilterHandle);
+
+	// Send and deliver packets
+	Server->UpdateAndSend({ Client });
+
+	// Both objects should be filtered out
+	UE_NET_ASSERT_FALSE(Client->IsResolvableNetRefHandle(ParentObject->NetRefHandle));
+	UE_NET_ASSERT_FALSE(Client->IsResolvableNetRefHandle(ChildObject->NetRefHandle));
+
+	// Add parent to an inclusion group and allow it for this connection, overriding the dynamic filter
+	FNetObjectGroupHandle InclusionGroup = ReplicationSystem->CreateGroup(NAME_None);
+	ReplicationSystem->AddInclusionFilterGroup(InclusionGroup);
+	ReplicationSystem->AddToGroup(InclusionGroup, ParentObject->NetRefHandle);
+	ReplicationSystem->SetGroupFilterStatus(InclusionGroup, Client->ConnectionIdOnServer, ENetFilterStatus::Allow);
+
+	// Send and deliver packets
+	Server->UpdateAndSend({ Client });
+
+	// Both objects should exist on client again since the inclusion group overrides the dynamic filter
+	UE_NET_ASSERT_TRUE(Client->IsResolvableNetRefHandle(ParentObject->NetRefHandle));
+	UE_NET_ASSERT_TRUE(Client->IsResolvableNetRefHandle(ChildObject->NetRefHandle));
+}
+
+// Add a replicated subobject to a child in a creation dependency, filter out the parent, then filter it back in.
+UE_NET_TEST_FIXTURE(FTestFilteredDependentObjectFixture, TestCreationDependencyChildSubObjectFilteredWithParent)
+{
+	UReplicationSystem* ReplicationSystem = Server->ReplicationSystem;
+	UReplicatedTestObjectBridge* Bridge = Server->GetReplicationBridge();
+
+	// Add a client
+	FReplicationSystemTestClient* Client = CreateClient();
+
+	// Spawn server objects and set up creation dependency
+	UReplicatedTestObject* ParentObject = Server->CreateObject();
+	UReplicatedTestObject* ChildObject = Server->CreateObject();
+	Bridge->AddCreationDependencyLink(ParentObject->NetRefHandle, ChildObject->NetRefHandle);
+
+	// Add a replicated subobject to the child
+	UReplicatedTestObject* ChildSubObject = Server->CreateSubObject(ChildObject->NetRefHandle);
+
+	// Send and deliver packets
+	Server->UpdateAndSend({ Client });
+
+	// All objects should exist on client
+	UE_NET_ASSERT_TRUE(Client->IsResolvableNetRefHandle(ParentObject->NetRefHandle));
+	UE_NET_ASSERT_TRUE(Client->IsResolvableNetRefHandle(ChildObject->NetRefHandle));
+	UE_NET_ASSERT_TRUE(Client->IsResolvableNetRefHandle(ChildSubObject->NetRefHandle));
+
+	// Filter out the parent
+	ReplicationSystem->SetFilter(ParentObject->NetRefHandle, NotRoutedFilterHandle);
+
+	// Send and deliver packets
+	Server->UpdateAndSend({ Client });
+
+	// Parent, child, and the child's subobject should all be filtered out
+	UE_NET_ASSERT_FALSE(Client->IsResolvableNetRefHandle(ParentObject->NetRefHandle));
+	UE_NET_ASSERT_FALSE(Client->IsResolvableNetRefHandle(ChildObject->NetRefHandle));
+	UE_NET_ASSERT_FALSE(Client->IsResolvableNetRefHandle(ChildSubObject->NetRefHandle));
+
+	// Add a replicated subobject to the child AFTER parent is filtered
+	UReplicatedTestObject* ChildSubObjectAddedAfterFilterApplied = Server->CreateSubObject(ChildObject->NetRefHandle);
+
+	// Send and deliver packets
+	Server->UpdateAndSend({ Client });
+
+	// All objects should be filtered, including the child subobject added after the parent was filtered.
+	UE_NET_ASSERT_FALSE(Client->IsResolvableNetRefHandle(ParentObject->NetRefHandle));
+	UE_NET_ASSERT_FALSE(Client->IsResolvableNetRefHandle(ChildObject->NetRefHandle));
+	UE_NET_ASSERT_FALSE(Client->IsResolvableNetRefHandle(ChildSubObject->NetRefHandle));
+	UE_NET_ASSERT_FALSE(Client->IsResolvableNetRefHandle(ChildSubObjectAddedAfterFilterApplied->NetRefHandle));
+
+	// Filter the parent back in
+	ReplicationSystem->SetFilter(ParentObject->NetRefHandle, InvalidNetObjectFilterHandle);
+
+	// Send and deliver packets
+	Server->UpdateAndSend({ Client });
+
+	// All objects should exist on client again
+	UE_NET_ASSERT_TRUE(Client->IsResolvableNetRefHandle(ParentObject->NetRefHandle));
+	UE_NET_ASSERT_TRUE(Client->IsResolvableNetRefHandle(ChildObject->NetRefHandle));
+	UE_NET_ASSERT_TRUE(Client->IsResolvableNetRefHandle(ChildSubObject->NetRefHandle));
+	UE_NET_ASSERT_TRUE(Client->IsResolvableNetRefHandle(ChildSubObjectAddedAfterFilterApplied->NetRefHandle));
+}
+
+// Destroy a subobject filtered by creation dependency, add a new object and verify stale state doesn't filter it out
+UE_NET_TEST_FIXTURE(FTestFilteredDependentObjectFixture, TestCreationDependencyChildSubObjectDeleted)
+{
+	UReplicationSystem* ReplicationSystem = Server->ReplicationSystem;
+	UReplicatedTestObjectBridge* Bridge = Server->GetReplicationBridge();
+
+	// Add a client
+	FReplicationSystemTestClient* Client = CreateClient();
+
+	// Spawn server objects and set up creation dependency
+	UReplicatedTestObject* ParentObject = Server->CreateObject();
+	UReplicatedTestObject* ChildObject = Server->CreateObject();
+	Bridge->AddCreationDependencyLink(ParentObject->NetRefHandle, ChildObject->NetRefHandle);
+
+	// Add a replicated subobject to the child
+	UReplicatedTestObject* ChildSubObject = Server->CreateSubObject(ChildObject->NetRefHandle);
+
+	// Send and deliver packets
+	Server->UpdateAndSend({ Client });
+
+	// All objects should exist on client
+	UE_NET_ASSERT_TRUE(Client->IsResolvableNetRefHandle(ParentObject->NetRefHandle));
+	UE_NET_ASSERT_TRUE(Client->IsResolvableNetRefHandle(ChildObject->NetRefHandle));
+	UE_NET_ASSERT_TRUE(Client->IsResolvableNetRefHandle(ChildSubObject->NetRefHandle));
+
+	// Filter out the parent
+	ReplicationSystem->SetFilter(ParentObject->NetRefHandle, NotRoutedFilterHandle);
+
+	// Send and deliver packets
+	Server->UpdateAndSend({ Client });
+
+	// Parent, child, and the child's subobject should all be filtered out
+	UE_NET_ASSERT_FALSE(Client->IsResolvableNetRefHandle(ParentObject->NetRefHandle));
+	UE_NET_ASSERT_FALSE(Client->IsResolvableNetRefHandle(ChildObject->NetRefHandle));
+	UE_NET_ASSERT_FALSE(Client->IsResolvableNetRefHandle(ChildSubObject->NetRefHandle));
+
+	// Destroy child subobject
+	const FNetRefHandle ChildSubObjectHandle = ChildSubObject->NetRefHandle;
+	Server->DestroyObject(ChildSubObject);
+
+	// Send and deliver packets
+	Server->UpdateAndSend({ Client });
+
+	// Parent, child, and the child's subobject should all not exist on the client
+	UE_NET_ASSERT_FALSE(Client->IsResolvableNetRefHandle(ParentObject->NetRefHandle));
+	UE_NET_ASSERT_FALSE(Client->IsResolvableNetRefHandle(ChildObject->NetRefHandle));
+	UE_NET_ASSERT_FALSE(Client->IsResolvableNetRefHandle(ChildSubObjectHandle));
+
+	// Filter the parent back in
+	ReplicationSystem->SetFilter(ParentObject->NetRefHandle, InvalidNetObjectFilterHandle);
+
+	// Send and deliver packets
+	Server->UpdateAndSend({ Client });
+
+	// All objects except the destroyed subobject should exist on client again
+	UE_NET_ASSERT_TRUE(Client->IsResolvableNetRefHandle(ParentObject->NetRefHandle));
+	UE_NET_ASSERT_TRUE(Client->IsResolvableNetRefHandle(ChildObject->NetRefHandle));
+	UE_NET_ASSERT_FALSE(Client->IsResolvableNetRefHandle(ChildSubObjectHandle));
+
+	// Create another object, should reuse internal index of destroyed subobject
+	UReplicatedTestObject* ObjectReusingInternalIndex = Server->CreateObject();
+
+	// Send and deliver packets
+	Server->UpdateAndSend({ Client });
+
+	// All objects should exist on the client (old creation dependency state from the destroyed subobject shouldn't filter out the new object)
+	UE_NET_ASSERT_TRUE(Client->IsResolvableNetRefHandle(ParentObject->NetRefHandle));
+	UE_NET_ASSERT_TRUE(Client->IsResolvableNetRefHandle(ChildObject->NetRefHandle));
+	UE_NET_ASSERT_TRUE(Client->IsResolvableNetRefHandle(ObjectReusingInternalIndex->NetRefHandle));
+}
+
+// Destroy a child + subobject filtered by creation dependency, add two new objects and verify stale state doesn't filter them out
+UE_NET_TEST_FIXTURE(FTestFilteredDependentObjectFixture, TestCreationDependencyChildAndSubObjectDeleted)
+{
+	UReplicationSystem* ReplicationSystem = Server->ReplicationSystem;
+	UReplicatedTestObjectBridge* Bridge = Server->GetReplicationBridge();
+
+	// Add a client
+	FReplicationSystemTestClient* Client = CreateClient();
+
+	// Spawn server objects and set up creation dependency
+	UReplicatedTestObject* ParentObject = Server->CreateObject();
+	UReplicatedTestObject* ChildObject = Server->CreateObject();
+	Bridge->AddCreationDependencyLink(ParentObject->NetRefHandle, ChildObject->NetRefHandle);
+
+	// Add a replicated subobject to the child
+	UReplicatedTestObject* ChildSubObject = Server->CreateSubObject(ChildObject->NetRefHandle);
+
+	// Send and deliver packets
+	Server->UpdateAndSend({ Client });
+
+	// All objects should exist on client
+	UE_NET_ASSERT_TRUE(Client->IsResolvableNetRefHandle(ParentObject->NetRefHandle));
+	UE_NET_ASSERT_TRUE(Client->IsResolvableNetRefHandle(ChildObject->NetRefHandle));
+	UE_NET_ASSERT_TRUE(Client->IsResolvableNetRefHandle(ChildSubObject->NetRefHandle));
+
+	// Filter out the parent
+	ReplicationSystem->SetFilter(ParentObject->NetRefHandle, NotRoutedFilterHandle);
+
+	// Send and deliver packets
+	Server->UpdateAndSend({ Client });
+
+	// Parent, child, and the child's subobject should all be filtered out
+	UE_NET_ASSERT_FALSE(Client->IsResolvableNetRefHandle(ParentObject->NetRefHandle));
+	UE_NET_ASSERT_FALSE(Client->IsResolvableNetRefHandle(ChildObject->NetRefHandle));
+	UE_NET_ASSERT_FALSE(Client->IsResolvableNetRefHandle(ChildSubObject->NetRefHandle));
+
+	// Destroy child object and subobject.
+	const FNetRefHandle ChildObjectHandle = ChildObject->NetRefHandle;
+	const FNetRefHandle ChildSubObjectHandle = ChildSubObject->NetRefHandle;
+	Server->DestroyObject(ChildObject);
+
+	// Send and deliver packets
+	Server->UpdateAndSend({ Client });
+
+	// Parent, child, and the child's subobject should all not exist on the client
+	UE_NET_ASSERT_FALSE(Client->IsResolvableNetRefHandle(ParentObject->NetRefHandle));
+	UE_NET_ASSERT_FALSE(Client->IsResolvableNetRefHandle(ChildObjectHandle));
+	UE_NET_ASSERT_FALSE(Client->IsResolvableNetRefHandle(ChildSubObjectHandle));
+
+	// Filter the parent back in
+	ReplicationSystem->SetFilter(ParentObject->NetRefHandle, InvalidNetObjectFilterHandle);
+
+	// Send and deliver packets
+	Server->UpdateAndSend({ Client });
+
+	// Only the parent object should exist.
+	UE_NET_ASSERT_TRUE(Client->IsResolvableNetRefHandle(ParentObject->NetRefHandle));
+	UE_NET_ASSERT_FALSE(Client->IsResolvableNetRefHandle(ChildObjectHandle));
+	UE_NET_ASSERT_FALSE(Client->IsResolvableNetRefHandle(ChildSubObjectHandle));
+
+	// Create two objects, should reuse internal index of the destroyed child + subobject
+	UReplicatedTestObject* ObjectReusingChildInternalIndex = Server->CreateObject();
+	UReplicatedTestObject* ObjectReusingChildSubObjectInternalIndex = Server->CreateObject();	
+
+	// Send and deliver packets
+	Server->UpdateAndSend({ Client });
+
+	// All objects should exist on the client (old creation dependency state from the destroyed objects shouldn't filter out the new objects)
+	UE_NET_ASSERT_TRUE(Client->IsResolvableNetRefHandle(ParentObject->NetRefHandle));
+	UE_NET_ASSERT_TRUE(Client->IsResolvableNetRefHandle(ObjectReusingChildInternalIndex->NetRefHandle));
+	UE_NET_ASSERT_TRUE(Client->IsResolvableNetRefHandle(ObjectReusingChildSubObjectInternalIndex->NetRefHandle));
+}
+
 UE_NET_TEST_FIXTURE(FTestFilteredDependentObjectFixture, TestCircularDependencyIsDenied)
 {
 	UReplicatedTestObjectBridge* Bridge = Server->GetReplicationBridge();
diff --git a/Engine/Source/Runtime/Net/Iris/Private/Iris/ReplicationSystem/Filtering/ReplicationFiltering.cpp b/Engine/Source/Runtime/Net/Iris/Private/Iris/ReplicationSystem/Filtering/ReplicationFiltering.cpp
--- a/Engine/Source/Runtime/Net/Iris/Private/Iris/ReplicationSystem/Filtering/ReplicationFiltering.cpp
+++ b/Engine/Source/Runtime/Net/Iris/Private/Iris/ReplicationSystem/Filtering/ReplicationFiltering.cpp
@@ -1284,28 +1284,34 @@ void FReplicationFiltering::PreUpdateDynamicFiltering()
 
 void FReplicationFiltering::FilterOutCreationDependents(uint32 ParentIndex, FPerConnectionInfo& ConnectionInfo, const FNetBitArrayView ObjectsWithCreationDependents) const
 {
-	if (ConnectionInfo.FilteredOutCreationDependents.GetBit(ParentIndex))
-	{
-		// The parent and its descendents have already been filtered out.
-		return;
-	}
-
 	TArray<FInternalNetRefIndex, TInlineAllocator<64>> ObjectsToProcess;
 	ObjectsToProcess.Add(ParentIndex);
 	while (ObjectsToProcess.Num() > 0)
 	{
 		const FInternalNetRefIndex CurrentIndex = ObjectsToProcess.Pop(EAllowShrinking::No);
-		if (ConnectionInfo.FilteredOutCreationDependents.GetBit(CurrentIndex))
+
+		TConstArrayView<const FInternalNetRefIndex> Children = NetRefHandleManager->GetCreationDependents(CurrentIndex);
+		for (const FInternalNetRefIndex ChildIndex : Children)
 		{
-			// This section of the dependency graph has already been filtered out, so we can skip processing it.
-			continue;
-		}
+			if (ConnectionInfo.FilteredOutCreationDependents.GetBit(ChildIndex))
+			{
+				// This section of the dependency graph has already been filtered out, so we can skip processing it.
+				continue;
+			}
 
-		ConnectionInfo.FilteredOutCreationDependents.SetBit(CurrentIndex);
+			// Filter out this child and its subobjects.
+			ConnectionInfo.FilteredOutCreationDependents.SetBit(ChildIndex);
+			for (const FInternalNetRefIndex SubObjectIndex : NetRefHandleManager->GetSubObjects(ChildIndex))
+			{
+				ConnectionInfo.FilteredOutCreationDependents.SetBit(SubObjectIndex);
+			}
 
-		TConstArrayView<const FInternalNetRefIndex> Children = NetRefHandleManager->GetCreationDependents(CurrentIndex);
-		// Add the children to process them and their descendents.
-		ObjectsToProcess.Append(Children);
+			if (ObjectsWithCreationDependents.GetBit(ChildIndex))
+			{
+				// This child has dependents, so process them as well.
+				ObjectsToProcess.Push(ChildIndex);
+			}
+		}
 	}
 }
 
@@ -1315,48 +1321,109 @@ void FReplicationFiltering::FilterInCreationDependents(uint32 ParentIndex, FPerC
 	ObjectsToProcess.Add(ParentIndex);
 	while (ObjectsToProcess.Num() > 0)
 	{
-		const FInternalNetRefIndex CurrentIndex = ObjectsToProcess.Pop(EAllowShrinking::No);
+		const FInternalNetRefIndex CurrentParent = ObjectsToProcess.Pop(EAllowShrinking::No);
 
-		bool bFilteredOutByParent = false;
-		TConstArrayView<const FInternalNetRefIndex> Dependencies = NetRefHandleManager->GetCreationDependencies(CurrentIndex);
-		for (const FInternalNetRefIndex DependencyIndex : Dependencies)
+		TConstArrayView<const FInternalNetRefIndex> Children = NetRefHandleManager->GetCreationDependents(CurrentParent);
+		for (const FInternalNetRefIndex ChildIndex : Children)
 		{
-			if (ConnectionInfo.FilteredOutCreationDependents.GetBit(DependencyIndex))
+			if (!ConnectionInfo.FilteredOutCreationDependents.GetBit(ChildIndex))
+			{
+				continue;
+			}
+
+			bool bFilteredOutByParent = false;
+			TConstArrayView<const FInternalNetRefIndex> Dependencies = NetRefHandleManager->GetCreationDependencies(ChildIndex);
+			for (const FInternalNetRefIndex DependencyIndex : Dependencies)
+			{
+				if (ConnectionInfo.FilteredOutCreationDependents.GetBit(DependencyIndex) || !ConnectionInfo.ObjectsInScopeBeforeDynamicFiltering.GetBit(DependencyIndex) || ConnectionInfo.DynamicFilteredOutObjectsHysteresisAdjusted.GetBit(DependencyIndex))
+				{
+					// This object has a parent that has been filtered out, so we can't filter it in.
+					//
+					// Note: In the case that an object has N>1 parents that would be filtered in by one call to FilterInCreationDependents,
+					//       that object will be processed N times and only be filtered in the Nth time.
+					// Example: W starts filtered out. X, Y, and Z would be relevant, but creation dependency filters them out.
+					//          W is filtered in. Because of DFS, we process in the order W->Y->Z->X->Z. Z is only filtered in the 2nd time it is processed.
+					//   W
+					//  / \
+					// X   Y
+					//  \ /
+					//   Z
+					bFilteredOutByParent = true;
+					break;
+				}
+			}
+
+			if (!bFilteredOutByParent)
 			{
-				// This object has a parent that has been filtered out, so we can't filter it in.
-				//
-				// Note: In the case that an object has N>1 parents that would be filtered in by one call to FilterInCreationDependents,
-				//       that object will be processed N times and only be filtered in the Nth time.
-				// Example: W starts filtered out. X, Y, and Z would be relevant, but creation dependency filters them out.
-				//          W is filtered in. Because of DFS, we process in the order W->Y->Z->X->Z. Z is only filtered in the 2nd time it is processed.
-				//   W
-				//  / \
-				// X   Y
-				//  \ /
-				//   Z
-				bFilteredOutByParent = true;
-				break;
+				// Filter in this child and its subobjects.
+				ConnectionInfo.FilteredOutCreationDependents.ClearBit(ChildIndex);
+				for (const FInternalNetRefIndex SubObjectIndex : NetRefHandleManager->GetSubObjects(ChildIndex))
+				{
+					ConnectionInfo.FilteredOutCreationDependents.ClearBit(SubObjectIndex);
+				}
+
+				if (ObjectsWithCreationDependents.GetBit(ChildIndex))
+				{
+					// This child has dependents, so process them as well.
+					ObjectsToProcess.Add(ChildIndex);
+				}
 			}
 		}
+	}
+}
 
-		if (bFilteredOutByParent)
+void FReplicationFiltering::RefreshCreationDependencyFilterStatus(uint32 ObjectIndex, FPerConnectionInfo& ConnectionInfo, const FNetBitArrayView ObjectsWithCreationDependents) const
+{
+	if (NetRefHandleManager->IsSubObject(ObjectIndex))
+	{
+		// Subobjects should follow the state of their root
+		// NOTE: We flag removed subobjects, but by the time this function is called, they are not subobjects. They will have their bit cleared in the !bHasFilteredParent block below.
+		const uint32 RootIndex = NetRefHandleManager->GetReplicatedObjectDataNoCheck(ObjectIndex).SubObjectRootIndex;
+		const bool bIsRootFiltered = ConnectionInfo.FilteredOutCreationDependents.GetBit(RootIndex);
+		ConnectionInfo.FilteredOutCreationDependents.SetBitValue(ObjectIndex, bIsRootFiltered);
+		return;
+	}
+
+	bool bHasFilteredParent = false;
+	TConstArrayView<const FInternalNetRefIndex> Parents = NetRefHandleManager->GetCreationDependencies(ObjectIndex);
+	for (const FInternalNetRefIndex ParentIndex : Parents)
+	{
+		if (ConnectionInfo.FilteredOutCreationDependents.GetBit(ParentIndex) || !ConnectionInfo.ObjectsInScopeBeforeDynamicFiltering.GetBit(ParentIndex) || ConnectionInfo.DynamicFilteredOutObjectsHysteresisAdjusted.GetBit(ParentIndex))
 		{
-			// This object has a filtered parent, we can't filter them or their children in.
-			continue;
+			bHasFilteredParent = true;
+			break;
 		}
+	}
 
-		ConnectionInfo.FilteredOutCreationDependents.ClearBit(CurrentIndex);
+	const bool bWasFilteredCreationDependent = ConnectionInfo.FilteredOutCreationDependents.GetBit(ObjectIndex);
+	if (bHasFilteredParent)
+	{
+		// This object has a filtered dependency, so it and its subobjects should be filtered by creation dependency.
+		ConnectionInfo.FilteredOutCreationDependents.SetBit(ObjectIndex);
+		for (const FInternalNetRefIndex SubObjectIndex : NetRefHandleManager->GetSubObjects(ObjectIndex))
+		{
+			ConnectionInfo.FilteredOutCreationDependents.SetBit(SubObjectIndex);
+		}
 
-		TConstArrayView<const FInternalNetRefIndex> Children = NetRefHandleManager->GetCreationDependents(CurrentIndex);
-		for (const FInternalNetRefIndex ChildIndex : Children)
+		// Only filter out dependents if this object wasn't already filtered out by creation dependents.
+		if (ObjectsWithCreationDependents.GetBit(ObjectIndex) && !bWasFilteredCreationDependent)
 		{
-			if (ConnectionInfo.DynamicFilteredOutObjectsHysteresisAdjusted.GetBit(ChildIndex))
-			{
-				// This child is filtered out, so don't process their subtree.
-				continue;
-			}
+			FilterOutCreationDependents(ObjectIndex, ConnectionInfo, ObjectsWithCreationDependents);
+		}
+	}
+	else
+	{
+		// All of this object's dependencies aren't filtered, so it and its subobjects shouldn't be filtered by creation dependency.
+		ConnectionInfo.FilteredOutCreationDependents.ClearBit(ObjectIndex);
+		for (const FInternalNetRefIndex SubObjectIndex : NetRefHandleManager->GetSubObjects(ObjectIndex))
+		{
+			ConnectionInfo.FilteredOutCreationDependents.ClearBit(SubObjectIndex);
+		}
 
-			ObjectsToProcess.Add(ChildIndex);
+		// Only filter in dependents if this object was filtered by creation dependency and its not currently filtered otherwise.
+		if (ObjectsWithCreationDependents.GetBit(ObjectIndex) && bWasFilteredCreationDependent && ConnectionInfo.ObjectsInScopeBeforeDynamicFiltering.GetBit(ObjectIndex) && !ConnectionInfo.DynamicFilteredOutObjectsHysteresisAdjusted.GetBit(ObjectIndex))
+		{
+			FilterInCreationDependents(ObjectIndex, ConnectionInfo, ObjectsWithCreationDependents);
 		}
 	}
 }
@@ -1445,7 +1512,10 @@ void FReplicationFiltering::UpdateDynamicFiltering()
 	const FNetBitArrayView DependentObjectsView = NetRefHandleManager->GetDependentObjectInternalIndices();
 	const uint32* DependentObjectsData = DependentObjectsView.GetDataChecked(WordCountForObjectBitArrays);
 
+	FNetBitArrayView ObjectsWithDirtyCreationDependencies = NetRefHandleManager->GetObjectsWithDirtyCreationDependencies();
+
 	const FNetBitArrayView ObjectsWithCreationDependents = NetRefHandleManager->GetObjectsWithCreationDependents();
+	const uint32* ObjectsWithCreationDependentsData = ObjectsWithCreationDependents.GetDataChecked(WordCountForObjectBitArrays);
 
 	const FNetBitArrayView ObjectsWithCreationDependencies = NetRefHandleManager->GetObjectsWithCreationDependencies();
 	const uint32* ObjectsWithCreationDependenciesData = ObjectsWithCreationDependencies.GetDataChecked(WordCountForObjectBitArrays);
@@ -1687,67 +1757,45 @@ void FReplicationFiltering::UpdateDynamicFiltering()
 		if (bCVarRepFilterUpdateCreationDependenciesOnRelevancyChange)
 		{
 			IRIS_PROFILER_SCOPE(FReplicationFiltering_UpdateCreationDependencies);
-
-			// This will be run on any object that was filtered out in the last frame.
-			auto ProcessFilteredOutObject = [&](const uint32 ObjectIndex)
+			// Reevaluate filter status for any objects explicitly dirtied.
+			ObjectsWithDirtyCreationDependencies.ForAllSetBits([this, &ConnectionInfo, &ObjectsWithCreationDependents](uint32 ObjectIndex)
 			{
-				if (ObjectsWithCreationDependents.IsBitSet(ObjectIndex))
-				{
-					// Filter out this object and all of its descendents
-					FilterOutCreationDependents(ObjectIndex, ConnectionInfo, ObjectsWithCreationDependents);
-				}
-				else if (ObjectsWithCreationDependencies.IsBitSet(ObjectIndex))
-				{
-					// This object is still a part of a dependency graph, so mark as filtered out.
-					ConnectionInfo.FilteredOutCreationDependents.SetBit(ObjectIndex);
-				}
-				else
-				{
-					// This object isn't a part of a dependency graph, so it shouldn't be marked.
-					ConnectionInfo.FilteredOutCreationDependents.ClearBit(ObjectIndex);
-				}
-			};
+				this->RefreshCreationDependencyFilterStatus(ObjectIndex, ConnectionInfo, ObjectsWithCreationDependents);
+			});
 
-			// For parents that have been filtered out this frame, filter out their dependents.
-			if (HysteresisState.Mode == EHysteresisProcessingMode::Enabled)
+			// Update filter status for any object with creation dependents that became relevant/filtered this frame.
+			const uint32* PrevObjectsInScopeData = ConnectionInfo.ObjectsInScope.GetDataChecked(WordCountForObjectBitArrays);
+			const uint32* ObjectsInScopeBeforeDynamicFilteringData = ConnectionInfo.ObjectsInScopeBeforeDynamicFiltering.GetDataChecked(WordCountForObjectBitArrays);
+			const uint32* DynamicFilteredOutObjectsHysteresisAdjustedData = DynamicFilteredOutObjectsHysteresisAdjusted.GetDataChecked(WordCountForObjectBitArrays);
+			const uint32* FilteredOutCreationDependentsData = ConnectionInfo.FilteredOutCreationDependents.GetDataChecked(WordCountForObjectBitArrays);
+			for (uint32 WordIt = 0, WordEndIt = WordCountForObjectBitArrays; WordIt != WordEndIt; ++WordIt)
 			{
-				// Only filter creation dependents of objects in FilteredOutObjects if they didn't enter hysteresis.
-				for (const uint32 ObjectIndex : FilteredOutObjects)
-				{
-					if (!ConnectionInfo.HysteresisUpdater.IsObjectUpdated(ObjectIndex))
-					{
-						ProcessFilteredOutObject(ObjectIndex);
-					}
-				}
+				const uint32 BitOffset = WordIt * 32U;
 
-				// FilteredOutByHysteresisObjects will contain all of the hysteresis objects whose frame timer just expired.
-				for (const uint32 ObjectIndex : FilteredOutByHysteresisObjects)
+				const uint32 PrevInScope = PrevObjectsInScopeData[WordIt];
+				const uint32 InScope = ObjectsInScopeBeforeDynamicFilteringData[WordIt] & ~DynamicFilteredOutObjectsHysteresisAdjustedData[WordIt];
+				const uint32 WithCreationDependents = ObjectsWithCreationDependentsData[WordIt];
+				const uint32 FilteredOutCreationDependents = FilteredOutCreationDependentsData[WordIt];
+				
+				uint32 NewlyFilteredOut = PrevInScope & ~InScope & WithCreationDependents & ~FilteredOutCreationDependents;
+				while (NewlyFilteredOut)
 				{
-					ProcessFilteredOutObject(ObjectIndex);
-				}
-			}
-			else
-			{
-				// When hysteresis is disabled, FilteredOutObjects will actually contain all of the objects that have been filtered out.
-				for (const uint32 ObjectIndex : FilteredOutObjects)
-				{
-					ProcessFilteredOutObject(ObjectIndex);
+					const uint32 LeastSignificantBit = GetLeastSignificantBit(NewlyFilteredOut);
+					NewlyFilteredOut ^= LeastSignificantBit;
+
+					const uint32 ObjectIndex = BitOffset + FPlatformMath::CountTrailingZeros(LeastSignificantBit);
+					FilterOutCreationDependents(ObjectIndex, ConnectionInfo, ObjectsWithCreationDependents);
 				}
-			}
 
-			// For parents that have been filtered in this frame, filter in their dependents that aren't otherwise filtered out.
-			for (const uint32 ObjectIndex : FilteredInObjects)
-			{
-				if (ObjectsWithCreationDependents.IsBitSet(ObjectIndex) || ObjectsWithCreationDependencies.IsBitSet(ObjectIndex))
+				uint32 NewlyFilteredIn = ~PrevInScope & InScope & WithCreationDependents & ~FilteredOutCreationDependents;
+				while (NewlyFilteredIn)
 				{
-					// Even if this object doesn't have creation dependents, it might need to be filtered in if a dependency link to a filtered out parent was removed.
+					const uint32 LeastSignificantBit = GetLeastSignificantBit(NewlyFilteredIn);
+					NewlyFilteredIn ^= LeastSignificantBit;
+
+					const uint32 ObjectIndex = BitOffset + FPlatformMath::CountTrailingZeros(LeastSignificantBit);
 					FilterInCreationDependents(ObjectIndex, ConnectionInfo, ObjectsWithCreationDependents);
 				}
-				else
-				{
-					// This object doesn't have dependents or dependencies.
-					ConnectionInfo.FilteredOutCreationDependents.ClearBit(ObjectIndex);
-				}
 			}
 		}
 
@@ -1893,6 +1941,12 @@ void FReplicationFiltering::PostUpdateDynamicFiltering()
 
 	ObjectsRequiringDynamicFilterUpdate.ClearAllBits();
 
+	// We always process all objects with dirty creation dependencies on every valid connection.
+	{
+		FNetBitArrayView ObjectsWithDirtyCreationDependencies = NetRefHandleManager->GetObjectsWithDirtyCreationDependencies();
+		ObjectsWithDirtyCreationDependencies.ClearAllBits();
+	}
+
 	// Tell filters to clean up after filtering. It's only called if any object has the filter set.
 	{
 		FNetObjectPostFilteringParams PostFilteringParams;
@@ -2744,26 +2798,6 @@ void FReplicationFiltering::NotifyRemovedDependentObject(FInternalNetRefIndex Ob
 	ObjectsRequiringDynamicFilterUpdate.SetBit(ObjectIndex);
 }
 
-void FReplicationFiltering::NotifyAddedCreationDependency(FInternalNetRefIndex ParentIndex, FInternalNetRefIndex ChildIndex)
-{
-	if (bCVarRepFilterUpdateCreationDependenciesOnRelevancyChange)
-	{
-		// This causes the parent and child to be reprocessed in the creation dependency update.
-		ObjectsRequiringDynamicFilterUpdate.SetBit(ParentIndex);
-		ObjectsRequiringDynamicFilterUpdate.SetBit(ChildIndex);
-	}
-}
-
-void FReplicationFiltering::NotifyRemovedCreationDependency(FInternalNetRefIndex ParentIndex, FInternalNetRefIndex ChildIndex)
-{
-	if (bCVarRepFilterUpdateCreationDependenciesOnRelevancyChange)
-	{
-		// This causes the parent and child to be reprocessed in the creation dependency update.
-		ObjectsRequiringDynamicFilterUpdate.SetBit(ParentIndex);
-		ObjectsRequiringDynamicFilterUpdate.SetBit(ChildIndex);
-	}
-}
-
 // Determine whether a dynamically filtered out dependent object should be in scope or not. Note that this function is expected to only be called for dependent objects NOT in scope. The Visited set is expected to be empty when calling this function.
 ENetFilterStatus FReplicationFiltering::GetDependentObjectFilterStatus(uint32 ConnectionId, FInternalNetRefIndex DependentObjectIndex, TSet<FInternalNetRefIndex>& Visited)
 {
diff --git a/Engine/Source/Runtime/Net/Iris/Private/Iris/ReplicationSystem/Filtering/ReplicationFiltering.h b/Engine/Source/Runtime/Net/Iris/Private/Iris/ReplicationSystem/Filtering/ReplicationFiltering.h
--- a/Engine/Source/Runtime/Net/Iris/Private/Iris/ReplicationSystem/Filtering/ReplicationFiltering.h
+++ b/Engine/Source/Runtime/Net/Iris/Private/Iris/ReplicationSystem/Filtering/ReplicationFiltering.h
@@ -151,9 +151,6 @@ public:
 
 	void NotifyAddedDependentObject(FInternalNetRefIndex ObjectIndex);
 	void NotifyRemovedDependentObject(FInternalNetRefIndex ObjectIndex);
-
-	void NotifyAddedCreationDependency(FInternalNetRefIndex ParentIndex, FInternalNetRefIndex ChildIndex);
-	void NotifyRemovedCreationDependency(FInternalNetRefIndex ParentIndex, FInternalNetRefIndex ChildIndex);
 	
 	// SubObjectFilter status
 	void AddSubObjectFilter(FNetObjectGroupHandle GroupHandle);
@@ -258,11 +255,12 @@ private:
 
 	bool HasDynamicFilters() const;
 
-	// Mark creation dependents as filtered out.
+	// Mark creation dependents as filtered out. This may only be called from within the creation dependency update.
 	void FilterOutCreationDependents(uint32 ParentIndex, FPerConnectionInfo& ConnectionInfo, const FNetBitArrayView ObjectsWithCreationDependents) const;
-
-	// Mark creation dependents as filtered in, if they aren't filtered out by other dependencies.
+	// Mark creation dependents as filtered in if none of their ancestors are filtered out. This may only be called from within the creation dependency update.
 	void FilterInCreationDependents(uint32 ParentIndex, FPerConnectionInfo& ConnectionInfo, const FNetBitArrayView ObjectsWithCreationDependents) const;
+	// Reevaluate the filter status of an object based on creation dependencies. This may only be called from within the creation dependency update.
+	void RefreshCreationDependencyFilterStatus(uint32 ObjectIndex, FPerConnectionInfo& ConnectionInfo, const FNetBitArrayView ObjectsWithCreationDependents) const;
 
 	void UpdateCreationDependentParent(uint32 ChildIndex, const FNetBitArrayView ObjectsWithCreationDependencies, FNetBitArrayView OutConnectionObjectsInScope, FNetBitArrayView ProcessedObjects, bool bIsRecursive) const;
 
diff --git a/Engine/Source/Runtime/Net/Iris/Private/Iris/ReplicationSystem/NetRefHandleManager.cpp b/Engine/Source/Runtime/Net/Iris/Private/Iris/ReplicationSystem/NetRefHandleManager.cpp
--- a/Engine/Source/Runtime/Net/Iris/Private/Iris/ReplicationSystem/NetRefHandleManager.cpp
+++ b/Engine/Source/Runtime/Net/Iris/Private/Iris/ReplicationSystem/NetRefHandleManager.cpp
@@ -139,6 +139,7 @@ void FNetRefHandleManager::Init(const FInitParams& InitParams)
 		InitNetBitArray(&DependentObjectInternalIndices);
 		InitNetBitArray(&ObjectsWithCreationDependents);
 		InitNetBitArray(&ObjectsWithCreationDependencies);
+		InitNetBitArray(&ObjectsWithDirtyCreationDependencies);
 		InitNetBitArray(&ObjectsWithDependentObjectsInternalIndices);
 		InitNetBitArray(&DestroyedStartupObjectInternalIndices);
 		InitNetBitArray(&WantToBeDormantInternalIndices);
@@ -844,6 +845,12 @@ bool FNetRefHandleManager::InternalAddSubObject(FInternalNetRefIndex RootObjectI
 			SubObjectConditionalsArray->Add(ELifetimeCondition::COND_None);
 		}
 	}
+
+	if (ObjectsWithCreationDependencies.GetBit(RootObjectInternalIndex))
+	{
+		// Let filtering know that this subobject might be filtered by creation dependencies
+		ObjectsWithDirtyCreationDependencies.SetBit(SubObjectInternalIndex);
+	}
 	
 	SubObjectData.SubObjectParentIndex = ParentOfSubObjectIndex;
 	SubObjectData.bReplicateChildSubObjectsAfterParent = EnumHasAnyFlags(Flags, EAddSubObjectFlags::ReplicateChildSubObjectsAfterParent) ? 1U : 0U;
@@ -892,6 +899,12 @@ void FNetRefHandleManager::InternalRemoveSubObject(FInternalNetRefIndex RootObje
 		}
 	}
 
+	if (ObjectsWithCreationDependencies.GetBit(SubObjectData.SubObjectRootIndex))
+	{
+		// Let filtering know that creation dependency data may need to be cleared for this subobject
+		ObjectsWithDirtyCreationDependencies.SetBit(SubObjectInternalIndex);
+	}
+
 	SubObjectData.SubObjectRootIndex = InvalidInternalIndex;
 	SubObjectData.SubObjectParentIndex = InvalidInternalIndex;
 
@@ -1151,6 +1164,9 @@ bool FNetRefHandleManager::InternalAddCreationDependency(FInternalNetRefIndex Pa
 		ObjectsWithCreationDependencies.SetBit(ChildIndex);
 	}
 
+	// The child might need to become filtered by its new creation dependency, so dirty it for the creation dependency update.
+	ObjectsWithDirtyCreationDependencies.SetBit(ChildIndex);
+
 	return true;
 }
 
@@ -1185,6 +1201,9 @@ void FNetRefHandleManager::InternalRemoveCreationDependency(FInternalNetRefIndex
 			// We don't free the array because we usually add a new dependency right after or the object will stop replicating soon and the array will be freed there.
 		}
 	}
+
+	// The child might need to become relevant by losing a creation dependency, so dirty it for the creation dependency update.
+	ObjectsWithDirtyCreationDependencies.SetBit(ChildIndex);
 }
 
 TConstArrayView<const FInternalNetRefIndex> FNetRefHandleManager::GetCreationDependents(FInternalNetRefIndex ParentInternalIndex) const
diff --git a/Engine/Source/Runtime/Net/Iris/Private/Iris/ReplicationSystem/NetRefHandleManager.h b/Engine/Source/Runtime/Net/Iris/Private/Iris/ReplicationSystem/NetRefHandleManager.h
--- a/Engine/Source/Runtime/Net/Iris/Private/Iris/ReplicationSystem/NetRefHandleManager.h
+++ b/Engine/Source/Runtime/Net/Iris/Private/Iris/ReplicationSystem/NetRefHandleManager.h
@@ -364,6 +364,11 @@ public:
 
 	TConstArrayView<const FInternalNetRefIndex> GetCreationDependencies(FInternalNetRefIndex ChildInternalIndex) const;
 
+	FNetBitArrayView GetObjectsWithDirtyCreationDependencies() const
+	{
+		return MakeNetBitArrayView(ObjectsWithDirtyCreationDependencies);
+	}
+
 	// Remove DependentHandles from all dependent object tracking
 	// $TODO: This might need a pass ObjectReplicationBridge as well since most likely also must restore filter/polling logic
 	void RemoveDependentObject(FNetRefHandle DependentHandle);
@@ -587,6 +592,9 @@ private:
 	// Objects that cannot be created until the parents they depend on also exist
 	FNetBitArray ObjectsWithCreationDependencies;
 
+	// Objects that need to have their filter status re-evaluated due to a change in creation dependency
+	FNetBitArray ObjectsWithDirtyCreationDependencies;
+
 	// Bitset containing all internal indices that have dependentObjects
 	FNetBitArray ObjectsWithDependentObjectsInternalIndices;
 
diff --git a/Engine/Source/Runtime/Net/Iris/Private/Iris/ReplicationSystem/ObjectReplicationBridge.cpp b/Engine/Source/Runtime/Net/Iris/Private/Iris/ReplicationSystem/ObjectReplicationBridge.cpp
--- a/Engine/Source/Runtime/Net/Iris/Private/Iris/ReplicationSystem/ObjectReplicationBridge.cpp
+++ b/Engine/Source/Runtime/Net/Iris/Private/Iris/ReplicationSystem/ObjectReplicationBridge.cpp
@@ -1260,11 +1260,6 @@ void UObjectReplicationBridge::AddCreationDependencyLink(FNetRefHandle Parent, F
     
     if (bSuccess)
 	{
-		const FInternalNetRefIndex ParentIndex = LocalNetRefHandleManager.GetInternalIndex(Parent);
-		const FInternalNetRefIndex ChildIndex = LocalNetRefHandleManager.GetInternalIndex(Child);
-		FReplicationFiltering& Filtering = ReplicationSystemInternal->GetFiltering();
-		Filtering.NotifyAddedCreationDependency(ParentIndex, ChildIndex);
-
 		UE_LOG_OBJECTREPLICATIONBRIDGE(Verbose, TEXT("AddCreationDependencyLink Added creation dependency between child: %s and parent: %s"), *PrintObjectFromNetRefHandle(Child), *PrintObjectFromNetRefHandle(Parent));
 	}
 	else
@@ -1282,14 +1277,6 @@ void UObjectReplicationBridge::RemoveCreationDependencyLink(FNetRefHandle Parent
 	FNetRefHandleManager& LocalNetRefHandleManager = ReplicationSystemInternal->GetNetRefHandleManager();
 
 	LocalNetRefHandleManager.RemoveCreationDependency(Parent, Child);
-
-	const FInternalNetRefIndex ParentIndex = LocalNetRefHandleManager.GetInternalIndex(Parent);
-	const FInternalNetRefIndex ChildIndex = LocalNetRefHandleManager.GetInternalIndex(Child);
-	if (ParentIndex != FNetRefHandleManager::InvalidInternalIndex && ChildIndex != FNetRefHandleManager::InvalidInternalIndex)
-	{
-		FReplicationFiltering& Filtering = ReplicationSystemInternal->GetFiltering();
-		Filtering.NotifyRemovedCreationDependency(ParentIndex, ChildIndex);
-	}
 }
 
 bool UObjectReplicationBridge::WriteNetRefHandleCreationInfo(FReplicationBridgeSerializationContext& Context, FNetRefHandle Handle, UE::Net::Private::FInternalNetRefIndex InternalObjectIndex)
