Iris
- Fixed issue where out of scope objects being created and destroyed the same frame requiring flush would be destroyed using the wrong path.

#Synced-CL 51088011

diff --git a/Engine/Plugins/Runtime/ReplicationSystemTestPlugin/Source/Private/Tests/ReplicationSystem/Filtering/TestGridFilter.cpp b/Engine/Plugins/Runtime/ReplicationSystemTestPlugin/Source/Private/Tests/ReplicationSystem/Filtering/TestGridFilter.cpp
--- a/Engine/Plugins/Runtime/ReplicationSystemTestPlugin/Source/Private/Tests/ReplicationSystem/Filtering/TestGridFilter.cpp
+++ b/Engine/Plugins/Runtime/ReplicationSystemTestPlugin/Source/Private/Tests/ReplicationSystem/Filtering/TestGridFilter.cpp
@@ -388,6 +388,43 @@ UE_NET_TEST_FIXTURE(FTestGridFilterFixture, TestWorldLocCullDistance)
 	Server->DestroyObject(ServerObjectTest);
 }
 
+UE_NET_TEST_FIXTURE(FTestGridFilterFixture, TestStartQueueRpcFlushNetDormancyStop)
+{
+	// Need to set the world location update function to be able to set world locations.
+	Server->GetReplicationBridge()->SetExternalWorldLocationUpdateFunctor([&](FNetRefHandle NetHandle, const UObject* Object, FVector& OutLocation, float& OutCullDistance)
+		{
+		});
+
+	FReplicationSystemTestClient* Client = CreateClient();
+
+	// Spawn object with WorldLocation's on server
+	const UNetObjectGridFilterConfig* DefaultGridConfig = GetDefault<UNetObjectGridFilterConfig>();
+	WorldInfoToBeAssigned.Add(FObjectWorldInfo{ .Loc = FVector::ZeroVector, .CullDistance = 1500.f });
+
+	UObjectReplicationBridge::FRootObjectReplicationParams Params;
+	Params.bNeedsWorldLocationUpdate = true;
+	Params.bIsDormant = false;
+	Params.bUseClassConfigDynamicFilter = false;
+	Params.bUseExplicitDynamicFilter = true;
+	Params.ExplicitDynamicFilterName = FName("NetObjectGridWorldLocFilter");
+
+	UReplicatedTestObject* ServerObject = Server->CreateObject(Params);
+
+	// Filter out to ensure ServrObject is not referenced by any connection
+	FNetObjectGroupHandle NotReplicatedGroupHandle = Server->GetReplicationSystem()->GetNotReplicatedNetObjectGroup();
+	Server->ReplicationSystem->AddToGroup(NotReplicatedGroupHandle, ServerObject->NetRefHandle);
+
+	Server->ReplicationBridge->NetFlushDormantObject(ServerObject->NetRefHandle);
+	Server->DestroyObject(ServerObject, EEndReplicationFlags::Destroy | EEndReplicationFlags::DestroyNetHandle | EEndReplicationFlags::Flush);
+
+	// Send update
+	Server->UpdateAndSend({Client});
+
+	// Send update. Object should be destroyed on server.
+	Server->UpdateAndSend({Client});
+}
+
+
 #if 0
 UE_NET_TEST_FIXTURE(FTestGridFilterFixture, TestFilterPerformance)
 {
diff --git a/Engine/Source/Runtime/Net/Iris/Private/Iris/ReplicationSystem/NetRefHandleManager.cpp b/Engine/Source/Runtime/Net/Iris/Private/Iris/ReplicationSystem/NetRefHandleManager.cpp
--- a/Engine/Source/Runtime/Net/Iris/Private/Iris/ReplicationSystem/NetRefHandleManager.cpp
+++ b/Engine/Source/Runtime/Net/Iris/Private/Iris/ReplicationSystem/NetRefHandleManager.cpp
@@ -710,8 +710,12 @@ void FNetRefHandleManager::DestroyObjectsPendingDestroy()
 	for (int32 It = PendingDestroyInfos.Num() - 1; It >= 0; --It)
 	{
 		const FInternalNetRefIndex InternalIndex = PendingDestroyInfos[It].InternalIndex;
-		// If we have subobjects pending tear off and such then wait before destroying the parent.
-		if (ReplicatedObjectRefCount[InternalIndex] == 0 && GetSubObjects(InternalIndex).Num() <= 0)
+
+		// Do not fully destroy objects until they
+		// - are removed from GlobalScopableInternalIndices
+		// - are not referenced by any ReplicationWriter
+		// - do not have any subobjects
+		if (!GlobalScopableInternalIndices.GetBit(InternalIndex) && ReplicatedObjectRefCount[InternalIndex] == 0 && GetSubObjects(InternalIndex).Num() <= 0)
 		{
 			FreedInternalIndices.Add(InternalIndex);
 
