Fixed test fail due to external change of garbage elimination.

#rnx

#Synced-CL 51087404

diff --git a/Engine/Plugins/Runtime/ReplicationSystemTestPlugin/Source/Private/Tests/ReplicationSystem/TestReplicationOperations.cpp b/Engine/Plugins/Runtime/ReplicationSystemTestPlugin/Source/Private/Tests/ReplicationSystem/TestReplicationOperations.cpp
--- a/Engine/Plugins/Runtime/ReplicationSystemTestPlugin/Source/Private/Tests/ReplicationSystem/TestReplicationOperations.cpp
+++ b/Engine/Plugins/Runtime/ReplicationSystemTestPlugin/Source/Private/Tests/ReplicationSystem/TestReplicationOperations.cpp
@@ -13,6 +13,7 @@
 #include "Iris/ReplicationState/PropertyReplicationState.h"
 #include "Iris/ReplicationSystem/ReplicationTypes.h"
 #include "Misc/MemStack.h"
+#include "Misc/ScopeExit.h"
 
 namespace UE::Net::Private
 {
@@ -504,14 +505,18 @@ private:
 
 UE_NET_TEST_FIXTURE(FTestReplicationOperationsForObjectsFixture, StaleObjectPointerIsUpdated)
 {
-	TGuardConsoleVariable<bool> CVarOverride(TEXT("gc.GarbageEliminationEnabled"), true);
+	const bool bIsGarbageEliminationEnabled = UObjectBaseUtility::IsGarbageEliminationEnabled();
+	UObjectBaseUtility::ForceSetGarbageEliminationEnabledForTest(true);
+	ON_SCOPE_EXIT
+	{
+		UObjectBaseUtility::ForceSetGarbageEliminationEnabledForTest(bIsGarbageEliminationEnabled);
+	};
 
 	// Create two objects, one which is referencing the other. Verify reference gets updated when first object is destroyed.
 	UTestReplicatedIrisObject* Object0 = CreateObject();
 
 	UTestReplicatedIrisObject* Object1 = CreateObject();
-	constexpr uint32 ObjectReferenceComponentCount = 1;
-	Object1->AddComponents(UTestReplicatedIrisObject::FComponents{0,0,0,0,ObjectReferenceComponentCount});
+	Object1->AddComponents({.ObjectReferenceComponentCount = 1});
 	Object1->ObjectReferenceComponents[0]->RawObjectPtrRef = Object0;
 	FNetRefHandle HandleToObject1 = ReplicationBridge->BeginReplication(Object1);
 
@@ -555,14 +560,18 @@ UE_NET_TEST_FIXTURE(FTestReplicationOperationsForObjectsFixture, StaleObjectPoin
 
 UE_NET_TEST_FIXTURE(FTestReplicationOperationsForObjectsFixture, StaleObjectPointerPolledAfterGC)
 {
-	TGuardConsoleVariable<bool> CVarOverride(TEXT("gc.GarbageEliminationEnabled"), true);
+	const bool bIsGarbageEliminationEnabled = UObjectBaseUtility::IsGarbageEliminationEnabled();
+	UObjectBaseUtility::ForceSetGarbageEliminationEnabledForTest(true);
+	ON_SCOPE_EXIT
+	{
+		UObjectBaseUtility::ForceSetGarbageEliminationEnabledForTest(bIsGarbageEliminationEnabled);
+	};
 
 	// Create two objects, one which is referencing the other. Verify reference gets updated when first object is destroyed.
 	UTestReplicatedIrisObject* Object0 = CreateObject();
 
 	UTestReplicatedIrisObject* Object1 = CreateObject();
-	constexpr uint32 ObjectReferenceComponentCount = 1;
-	Object1->AddComponents(UTestReplicatedIrisObject::FComponents{0,0,0,0,ObjectReferenceComponentCount});
+	Object1->AddComponents({.ObjectReferenceComponentCount = 1});
 	Object1->ObjectReferenceComponents[0]->RawObjectPtrRef = Object0;
 	FNetRefHandle HandleToObject1 = ReplicationBridge->BeginReplication(Object1);
 
