Iris
- Implemented OptionalPropertyNetSerialzier

#rnx
#Synced-CL 50524808

diff --git a/Engine/Plugins/Runtime/ReplicationSystemTestPlugin/Source/Private/Tests/Serialization/TestOptionalPropertyNetSerializer.cpp b/Engine/Plugins/Runtime/ReplicationSystemTestPlugin/Source/Private/Tests/Serialization/TestOptionalPropertyNetSerializer.cpp
new file mode 100644
--- /dev/null
+++ b/Engine/Plugins/Runtime/ReplicationSystemTestPlugin/Source/Private/Tests/Serialization/TestOptionalPropertyNetSerializer.cpp
@@ -0,0 +1,755 @@
+// Copyright Epic Games, Inc. All Rights Reserved.
+
+#include "TestOptionalPropertyNetSerializer.h"
+#include "Tests/ReplicationSystem/ReplicationSystemServerClientTestFixture.h"
+#include "Tests/Serialization/TestNetSerializerFixture.h"
+#include "Iris/ReplicationSystem/ReplicationFragmentUtil.h"
+#include "Iris/ReplicationSystem/ReplicationSystemInternal.h"
+#include "Iris/ReplicationState/PropertyNetSerializerInfoRegistry.h"
+#include "Iris/ReplicationState/ReplicationStateDescriptorBuilder.h"
+#include "Iris/Serialization/InternalNetSerializationContext.h"
+#include "Iris/Serialization/NetReferenceCollector.h"
+#include "Iris/Serialization/NetSerializers.h"
+#include "Iris/Serialization/InternalNetSerializers.h"
+
+#include "Net/UnrealNetwork.h"
+#include "UObject/PropertyOptional.h"
+#include "Iris/ReplicationSystem/ReplicationFragmentUtil.h"
+
+namespace UE::Net::Private
+{
+using OptionalPropertyType = TOptional<FTestOptionalPropertyNetSerializer_StructWithOptionals>;
+
+class FTestOptionalPropertyNetSerializerFixture : public FReplicationSystemServerClientTestFixture
+{
+public:
+	FTestOptionalPropertyNetSerializerFixture() : FReplicationSystemServerClientTestFixture() {}
+
+protected:
+	virtual void SetUp() override;
+	virtual void TearDown() override;
+
+	// Composable operations for testing the serializer
+	void Serialize();
+	void Deserialize();
+	void SerializeDelta();
+	void DeserializeDelta();
+	void Quantize();
+	void QuantizeTwoStates();
+	void Dequantize();
+	bool IsEqual(bool bQuantized);
+	void Clone();
+	void Validate();
+	void FreeQuantizedState();
+
+	// Modify some of the optional properties
+	void SetNonDefaultInstanceState();
+
+	// Adds multiple elements, at least one uninitialized and at least two initialized with at least one modified property
+	void SetNonDefaultArrayState();
+
+	bool AreEqual(const OptionalPropertyType* Value0, const OptionalPropertyType* Value1);
+
+	const FNetSerializer* GetOptionalPropertyNetSerializer() const;
+	const FOptionalPropertyNetSerializerConfig* GetOptionalPropertyNetSerializerConfig() const;
+
+protected:
+	FNetSerializationContext NetSerializationContext;
+	FInternalNetSerializationContext InternalNetSerializationContext;
+
+	OptionalPropertyType OptionalProperty0;
+	OptionalPropertyType OptionalProperty1;
+
+	FStructNetSerializerConfig StructWithOptionalsNetSerializerConfig;
+
+	alignas(16) uint8 QuantizedBuffer[2][128];
+	alignas(16) uint8 ClonedQuantizedBuffer[2][128];
+	alignas(16) uint8 BitStreamBuffer[2048];
+
+	bool bHasQuantizedState = false;
+	bool bHasClonedQuantizedState = false;
+
+	uint32 QuantizedStateCount = 0;
+	uint32 ClonedQuantizedStateCount = 0;
+
+	FNetBitStreamWriter Writer;
+	FNetBitStreamReader Reader;
+};
+
+UE_NET_TEST_FIXTURE(FTestOptionalPropertyNetSerializerFixture, TestTOptionalCanSetInt)
+{
+	// Add a client
+	FReplicationSystemTestClient* Client = CreateClient();
+
+	// Spawn objects on server
+	UTestOptionalPropertyNetSerializer_TestObject* ServerObject = Server->CreateObject<UTestOptionalPropertyNetSerializer_TestObject>();
+
+	ServerObject->StructWithOptionals.OptInt = 1;
+
+	Server->UpdateAndSend({Client});
+
+	// Verify that object has been spawned on client
+	UTestOptionalPropertyNetSerializer_TestObject* ClientObject = Client->GetObjectAs<UTestOptionalPropertyNetSerializer_TestObject>(ServerObject->NetRefHandle);
+
+	// Verify that created server handle now also exists on client
+	UE_NET_ASSERT_NE(ClientObject, nullptr);	
+
+	UE_NET_ASSERT_TRUE(ServerObject->StructWithOptionals.OptInt == ClientObject->StructWithOptionals.OptInt);
+
+	Server->UpdateAndSend({Client});
+}
+
+UE_NET_TEST_FIXTURE(FTestOptionalPropertyNetSerializerFixture, TestTOptionalDoesNotOverwriteNotReplicatedData)
+{
+	// Add a client
+	FReplicationSystemTestClient* Client = CreateClient();
+
+	// Spawn objects on server
+	UTestOptionalPropertyNetSerializer_TestObject* ServerObject = Server->CreateObject<UTestOptionalPropertyNetSerializer_TestObject>();
+
+	ServerObject->StructWithNotFullyReplicatedOptionalStruct.OptStruct.Emplace(FTestOptionalPropertyNetSerializer_NotFullyReplicatedStruct());
+
+	Server->UpdateAndSend({Client});
+
+	// Verify that object has been spawned on client
+	UTestOptionalPropertyNetSerializer_TestObject* ClientObject = Client->GetObjectAs<UTestOptionalPropertyNetSerializer_TestObject>(ServerObject->NetRefHandle);
+
+	// Verify that created server handle now also exists on client
+	UE_NET_ASSERT_NE(ClientObject, nullptr);
+
+	
+	UE_NET_ASSERT_EQ(ServerObject->StructWithNotFullyReplicatedOptionalStruct.OptStruct->NotReplicatedInt, ClientObject->StructWithNotFullyReplicatedOptionalStruct.OptStruct->NotReplicatedInt);
+
+	// Modify replicated struct member 
+	++ServerObject->StructWithNotFullyReplicatedOptionalStruct.OptStruct->IntA;
+	// Modify not replicated struct member.
+	ServerObject->StructWithNotFullyReplicatedOptionalStruct.OptStruct->NotReplicatedInt = 7;
+
+	Server->UpdateAndSend({Client});
+
+	UE_NET_ASSERT_EQ(ServerObject->StructWithNotFullyReplicatedOptionalStruct.OptStruct->IntA, ClientObject->StructWithNotFullyReplicatedOptionalStruct.OptStruct->IntA);
+	// Verify that we leave NotReplicatedInt untouched.
+	UE_NET_ASSERT_NE(ServerObject->StructWithNotFullyReplicatedOptionalStruct.OptStruct->NotReplicatedInt, ClientObject->StructWithNotFullyReplicatedOptionalStruct.OptStruct->NotReplicatedInt);
+}
+
+UE_NET_TEST_FIXTURE(FTestOptionalPropertyNetSerializerFixture, TestTOptionalCanSetStruct)
+{
+	// Add a client
+	FReplicationSystemTestClient* Client = CreateClient();
+
+	// Spawn objects on server
+	UTestOptionalPropertyNetSerializer_TestObject* ServerObject = Server->CreateObject<UTestOptionalPropertyNetSerializer_TestObject>();
+
+	FTestOptionalPropertyNetSerializer_SimpleStruct OptStruct;
+	OptStruct.IntA = 1;
+	ServerObject->StructWithOptionals.OptStruct = OptStruct;
+
+	Server->UpdateAndSend({Client});
+
+	// Verify that object has been spawned on client
+	UTestOptionalPropertyNetSerializer_TestObject* ClientObject = Client->GetObjectAs<UTestOptionalPropertyNetSerializer_TestObject>(ServerObject->NetRefHandle);
+
+	// Verify that created server handle now also exists on client
+	UE_NET_ASSERT_NE(ClientObject, nullptr);
+
+	UE_NET_ASSERT_EQ(ServerObject->StructWithOptionals.OptStruct.IsSet(), ClientObject->StructWithOptionals.OptStruct.IsSet());
+	UE_NET_ASSERT_EQ(ServerObject->StructWithOptionals.OptStruct->IntA, ClientObject->StructWithOptionals.OptStruct->IntA);
+}
+
+// Instance tests
+UE_NET_TEST_FIXTURE(FTestOptionalPropertyNetSerializerFixture, TestQuantizeUninitialized)
+{
+	Quantize();
+}
+
+UE_NET_TEST_FIXTURE(FTestOptionalPropertyNetSerializerFixture, TestQuantizeInitialized)
+{
+	SetNonDefaultInstanceState();
+	Quantize();
+}
+
+UE_NET_TEST_FIXTURE(FTestOptionalPropertyNetSerializerFixture, TestDequantizeUninitialized)
+{
+	Quantize();
+	Dequantize();
+	UE_NET_ASSERT_TRUE(AreEqual(&OptionalProperty0, &OptionalProperty1));
+}
+
+UE_NET_TEST_FIXTURE(FTestOptionalPropertyNetSerializerFixture, TestDequantizeInitialized)
+{
+	SetNonDefaultInstanceState();
+	Quantize();
+	Dequantize();
+	UE_NET_ASSERT_TRUE(AreEqual(&OptionalProperty0, &OptionalProperty1));
+}
+
+UE_NET_TEST_FIXTURE(FTestOptionalPropertyNetSerializerFixture, TestSerializeUninitialized)
+{
+	Quantize();
+	Serialize();
+}
+
+UE_NET_TEST_FIXTURE(FTestOptionalPropertyNetSerializerFixture, TestSerializeInitialized)
+{
+	SetNonDefaultInstanceState();
+	Quantize();
+	Serialize();
+}
+
+UE_NET_TEST_FIXTURE(FTestOptionalPropertyNetSerializerFixture, TestDeserializeUninitialized)
+{
+	Quantize();
+	Serialize();
+	FreeQuantizedState();
+	Deserialize();
+}
+
+UE_NET_TEST_FIXTURE(FTestOptionalPropertyNetSerializerFixture, TestDeserializeInitialized)
+{
+	SetNonDefaultInstanceState();
+	Quantize();
+	Serialize();
+	FreeQuantizedState();
+	Deserialize();
+}
+
+UE_NET_TEST_FIXTURE(FTestOptionalPropertyNetSerializerFixture, TestDequantizeSerializedUninitializedState)
+{
+	Quantize();
+	Serialize();
+	FreeQuantizedState();
+	Deserialize();
+	Dequantize();
+	UE_NET_ASSERT_TRUE(AreEqual(&OptionalProperty0, &OptionalProperty1));
+}
+
+UE_NET_TEST_FIXTURE(FTestOptionalPropertyNetSerializerFixture, TestDequantizeSerializedInitializedState)
+{
+	SetNonDefaultInstanceState();
+	Quantize();
+	Serialize();
+	FreeQuantizedState();
+	Deserialize();
+	Dequantize();
+	UE_NET_ASSERT_TRUE(AreEqual(&OptionalProperty0, &OptionalProperty1));
+}
+
+UE_NET_TEST_FIXTURE(FTestOptionalPropertyNetSerializerFixture, TestSerializeDeltaEqualUnsetInstances)
+{
+	QuantizeTwoStates();
+	SerializeDelta();
+}
+
+UE_NET_TEST_FIXTURE(FTestOptionalPropertyNetSerializerFixture, TestSerializeDeltaNonEqualStates)
+{
+	OptionalProperty0.Emplace(FTestOptionalPropertyNetSerializer_StructWithOptionals());
+	OptionalProperty0->IntValue += 4711U;
+
+	OptionalProperty1.Emplace();
+	OptionalProperty1->IntValue += 711U;
+
+	QuantizeTwoStates();
+	SerializeDelta();
+}
+
+UE_NET_TEST_FIXTURE(FTestOptionalPropertyNetSerializerFixture, TestDeserializeDeltaEqualUnsetInstances)
+{
+	QuantizeTwoStates();
+	SerializeDelta();
+	DeserializeDelta();
+}
+
+UE_NET_TEST_FIXTURE(FTestOptionalPropertyNetSerializerFixture, TestDeserializeDeltaNonEqualStates)
+{
+	OptionalProperty0.Emplace(FTestOptionalPropertyNetSerializer_StructWithOptionals());
+	OptionalProperty0->IntValue += 4711U;
+
+	OptionalProperty1.Emplace(FTestOptionalPropertyNetSerializer_StructWithOptionals());
+	OptionalProperty1->IntValue += 711U;
+
+	QuantizeTwoStates();
+	SerializeDelta();
+	DeserializeDelta();
+}
+
+UE_NET_TEST_FIXTURE(FTestOptionalPropertyNetSerializerFixture, TestDequantizeDeltaSerializedUnsetInstances)
+{
+	QuantizeTwoStates();
+	SerializeDelta();
+	DeserializeDelta();
+	Dequantize();
+
+	UE_NET_ASSERT_TRUE(AreEqual(&OptionalProperty0, &OptionalProperty1));
+}
+
+UE_NET_TEST_FIXTURE(FTestOptionalPropertyNetSerializerFixture, TestDequantizeDeltaSerializedNonEqualStates)
+{
+	OptionalProperty0.Emplace(FTestOptionalPropertyNetSerializer_StructWithOptionals());
+	OptionalProperty0->IntValue += 4711U;
+
+	OptionalProperty1.Emplace(FTestOptionalPropertyNetSerializer_StructWithOptionals());
+	OptionalProperty1->IntValue += 711U;
+
+	QuantizeTwoStates();
+	SerializeDelta();
+	DeserializeDelta();
+	Dequantize();
+
+	UE_NET_ASSERT_TRUE(AreEqual(&OptionalProperty0, &OptionalProperty1));
+}
+
+UE_NET_TEST_FIXTURE(FTestOptionalPropertyNetSerializerFixture, TestCollectReferencesUninitialized)
+{
+	Quantize();
+
+	FNetReferenceCollector Collector;
+
+	FNetCollectReferencesArgs Args = {};
+	Args.NetSerializerConfig = NetSerializerConfigParam(GetOptionalPropertyNetSerializerConfig());
+	Args.Source = NetSerializerValuePointer(&QuantizedBuffer[0]);
+	Args.Collector = NetSerializerValuePointer(&Collector);
+	GetOptionalPropertyNetSerializer()->CollectNetReferences(NetSerializationContext, Args);
+
+	UE_NET_ASSERT_EQ(Collector.GetCollectedReferences().Num(), 0);
+}
+
+UE_NET_TEST_FIXTURE(FTestOptionalPropertyNetSerializerFixture, TestCollectReferencesNoRef)
+{
+	OptionalProperty0.Emplace(FTestOptionalPropertyNetSerializer_StructWithOptionals());
+
+	Quantize();
+
+	FNetReferenceCollector Collector(ENetReferenceCollectorTraits::IncludeInvalidReferences);
+
+	FNetCollectReferencesArgs Args = {};
+	Args.NetSerializerConfig = NetSerializerConfigParam(GetOptionalPropertyNetSerializerConfig());
+	Args.Source = NetSerializerValuePointer(&QuantizedBuffer[0]);
+	Args.Collector = NetSerializerValuePointer(&Collector);
+
+	GetOptionalPropertyNetSerializer()->CollectNetReferences(NetSerializationContext, Args);
+
+	UE_NET_ASSERT_EQ(Collector.GetCollectedReferences().Num(), 1);
+}
+
+UE_NET_TEST_FIXTURE(FTestOptionalPropertyNetSerializerFixture, TestCollectReferencesWithRef)
+{
+	OptionalProperty0.Emplace(FTestOptionalPropertyNetSerializer_StructWithOptionals());
+	OptionalProperty0->ObjectRef = FTestOptionalPropertyNetSerializer_StructWithOptionals::StaticStruct();
+
+	Quantize();
+
+	FNetReferenceCollector Collector(ENetReferenceCollectorTraits::IncludeInvalidReferences);
+
+	FNetCollectReferencesArgs Args = {};
+	Args.NetSerializerConfig = NetSerializerConfigParam(GetOptionalPropertyNetSerializerConfig());
+	Args.Source = NetSerializerValuePointer(&QuantizedBuffer[0]);
+	Args.Collector = NetSerializerValuePointer(&Collector);
+
+	GetOptionalPropertyNetSerializer()->CollectNetReferences(NetSerializationContext, Args);
+
+	UE_NET_ASSERT_GE(Collector.GetCollectedReferences().Num(), 1);
+}
+
+UE_NET_TEST_FIXTURE(FTestOptionalPropertyNetSerializerFixture, TestCollectReferencesWithRefAndOptRef)
+{
+	OptionalProperty0.Emplace(FTestOptionalPropertyNetSerializer_StructWithOptionals());
+	OptionalProperty0->ObjectRef = FTestOptionalPropertyNetSerializer_StructWithOptionals::StaticStruct();
+
+	OptionalProperty0->OptObjectRef = FTestOptionalPropertyNetSerializer_StructWithOptional::StaticStruct();
+
+	Quantize();
+
+	FNetReferenceCollector Collector(ENetReferenceCollectorTraits::IncludeInvalidReferences);
+
+	FNetCollectReferencesArgs Args = {};
+	Args.NetSerializerConfig = NetSerializerConfigParam(GetOptionalPropertyNetSerializerConfig());
+	Args.Source = NetSerializerValuePointer(&QuantizedBuffer[0]);
+	Args.Collector = NetSerializerValuePointer(&Collector);
+
+	GetOptionalPropertyNetSerializer()->CollectNetReferences(NetSerializationContext, Args);
+
+	UE_NET_ASSERT_GE(Collector.GetCollectedReferences().Num(), 2);
+}
+
+UE_NET_TEST_FIXTURE(FTestOptionalPropertyNetSerializerFixture, TestIsEqualExternal)
+{
+	constexpr bool bUseQuantizedState = false;
+
+	// Default state compared to default state
+	OptionalProperty0.Reset();
+	OptionalProperty1.Reset();
+	UE_NET_ASSERT_TRUE(IsEqual(bUseQuantizedState));
+
+	// Non-default state compared to default state
+	SetNonDefaultInstanceState();
+	UE_NET_ASSERT_FALSE(IsEqual(bUseQuantizedState));
+
+	// Non-default state compared to non-default state
+	OptionalProperty1 = OptionalProperty0;
+	UE_NET_ASSERT_TRUE(IsEqual(bUseQuantizedState));
+}
+
+UE_NET_TEST_FIXTURE(FTestOptionalPropertyNetSerializerFixture, TestIsEqualQuantized)
+{
+	constexpr bool bUseQuantizedState = true;
+
+	// Default state compared to default state
+	OptionalProperty0.Reset();
+	Quantize();
+	Clone();
+	UE_NET_ASSERT_TRUE(IsEqual(bUseQuantizedState));
+
+	// Non-default state compared to default state
+	SetNonDefaultInstanceState();
+	Quantize();
+	UE_NET_ASSERT_FALSE(IsEqual(bUseQuantizedState));
+
+	// Non-default state compared to non-default state
+	Clone();
+	UE_NET_ASSERT_TRUE(IsEqual(bUseQuantizedState));
+}
+
+UE_NET_TEST_FIXTURE(FTestOptionalPropertyNetSerializerFixture, TestValidate)
+{
+	Validate();
+}
+
+UE_NET_TEST_FIXTURE(FTestOptionalPropertyNetSerializerFixture, TestDequantizeSerializedArrayWithOptionals)
+{
+	SetNonDefaultInstanceState();
+
+	// Add some array entries
+	OptionalProperty0->ArrayWithOptionals.SetNum(10);
+	
+	OptionalProperty0->ArrayWithOptionals[0] = 0;
+	OptionalProperty0->ArrayWithOptionals[3] = 3;
+	OptionalProperty0->ArrayWithOptionals[9] = 9;
+
+	Quantize();
+	Serialize();
+	FreeQuantizedState();
+	Deserialize();
+	Dequantize();
+	
+	UE_NET_ASSERT_TRUE(AreEqual(&OptionalProperty0, &OptionalProperty1));
+}
+
+
+}
+
+UTestOptionalPropertyNetSerializer_TestObject::UTestOptionalPropertyNetSerializer_TestObject()
+: UReplicatedTestObject()
+{
+}
+
+void UTestOptionalPropertyNetSerializer_TestObject::GetLifetimeReplicatedProps(TArray<class FLifetimeProperty>& OutLifetimeProps) const
+{
+	FDoRepLifetimeParams Params;
+	Params.bIsPushBased = false;
+
+	DOREPLIFETIME_WITH_PARAMS_FAST(UTestOptionalPropertyNetSerializer_TestObject, StructWithOptionals, Params);
+	DOREPLIFETIME_WITH_PARAMS_FAST(UTestOptionalPropertyNetSerializer_TestObject, StructWithNotFullyReplicatedOptionalStruct, Params);
+}
+
+void UTestOptionalPropertyNetSerializer_TestObject::RegisterReplicationFragments(UE::Net::FFragmentRegistrationContext& Context, UE::Net::EFragmentRegistrationFlags RegistrationFlags)
+{
+	UE::Net::FReplicationFragmentUtil::CreateAndRegisterFragmentsForObject(this, Context, RegistrationFlags);
+}
+
+namespace UE::Net::Private
+{
+
+// Fixture implementation
+void FTestOptionalPropertyNetSerializerFixture::SetUp()
+{
+	FReplicationSystemServerClientTestFixture::SetUp();
+
+	// Init default serialization context
+	InternalNetSerializationContext.ReplicationSystem = Server->ReplicationSystem;
+
+	FInternalNetSerializationContext TempInternalNetSerializationContext;
+	FInternalNetSerializationContext::FInitParameters TempInternalNetSerializationContextInitParams;
+	TempInternalNetSerializationContextInitParams.ReplicationSystem = Server->ReplicationSystem;
+	TempInternalNetSerializationContextInitParams.ObjectResolveContext.RemoteNetTokenStoreState = Server->ReplicationSystem->GetNetTokenStore()->GetLocalNetTokenStoreState();
+	TempInternalNetSerializationContext.Init(TempInternalNetSerializationContextInitParams);
+
+	InternalNetSerializationContext = MoveTemp(TempInternalNetSerializationContext);
+	NetSerializationContext.SetInternalContext(&InternalNetSerializationContext);
+
+	FMemory::Memzero(QuantizedBuffer, sizeof(QuantizedBuffer));
+
+	bHasQuantizedState = false;
+	bHasClonedQuantizedState = false;
+
+	if (!StructWithOptionalsNetSerializerConfig.StateDescriptor.IsValid())
+	{
+		FReplicationStateDescriptorBuilder::FParameters Params;
+		StructWithOptionalsNetSerializerConfig.StateDescriptor = FReplicationStateDescriptorBuilder::CreateDescriptorForStruct(FTestOptionalPropertyNetSerializer_StructWithOptional::StaticStruct(), Params);
+	}
+}
+
+void FTestOptionalPropertyNetSerializerFixture::TearDown()
+{
+	OptionalProperty0.Reset();
+	OptionalProperty1.Reset();
+
+	FreeQuantizedState();
+
+	FReplicationSystemServerClientTestFixture::TearDown();
+}
+
+void FTestOptionalPropertyNetSerializerFixture::Serialize()
+{
+	// Must have run quantize before this
+	UE_NET_ASSERT_TRUE(bHasQuantizedState);
+
+	// Serialize data
+	Writer.InitBytes(BitStreamBuffer, sizeof(BitStreamBuffer));
+	FNetSerializationContext Context(&Writer);
+	Context.SetInternalContext(NetSerializationContext.GetInternalContext());
+
+	FNetSerializeArgs Args = {};
+	Args.NetSerializerConfig = NetSerializerConfigParam(GetOptionalPropertyNetSerializerConfig());
+	Args.Source = NetSerializerValuePointer(&QuantizedBuffer[0]);
+	GetOptionalPropertyNetSerializer()->Serialize(Context, Args);
+
+	Writer.CommitWrites();
+
+	UE_NET_ASSERT_FALSE(Context.HasError());
+	UE_NET_ASSERT_GT(Writer.GetPosBits(), 0U);
+}
+
+void FTestOptionalPropertyNetSerializerFixture::Deserialize()
+{
+	// Check pre-conditions
+	UE_NET_ASSERT_FALSE(bHasQuantizedState);
+	UE_NET_ASSERT_GT(Writer.GetPosBytes(), 0U);
+
+	Reader.InitBits(BitStreamBuffer, Writer.GetPosBits());
+
+	FNetSerializationContext Context(&Reader);
+	Context.SetInternalContext(NetSerializationContext.GetInternalContext());
+
+	FNetDeserializeArgs Args = {};
+	Args.NetSerializerConfig = NetSerializerConfigParam(GetOptionalPropertyNetSerializerConfig());
+	Args.Target = NetSerializerValuePointer(&QuantizedBuffer[0]);
+	GetOptionalPropertyNetSerializer()->Deserialize(Context, Args);
+
+	bHasQuantizedState = true;
+
+	UE_NET_ASSERT_FALSE(Context.HasErrorOrOverflow());
+	UE_NET_ASSERT_GT(Reader.GetPosBits(), 0U);
+}
+
+void FTestOptionalPropertyNetSerializerFixture::SerializeDelta()
+{
+	// Check pre-conditions
+	UE_NET_ASSERT_TRUE(bHasQuantizedState);
+	UE_NET_ASSERT_EQ(QuantizedStateCount, 2U);
+
+	// Serialize data
+	Writer.InitBytes(BitStreamBuffer, sizeof(BitStreamBuffer));
+	FNetSerializationContext Context(&Writer);
+	Context.SetInternalContext(NetSerializationContext.GetInternalContext());
+
+	FNetSerializeDeltaArgs Args = {};
+	Args.NetSerializerConfig = NetSerializerConfigParam(GetOptionalPropertyNetSerializerConfig());
+	Args.Source = NetSerializerValuePointer(&QuantizedBuffer[0]);
+	Args.Prev = NetSerializerValuePointer(&QuantizedBuffer[1]);
+	GetOptionalPropertyNetSerializer()->SerializeDelta(Context, Args);
+
+	Writer.CommitWrites();
+
+	UE_NET_ASSERT_FALSE(Context.HasErrorOrOverflow());
+	UE_NET_ASSERT_GT(Writer.GetPosBits(), 0U);
+}
+
+void FTestOptionalPropertyNetSerializerFixture::DeserializeDelta()
+{
+	// Check pre-conditions
+	UE_NET_ASSERT_GT(Writer.GetPosBytes(), 0U);
+
+	Reader.InitBits(BitStreamBuffer, Writer.GetPosBits());
+
+	FNetSerializationContext Context(&Reader);
+	Context.SetInternalContext(NetSerializationContext.GetInternalContext());
+
+	FNetDeserializeDeltaArgs Args = {};
+	Args.NetSerializerConfig = NetSerializerConfigParam(GetOptionalPropertyNetSerializerConfig());
+	Args.Target = NetSerializerValuePointer(&QuantizedBuffer[0]);
+	Args.Prev = NetSerializerValuePointer(&QuantizedBuffer[1]);
+	GetOptionalPropertyNetSerializer()->DeserializeDelta(Context, Args);
+
+	bHasQuantizedState = true;
+	QuantizedStateCount = 1;
+
+	UE_NET_ASSERT_FALSE(Context.HasErrorOrOverflow());
+	UE_NET_ASSERT_GT(Reader.GetPosBits(), 0U);
+}
+
+void FTestOptionalPropertyNetSerializerFixture::Quantize()
+{
+	FNetQuantizeArgs Args = {};
+	Args.NetSerializerConfig = NetSerializerConfigParam(GetOptionalPropertyNetSerializerConfig());
+	Args.Target = NetSerializerValuePointer(&QuantizedBuffer[0]);
+	Args.Source = NetSerializerValuePointer(&OptionalProperty0);
+	GetOptionalPropertyNetSerializer()->Quantize(NetSerializationContext, Args);
+
+	bHasQuantizedState = true;
+	QuantizedStateCount = 1;
+
+	UE_NET_ASSERT_FALSE(NetSerializationContext.HasError());
+}
+
+void FTestOptionalPropertyNetSerializerFixture::QuantizeTwoStates()
+{
+	Quantize();
+
+	FNetQuantizeArgs Args = {};
+	Args.NetSerializerConfig = NetSerializerConfigParam(GetOptionalPropertyNetSerializerConfig());
+	Args.Target = NetSerializerValuePointer(&QuantizedBuffer[1]);
+	Args.Source = NetSerializerValuePointer(&OptionalProperty1);
+	GetOptionalPropertyNetSerializer()->Quantize(NetSerializationContext, Args);
+
+	bHasQuantizedState = true;
+	QuantizedStateCount = 2;
+
+	UE_NET_ASSERT_FALSE(NetSerializationContext.HasError());
+}
+
+void FTestOptionalPropertyNetSerializerFixture::Clone()
+{
+	// Check pre-conditions
+	UE_NET_ASSERT_TRUE(bHasQuantizedState);
+
+	FMemory::Memcpy(ClonedQuantizedBuffer[0], QuantizedBuffer[0], sizeof(QuantizedBuffer[0]));
+
+	FNetCloneDynamicStateArgs Args = {};
+	Args.Source = NetSerializerValuePointer(&QuantizedBuffer[0]);
+	Args.Target = NetSerializerValuePointer(&ClonedQuantizedBuffer[0]);
+	Args.NetSerializerConfig = NetSerializerConfigParam(GetOptionalPropertyNetSerializerConfig());
+	GetOptionalPropertyNetSerializer()->CloneDynamicState(NetSerializationContext, Args);
+
+	bHasClonedQuantizedState = true;
+	ClonedQuantizedStateCount = 1;
+}
+
+void FTestOptionalPropertyNetSerializerFixture::FreeQuantizedState()
+{
+	FNetFreeDynamicStateArgs Args = {};
+	Args.NetSerializerConfig = NetSerializerConfigParam(GetOptionalPropertyNetSerializerConfig());
+
+	if (bHasQuantizedState)
+	{
+		for (uint32 StateIt = 0, StateEndIt = QuantizedStateCount; StateIt != StateEndIt; ++StateIt)
+		{
+			Args.Source = NetSerializerValuePointer(&QuantizedBuffer[StateIt]);
+			GetOptionalPropertyNetSerializer()->FreeDynamicState(NetSerializationContext, Args);
+
+			FMemory::Memzero(&QuantizedBuffer[StateIt], sizeof(QuantizedBuffer[StateIt]));
+		}
+		bHasQuantizedState = false;
+	}
+
+	if (bHasClonedQuantizedState)
+	{
+		for (uint32 StateIt = 0, StateEndIt = ClonedQuantizedStateCount; StateIt != StateEndIt; ++StateIt)
+		{
+			Args.Source = NetSerializerValuePointer(&ClonedQuantizedBuffer[StateIt]);
+			GetOptionalPropertyNetSerializer()->FreeDynamicState(NetSerializationContext, Args);
+
+			FMemory::Memzero(&ClonedQuantizedBuffer[StateIt], sizeof(ClonedQuantizedBuffer[StateIt]));
+		}
+		bHasClonedQuantizedState = false;
+	}
+}
+
+void FTestOptionalPropertyNetSerializerFixture::Dequantize()
+{
+	UE_NET_ASSERT_TRUE(bHasQuantizedState);
+
+	FNetDequantizeArgs Args = {};
+	Args.NetSerializerConfig = NetSerializerConfigParam(GetOptionalPropertyNetSerializerConfig());
+	Args.Source = NetSerializerValuePointer(&QuantizedBuffer[0]);
+	Args.Target = NetSerializerValuePointer(&OptionalProperty1);
+	GetOptionalPropertyNetSerializer()->Dequantize(NetSerializationContext, Args);
+}
+
+bool FTestOptionalPropertyNetSerializerFixture::IsEqual(bool bQuantized)
+{
+	if (bQuantized)
+	{
+		UE_NET_EXPECT_TRUE(bHasQuantizedState);
+		if (!bHasQuantizedState)
+		{
+			return false;
+		}
+
+		UE_NET_EXPECT_TRUE(bHasClonedQuantizedState);
+		if (!bHasClonedQuantizedState)
+		{
+			return false;
+		}
+	}
+
+	FNetIsEqualArgs Args = {};
+	if (bQuantized)
+	{
+		Args.Source0 = NetSerializerValuePointer(&QuantizedBuffer[0]);
+		Args.Source1 = NetSerializerValuePointer(&ClonedQuantizedBuffer[0]);
+	}
+	else
+	{
+		Args.Source0 = NetSerializerValuePointer(&OptionalProperty0);
+		Args.Source1 = NetSerializerValuePointer(&OptionalProperty1);
+	}
+	Args.bStateIsQuantized = bQuantized;
+	Args.NetSerializerConfig = NetSerializerConfigParam(GetOptionalPropertyNetSerializerConfig());
+	return GetOptionalPropertyNetSerializer()->IsEqual(NetSerializationContext, Args);
+}
+
+void FTestOptionalPropertyNetSerializerFixture::Validate()
+{
+	FNetValidateArgs Args = {};
+	Args.NetSerializerConfig = NetSerializerConfigParam(GetOptionalPropertyNetSerializerConfig());
+	Args.Source = NetSerializerValuePointer(&OptionalProperty0);
+
+	GetOptionalPropertyNetSerializer()->Validate(NetSerializationContext, Args);
+}
+
+void FTestOptionalPropertyNetSerializerFixture::SetNonDefaultInstanceState()
+{
+	OptionalProperty0.Emplace(FTestOptionalPropertyNetSerializer_StructWithOptionals());
+	OptionalProperty0->IntValue += 13;
+}
+
+bool FTestOptionalPropertyNetSerializerFixture::AreEqual(const OptionalPropertyType* Value0, const OptionalPropertyType* Value1)
+{
+	const FOptionalPropertyNetSerializerConfig* Config =  GetOptionalPropertyNetSerializerConfig();
+
+	return Config->Property->Identical(Value0, Value1, 0);
+}
+
+const FNetSerializer* FTestOptionalPropertyNetSerializerFixture::GetOptionalPropertyNetSerializer() const
+{
+	if (const FReplicationStateDescriptor* Desc = StructWithOptionalsNetSerializerConfig.StateDescriptor.GetReference())
+	{
+		return Desc->MemberSerializerDescriptors[0].Serializer;
+	}
+
+	return nullptr;
+}
+
+const FOptionalPropertyNetSerializerConfig* FTestOptionalPropertyNetSerializerFixture::GetOptionalPropertyNetSerializerConfig() const
+{
+	if (const FReplicationStateDescriptor* Desc = StructWithOptionalsNetSerializerConfig.StateDescriptor.GetReference())
+	{
+		return static_cast<const FOptionalPropertyNetSerializerConfig*>(Desc->MemberSerializerDescriptors[0].SerializerConfig);
+	}
+
+	return nullptr;
+}
+
+}
+
+
diff --git a/Engine/Plugins/Runtime/ReplicationSystemTestPlugin/Source/Private/Tests/Serialization/TestOptionalPropertyNetSerializer.h b/Engine/Plugins/Runtime/ReplicationSystemTestPlugin/Source/Private/Tests/Serialization/TestOptionalPropertyNetSerializer.h
new file mode 100644
--- /dev/null
+++ b/Engine/Plugins/Runtime/ReplicationSystemTestPlugin/Source/Private/Tests/Serialization/TestOptionalPropertyNetSerializer.h
@@ -0,0 +1,114 @@
+// Copyright Epic Games, Inc. All Rights Reserved.
+
+#pragma once
+
+#include "Containers/Array.h"
+#include "Misc/Optional.h"
+#include "UObject/ObjectMacros.h"
+#include "Tests/ReplicationSystem/ReplicatedTestObject.h"
+#include "TestOptionalPropertyNetSerializer.generated.h"
+
+USTRUCT()
+struct FTestOptionalPropertyNetSerializer_SimpleStruct
+{
+	GENERATED_BODY()
+
+public:
+	UPROPERTY(transient)
+	int32 IntA = 0;
+};
+
+USTRUCT()
+struct FTestOptionalPropertyNetSerializer_NotFullyReplicatedStruct
+{
+	GENERATED_BODY()
+
+public:
+	UPROPERTY(transient)
+	int32 IntA = 0;
+
+	UPROPERTY(transient, NotReplicated)
+	int32 NotReplicatedInt = 0;
+};
+
+USTRUCT()
+struct FTestOptionalPropertyNetSerializer_StructWithObjectReference
+{
+	GENERATED_BODY()
+
+public:
+	UPROPERTY(transient)
+	TObjectPtr<UObject> ObjectRef = nullptr;
+};
+
+USTRUCT()
+struct FTestOptionalPropertyNetSerializer_StructWithOptionals
+{
+	GENERATED_BODY()
+
+public:
+	UPROPERTY(transient)
+	int32 IntValue = 0;
+
+	UPROPERTY(transient)
+	TArray<TOptional<int32>> ArrayWithOptionals;
+
+	UPROPERTY(transient)
+	TObjectPtr<UObject> ObjectRef = nullptr;
+
+	UPROPERTY(transient)
+	TOptional<int32> OptInt;
+
+	UPROPERTY(transient)
+	TOptional<TObjectPtr<UObject>> OptObjectRef;
+
+	UPROPERTY(transient)
+	TOptional<TArray<int32>> OptIntArray;
+
+	UPROPERTY(transient)
+	TOptional<TArray<TOptional<int32>>> OptNestOptIntArray;
+
+	UPROPERTY(transient)
+	TOptional<FTestOptionalPropertyNetSerializer_SimpleStruct> OptStruct;
+};
+
+// Struct that we use to grab NetSerializer and config for TOptional<FTestOptionalPropertyNetSerializer_StructWithOptionals>
+USTRUCT()
+struct FTestOptionalPropertyNetSerializer_StructWithOptional
+{
+	GENERATED_BODY()
+
+public:
+	UPROPERTY(transient)
+	TOptional<FTestOptionalPropertyNetSerializer_StructWithOptionals> OptInt;
+};
+
+USTRUCT()
+struct FTestOptionalPropertyNetSerializer_StructWithNotFullyReplicatedOptionalStruct
+{
+	GENERATED_BODY()
+
+public:
+	UPROPERTY(transient)
+	TOptional<FTestOptionalPropertyNetSerializer_NotFullyReplicatedStruct> OptStruct;
+};
+
+UCLASS()
+class UTestOptionalPropertyNetSerializer_TestObject : public UReplicatedTestObject
+{
+	GENERATED_BODY()
+public:
+	UTestOptionalPropertyNetSerializer_TestObject();
+
+	void RegisterReplicationFragments(UE::Net::FFragmentRegistrationContext& Fragments, UE::Net::EFragmentRegistrationFlags RegistrationFlags) override;
+
+	UPROPERTY(Transient, Replicated)
+	FTestOptionalPropertyNetSerializer_StructWithOptionals StructWithOptionals;
+
+	UPROPERTY(Transient, Replicated)
+	FTestOptionalPropertyNetSerializer_StructWithNotFullyReplicatedOptionalStruct StructWithNotFullyReplicatedOptionalStruct;
+
+// Not supported at the moment due to UHT.
+//	UPROPERTY(Transient, Replicated)
+//	TOptional<Int32> OptInt;
+};
diff --git a/Engine/Source/Runtime/Net/Iris/Private/Iris/ReplicationState/PropertyNetSerializerInfoRegistry.cpp b/Engine/Source/Runtime/Net/Iris/Private/Iris/ReplicationState/PropertyNetSerializerInfoRegistry.cpp
--- a/Engine/Source/Runtime/Net/Iris/Private/Iris/ReplicationState/PropertyNetSerializerInfoRegistry.cpp
+++ b/Engine/Source/Runtime/Net/Iris/Private/Iris/ReplicationState/PropertyNetSerializerInfoRegistry.cpp
@@ -84,6 +84,16 @@ struct FDefaultArrayPropertyNetSerializerInfo : public FPropertyNetSerializerInf
 };
 UE_NET_IMPLEMENT_NETSERIALIZER_INFO(FDefaultArrayPropertyNetSerializerInfo);
 
+struct FDefaultOptionalPropertyNetSerializerInfo : public FPropertyNetSerializerInfo
+{
+	virtual const FFieldClass* GetPropertyTypeClass() const override { return FOptionalProperty::StaticClass(); }
+	virtual bool IsSupported(const FProperty* Property) const override { return FPropertyNetSerializerInfoRegistry::FindSerializerInfo(CastFieldChecked<FOptionalProperty>(Property)->GetValueProperty()) != nullptr; }
+	virtual const FNetSerializer* GetNetSerializer(const FProperty* Property) const override { return &UE_NET_GET_SERIALIZER(FOptionalPropertyNetSerializer); }
+	virtual bool CanUseDefaultConfig(const FProperty* Property) const override { return false; }
+	virtual FNetSerializerConfig* BuildNetSerializerConfig(void* NetSerializerConfigBuffer, const FProperty* Property) const override { return nullptr; }
+};
+UE_NET_IMPLEMENT_NETSERIALIZER_INFO(FDefaultOptionalPropertyNetSerializerInfo);
+
 /** Last resort for properties that have no specialized serializer. An example is TextProperty. */
 FLastResortPropertyNetSerializerInfo::FLastResortPropertyNetSerializerInfo(const FLastResortPropertyNetSerializerUserOptions& UserOptions)
 : FPropertyNetSerializerInfo()
@@ -176,6 +186,16 @@ const FPropertyNetSerializerInfo* FPropertyNetSerializerInfoRegistry::FindSerial
 		}
 	}
 
+	// Special handling for optional properties required
+	if (Property->IsA(FOptionalProperty::StaticClass()))
+	{
+		const FPropertyNetSerializerInfo& Info = GetPropertyNetSerializerInfo_FDefaultOptionalPropertyNetSerializerInfo();
+		if (Info.IsSupported(Property))
+		{
+			return &Info;
+		}
+	}
+
 	// Last resort!
 	return &GetPropertyNetSerializerInfo_FLastResortPropertyNetSerializerInfo();
 }
diff --git a/Engine/Source/Runtime/Net/Iris/Private/Iris/ReplicationState/ReplicationStateDescriptorBuilder.cpp b/Engine/Source/Runtime/Net/Iris/Private/Iris/ReplicationState/ReplicationStateDescriptorBuilder.cpp
--- a/Engine/Source/Runtime/Net/Iris/Private/Iris/ReplicationState/ReplicationStateDescriptorBuilder.cpp
+++ b/Engine/Source/Runtime/Net/Iris/Private/Iris/ReplicationState/ReplicationStateDescriptorBuilder.cpp
@@ -260,6 +260,7 @@ private:
 		FSizeAndAlignment InternalSizeAndAlignment;					// Size and alignment of quantized member data
 		FSizeAndAlignment SerializerConfigSizeAndAlignment;			// Size and alignment for SerializerConfig
 		uint32 bUseSerializerDefaultConfig : 1;						// Whether the member can use the serializer's default config
+
 		uint32 bIsStruct : 1;
 		uint32 bIsDynamicArray : 1;
 		uint32 bHasCustomObjectReference : 1;
@@ -379,6 +380,7 @@ private:
 
 	TRefCountPtr<const FReplicationStateDescriptor> GetDescriptorForStructProperty(const FMemberProperty& MemberProperty, FReplicationStateDescriptorRegistry* DescriptorRegistry) const;
 	TRefCountPtr<const FReplicationStateDescriptor> GetDescriptorForArrayProperty(const FProperty* Property, FReplicationStateDescriptorRegistry* DescriptorRegistry) const;
+	TRefCountPtr<const FReplicationStateDescriptor> GetDescriptorForOptionalProperty(const FProperty* Property, FReplicationStateDescriptorRegistry* DescriptorRegistry) const;
 	TRefCountPtr<const FReplicationStateDescriptor> GetOrCreateDescriptorForFunction(const UFunction* Function, FReplicationStateDescriptorRegistry* DescriptorRegistry) const;
 	TRefCountPtr<const FReplicationStateDescriptor> CreateDescriptorForProperty(const FString& DescriptorName, const FProperty* Property, const FReplicationStateDescriptorBuilder::FParameters& Parameters) const;
 
@@ -562,6 +564,20 @@ TRefCountPtr<const FReplicationStateDescriptor> FPropertyReplicationStateDescrip
 	return CreateDescriptorForProperty(DescriptorName, InnerProperty, Params);
 }
 
+TRefCountPtr<const FReplicationStateDescriptor> FPropertyReplicationStateDescriptorBuilder::GetDescriptorForOptionalProperty(const FProperty* Property, FReplicationStateDescriptorRegistry* DescriptorRegistry) const
+{
+	FReplicationStateDescriptorBuilder::FParameters Params;
+	Params.DescriptorRegistry = DescriptorRegistry;
+
+	const FOptionalProperty* OptionalProperty = CastFieldChecked<FOptionalProperty>(Property);
+	const FProperty* ValueProperty = OptionalProperty->GetValueProperty();
+
+	FString DescriptorName;
+	GetPropertyPathName(OptionalProperty, DescriptorName);
+
+	return CreateDescriptorForProperty(DescriptorName, ValueProperty, Params);
+}
+
 TRefCountPtr<const FReplicationStateDescriptor> FPropertyReplicationStateDescriptorBuilder::GetOrCreateDescriptorForFunction(const UFunction* Function, FReplicationStateDescriptorRegistry* DescriptorRegistry) const
 {
 	FReplicationStateDescriptorBuilder::FParameters Params;
@@ -717,6 +733,35 @@ void FPropertyReplicationStateDescriptorBuilder::BuildMemberCache(FBuilderContex
 
 			CurrentCacheEntry->bIsDynamicArray = 1;
 		}
+		else if (IsOptionalPropertyNetSerializer(Serializer))
+		{
+			bSomeMembersAreProperties = true;
+
+			// Build descriptor for inner property
+			CurrentCacheEntry->Descriptor = GetDescriptorForOptionalProperty(Member.Property, DescriptorRegistry);
+
+			const FReplicationStateDescriptor* Descriptor = CurrentCacheEntry->Descriptor.GetReference();
+
+			// Propagate some traits. Doing this earlier would be quite expensive.
+			if (EnumHasAnyFlags(Descriptor->Traits, EReplicationStateTraits::HasObjectReference))
+			{
+				Member.Traits |= EMemberPropertyTraits::HasObjectReference;
+			}
+
+			if (EnumHasAnyFlags(Descriptor->Traits, EReplicationStateTraits::UseSerializerIsEqual))
+			{
+				Member.Traits |= EMemberPropertyTraits::UseSerializerIsEqual;
+			}
+
+			// We propagate this up so that a optional with a struct that is not fully replicated will also have this set
+			if (!EnumHasAnyFlags(Descriptor->Traits, EReplicationStateTraits::AllMembersAreReplicated))
+			{
+				bAllMembersAreReplicated = false;
+			}
+
+			CurrentCacheEntry->ExternalSizeAndAlignment = { (SIZE_T)Member.Property->GetElementSize(), (SIZE_T)Member.Property->GetMinAlignment() };
+			CurrentCacheEntry->InternalSizeAndAlignment = { Serializer->QuantizedTypeSize, Serializer->QuantizedTypeAlignment };
+		}
 		else if (Member.Property)
 		{
 			bSomeMembersAreProperties = true;
@@ -1474,6 +1519,13 @@ void FPropertyReplicationStateDescriptorBuilder::BuildMemberSerializerConfigs(ui
 				SerializerConfig->StateDescriptor = MemberCacheEntry->Descriptor;
 				MemberSerializerDescriptor->SerializerConfig = SerializerConfig;
 			}
+			else if (IsOptionalPropertyNetSerializer(MemberCacheEntry->Serializer))
+			{
+				FOptionalPropertyNetSerializerConfig* SerializerConfig = new (SerializerConfigBuffer) FOptionalPropertyNetSerializerConfig();
+				SerializerConfig->StateDescriptor = MemberCacheEntry->Descriptor;
+				SerializerConfig->Property = TFieldPath<FOptionalProperty>(const_cast<FOptionalProperty*>(CastField<FOptionalProperty>(Member.Property)));
+				MemberSerializerDescriptor->SerializerConfig = SerializerConfig;				
+			}
 			else
 			{
 				const FNetSerializerConfig* SerializerConfig = Member.SerializerInfo->BuildNetSerializerConfig(SerializerConfigBuffer, Member.Property);
diff --git a/Engine/Source/Runtime/Net/Iris/Private/Iris/Serialization/InternalNetSerializerUtils.h b/Engine/Source/Runtime/Net/Iris/Private/Iris/Serialization/InternalNetSerializerUtils.h
--- a/Engine/Source/Runtime/Net/Iris/Private/Iris/Serialization/InternalNetSerializerUtils.h
+++ b/Engine/Source/Runtime/Net/Iris/Private/Iris/Serialization/InternalNetSerializerUtils.h
@@ -27,4 +27,10 @@ inline bool IsArrayPropertyNetSerializer(const FNetSerializer* Serializer)
 	return Serializer == &UE_NET_GET_SERIALIZER(FArrayPropertyNetSerializer);
 }
 
+inline bool IsOptionalPropertyNetSerializer(const FNetSerializer* Serializer)
+{
+	return Serializer == &UE_NET_GET_SERIALIZER(FOptionalPropertyNetSerializer);
+}
+
+
 }
diff --git a/Engine/Source/Runtime/Net/Iris/Private/Iris/Serialization/InternalNetSerializers.cpp b/Engine/Source/Runtime/Net/Iris/Private/Iris/Serialization/InternalNetSerializers.cpp
--- a/Engine/Source/Runtime/Net/Iris/Private/Iris/Serialization/InternalNetSerializers.cpp
+++ b/Engine/Source/Runtime/Net/Iris/Private/Iris/Serialization/InternalNetSerializers.cpp
@@ -12,3 +12,13 @@ FArrayPropertyNetSerializerConfig::FArrayPropertyNetSerializerConfig()
 }
 
 FArrayPropertyNetSerializerConfig::~FArrayPropertyNetSerializerConfig() = default;
+
+FOptionalPropertyNetSerializerConfig::FOptionalPropertyNetSerializerConfig()
+: FNetSerializerConfig()
+{
+	// We have a generated descriptor, indicate that it must be destroyed
+	ConfigTraits = ENetSerializerConfigTraits::NeedDestruction;
+}
+
+FOptionalPropertyNetSerializerConfig::~FOptionalPropertyNetSerializerConfig() = default;
+
diff --git a/Engine/Source/Runtime/Net/Iris/Private/Iris/Serialization/InternalNetSerializers.h b/Engine/Source/Runtime/Net/Iris/Private/Iris/Serialization/InternalNetSerializers.h
--- a/Engine/Source/Runtime/Net/Iris/Private/Iris/Serialization/InternalNetSerializers.h
+++ b/Engine/Source/Runtime/Net/Iris/Private/Iris/Serialization/InternalNetSerializers.h
@@ -83,6 +83,36 @@ public:
 	uint32 MaxQuantizedSizeBits = DefaultMaxQuantizedSizeBits;
 };
 
+USTRUCT()
+struct FOptionalPropertyNetSerializerConfig : public FNetSerializerConfig
+{
+	GENERATED_BODY()
+
+	IRISCORE_API FOptionalPropertyNetSerializerConfig();
+	IRISCORE_API ~FOptionalPropertyNetSerializerConfig();
+
+	FOptionalPropertyNetSerializerConfig(const FOptionalPropertyNetSerializerConfig&) = delete;
+	FOptionalPropertyNetSerializerConfig(FOptionalPropertyNetSerializerConfig&&) = delete;
+
+	FOptionalPropertyNetSerializerConfig& operator=(const FOptionalPropertyNetSerializerConfig&) = delete;
+	FOptionalPropertyNetSerializerConfig& operator=(FOptionalPropertyNetSerializerConfig&&) = delete;
+
+	UPROPERTY()
+	TFieldPath<FOptionalProperty> Property;
+
+	TRefCountPtr<const UE::Net::FReplicationStateDescriptor> StateDescriptor;
+};
+
+template<>
+struct TStructOpsTypeTraits<FOptionalPropertyNetSerializerConfig> : public TStructOpsTypeTraitsBase2<FOptionalPropertyNetSerializerConfig>
+{
+	enum
+	{
+		WithCopy = false
+	};
+};
+
+
 // ENetRole. With role swapping at deserialization
 USTRUCT()
 struct FNetRoleNetSerializerConfig : public FNetSerializerConfig
@@ -140,6 +170,7 @@ UE_NET_DECLARE_SERIALIZER_INTERNAL(FLastResortPropertyNetSerializer, IRISCORE_AP
 UE_NET_DECLARE_SERIALIZER_INTERNAL(FArrayPropertyNetSerializer, IRISCORE_API);
 UE_NET_DECLARE_SERIALIZER_INTERNAL(FNetRoleNetSerializer, IRISCORE_API);
 UE_NET_DECLARE_SERIALIZER_INTERNAL(FFieldPathNetSerializer, IRISCORE_API);
+UE_NET_DECLARE_SERIALIZER_INTERNAL(FOptionalPropertyNetSerializer, IRISCORE_API);
 
 IRISCORE_API bool InitBitfieldNetSerializerConfigFromProperty(FBitfieldNetSerializerConfig& OutConfig, const FBoolProperty* Bitfield);
 
diff --git a/Engine/Source/Runtime/Net/Iris/Private/Iris/Serialization/OptionalPropertyNetSerializer.cpp b/Engine/Source/Runtime/Net/Iris/Private/Iris/Serialization/OptionalPropertyNetSerializer.cpp
new file mode 100644
--- /dev/null
+++ b/Engine/Source/Runtime/Net/Iris/Private/Iris/Serialization/OptionalPropertyNetSerializer.cpp
@@ -0,0 +1,457 @@
+// Copyright Epic Games, Inc. All Rights Reserved.
+
+#include "Iris/ReplicationState/ReplicationStateDescriptor.h"
+#include "Net/Core/Trace/NetTrace.h"
+#include "Iris/ReplicationSystem/ReplicationOperationsInternal.h"
+#include "Iris/ReplicationState/InternalPropertyReplicationState.h"
+#include "Iris/Serialization/NetSerializerArrayStorage.h"
+#include "Iris/Serialization/InternalNetSerializerUtils.h"
+#include "Iris/Serialization/InternalNetSerializers.h"
+#include "Iris/Serialization/NetBitStreamReader.h"
+#include "Iris/Serialization/NetBitStreamWriter.h"
+#include "UObject/PropertyOptional.h"
+
+namespace UE::Net
+{
+
+struct FOptionalPropertyNetSerializerQuantizedData
+{
+	FNetSerializerAlignedStorage OptionalPropertyData;
+};
+
+}
+
+template<> struct TIsPODType<UE::Net::FOptionalPropertyNetSerializerQuantizedData> { enum { Value = true }; };
+
+namespace UE::Net
+{
+
+struct FOptionalPropertyNetSerializer
+{
+public:
+	static const uint32 Version = 0;
+
+	static constexpr bool bIsForwardingSerializer = true; // Triggers asserts if a function is missing
+	static constexpr bool bHasDynamicState = true;
+	static constexpr bool bHasCustomNetReference = true;
+
+	typedef void SourceType; // Dummy
+	typedef FOptionalPropertyNetSerializerQuantizedData QuantizedType;
+	typedef FOptionalPropertyNetSerializerConfig ConfigType;
+
+	static void Serialize(FNetSerializationContext&, const FNetSerializeArgs& Args);
+	static void Deserialize(FNetSerializationContext&, const FNetDeserializeArgs& Args);
+
+	static void SerializeDelta(FNetSerializationContext&, const FNetSerializeDeltaArgs& Args);
+	static void DeserializeDelta(FNetSerializationContext&, const FNetDeserializeDeltaArgs& Args);
+
+	static void Quantize(FNetSerializationContext&, const FNetQuantizeArgs& Args);
+	static void Dequantize(FNetSerializationContext&, const FNetDequantizeArgs& Args);
+
+	static bool IsEqual(FNetSerializationContext&, const FNetIsEqualArgs& Args);
+	static bool Validate(FNetSerializationContext&, const FNetValidateArgs& Args);
+
+	static void CloneDynamicState(FNetSerializationContext&, const FNetCloneDynamicStateArgs&);
+	static void FreeDynamicState(FNetSerializationContext&, const FNetFreeDynamicStateArgs&);
+
+	static void CollectNetReferences(FNetSerializationContext&, const FNetCollectReferencesArgs&);
+
+	static void Apply(FNetSerializationContext&, const FNetApplyArgs&);
+
+private:
+
+	// Frees dynamic memory allocated and reset the entire quantized state to default.
+	static void Reset(FNetSerializationContext&, const ConfigType*, QuantizedType&);
+	static void InternalFreeOptionalPropertyData(FNetSerializationContext&, const ConfigType*, QuantizedType&);
+};
+
+UE_NET_IMPLEMENT_SERIALIZER(FOptionalPropertyNetSerializer);
+
+void FOptionalPropertyNetSerializer::Serialize(FNetSerializationContext& Context, const FNetSerializeArgs& Args)
+{
+	const QuantizedType& Source = *reinterpret_cast<const QuantizedType*>(Args.Source);
+	const ConfigType* Config = static_cast<const ConfigType*>(Args.NetSerializerConfig);
+	const FReplicationStateDescriptor* Descriptor = Config->StateDescriptor;
+
+	FNetBitStreamWriter* Writer = Context.GetBitStreamWriter();
+
+	const bool bIsSet = Source.OptionalPropertyData.GetData() != nullptr;
+	if (Writer->WriteBool(bIsSet))
+	{
+		const FReplicationStateMemberSerializerDescriptor& MemberSerializerDescriptor = Descriptor->MemberSerializerDescriptors[0];
+		const FNetSerializer* Serializer = MemberSerializerDescriptor.Serializer;
+
+		FNetSerializeArgs MemberArgs;
+		MemberArgs.Version = 0;
+		MemberArgs.NetSerializerConfig = MemberSerializerDescriptor.SerializerConfig;
+		MemberArgs.Source = NetSerializerValuePointer(Source.OptionalPropertyData.GetData());
+		MemberArgs.ChangeMaskInfo = Args.ChangeMaskInfo;
+
+		Serializer->Serialize(Context, MemberArgs);
+	}
+}
+
+void FOptionalPropertyNetSerializer::Deserialize(FNetSerializationContext& Context, const FNetDeserializeArgs& Args)
+{
+	QuantizedType& Target = *reinterpret_cast<QuantizedType*>(Args.Target);
+	const ConfigType* Config = static_cast<const ConfigType*>(Args.NetSerializerConfig);
+	const FReplicationStateDescriptor* Descriptor = Config->StateDescriptor;
+
+	FNetBitStreamReader* Reader = Context.GetBitStreamReader();
+
+	const bool bIsSet = Reader->ReadBool();
+	if (bIsSet)
+	{
+		if (Target.OptionalPropertyData.Num() == 0)
+		{
+			Target.OptionalPropertyData.AdjustSize(Context, Descriptor->InternalSize, Descriptor->InternalAlignment);
+		}
+
+		const FReplicationStateMemberSerializerDescriptor& MemberSerializerDescriptor = Descriptor->MemberSerializerDescriptors[0];
+		const FNetSerializer* Serializer = MemberSerializerDescriptor.Serializer;
+
+		FNetDeserializeArgs MemberArgs;
+		MemberArgs.Version = 0;
+		MemberArgs.NetSerializerConfig = MemberSerializerDescriptor.SerializerConfig;
+		MemberArgs.Target = NetSerializerValuePointer(Target.OptionalPropertyData.GetData());
+		MemberArgs.ChangeMaskInfo = Args.ChangeMaskInfo;
+
+		Serializer->Deserialize(Context, MemberArgs);
+	}
+	else
+	{
+		// Reset target
+		Reset(Context, Config, Target);
+	}
+}
+
+void FOptionalPropertyNetSerializer::SerializeDelta(FNetSerializationContext& Context, const FNetSerializeDeltaArgs& Args)
+{
+	const ConfigType* Config = static_cast<const ConfigType*>(Args.NetSerializerConfig);
+	const FReplicationStateDescriptor* Descriptor = Config->StateDescriptor;
+
+	FNetBitStreamWriter* Writer = Context.GetBitStreamWriter();
+	
+	const QuantizedType& Source = *reinterpret_cast<const QuantizedType*>(Args.Source);
+	const QuantizedType& Prev = *reinterpret_cast<const QuantizedType*>(Args.Prev);
+
+	const bool bIsSourceSet = Source.OptionalPropertyData.Num() > 0;
+	const bool bIsPrevSet = Prev.OptionalPropertyData.Num() > 0;
+
+	// If we cannot delta compress just call serialize
+	if (Writer->WriteBool(bIsSourceSet != bIsPrevSet))
+	{
+		Serialize(Context, Args);
+		return;
+	}
+	// Nothing to delta serialize
+	if (!bIsSourceSet)
+	{
+		return;
+	}
+
+	const FReplicationStateMemberSerializerDescriptor& MemberSerializerDescriptor = Descriptor->MemberSerializerDescriptors[0];
+	const FNetSerializer* Serializer = MemberSerializerDescriptor.Serializer;
+
+	FNetSerializeDeltaArgs MemberArgs;
+	MemberArgs.Version = 0;
+	MemberArgs.NetSerializerConfig = MemberSerializerDescriptor.SerializerConfig;
+	MemberArgs.Source = NetSerializerValuePointer(Source.OptionalPropertyData.GetData());
+	MemberArgs.Prev = NetSerializerValuePointer(Prev.OptionalPropertyData.GetData());
+	MemberArgs.ChangeMaskInfo = Args.ChangeMaskInfo;
+
+	Serializer->SerializeDelta(Context, MemberArgs);
+}
+
+void FOptionalPropertyNetSerializer::DeserializeDelta(FNetSerializationContext& Context, const FNetDeserializeDeltaArgs& Args)
+{
+	FNetBitStreamReader* Reader = Context.GetBitStreamReader();
+	if (const bool bIsSetDiffers = Reader->ReadBool())
+	{
+		Deserialize(Context, Args);
+		return;
+	}
+
+	const ConfigType* Config = static_cast<const ConfigType*>(Args.NetSerializerConfig);
+	const FReplicationStateDescriptor* Descriptor = Config->StateDescriptor;
+
+	QuantizedType& Target = *reinterpret_cast<QuantizedType*>(Args.Target);
+	const QuantizedType& Prev = *reinterpret_cast<const QuantizedType*>(Args.Prev);
+	const bool bIsPrevSet = Prev.OptionalPropertyData.Num() > 0;
+
+	// If not set, reset target.
+	if (!bIsPrevSet)
+	{
+		Reset(Context, Config, Target);
+		return;
+	}
+
+	if (Target.OptionalPropertyData.Num() == 0)
+	{
+		Target.OptionalPropertyData.AdjustSize(Context, Descriptor->InternalSize, Descriptor->InternalAlignment);
+	}
+
+	const FReplicationStateMemberSerializerDescriptor& MemberSerializerDescriptor = Descriptor->MemberSerializerDescriptors[0];
+	const FNetSerializer* Serializer = MemberSerializerDescriptor.Serializer;
+
+	FNetDeserializeDeltaArgs MemberArgs;
+	MemberArgs.Version = 0;
+	MemberArgs.NetSerializerConfig = MemberSerializerDescriptor.SerializerConfig;
+	MemberArgs.Target = NetSerializerValuePointer(Target.OptionalPropertyData.GetData());
+	MemberArgs.Prev = NetSerializerValuePointer(Prev.OptionalPropertyData.GetData());
+	MemberArgs.ChangeMaskInfo = Args.ChangeMaskInfo;
+
+	Serializer->DeserializeDelta(Context, MemberArgs);
+}
+
+void FOptionalPropertyNetSerializer::Reset(FNetSerializationContext& Context, const ConfigType* Config, QuantizedType& Value)
+{
+	InternalFreeOptionalPropertyData(Context, Config, Value);
+	Value.OptionalPropertyData.Free(Context);
+	FMemory::Memzero(&Value, sizeof(QuantizedType));
+}
+
+void FOptionalPropertyNetSerializer::InternalFreeOptionalPropertyData(FNetSerializationContext& Context, const ConfigType* Config, QuantizedType& Value)
+{
+	if (Value.OptionalPropertyData.Num() > 0 && EnumHasAnyFlags(Config->StateDescriptor->Traits, EReplicationStateTraits::HasDynamicState))
+	{
+		const FReplicationStateMemberSerializerDescriptor& MemberSerializerDescriptor = Config->StateDescriptor->MemberSerializerDescriptors[0];
+		const FNetSerializer* Serializer = MemberSerializerDescriptor.Serializer;
+
+		FNetFreeDynamicStateArgs MemberArgs;
+		MemberArgs.Version = 0;
+		MemberArgs.NetSerializerConfig = MemberSerializerDescriptor.SerializerConfig;
+		MemberArgs.Source = NetSerializerValuePointer(Value.OptionalPropertyData.GetData());
+
+		Serializer->FreeDynamicState(Context, MemberArgs);
+	}
+}
+
+void FOptionalPropertyNetSerializer::Quantize(FNetSerializationContext& Context, const FNetQuantizeArgs& Args)
+{
+	QuantizedType& Target = *reinterpret_cast<QuantizedType*>(Args.Target);
+
+	const ConfigType* Config = static_cast<const ConfigType*>(Args.NetSerializerConfig);
+	const FReplicationStateDescriptor* Descriptor = Config->StateDescriptor;
+
+	const FOptionalProperty* OptionalProperty = Config->Property.Get();
+	if (OptionalProperty->IsSet(reinterpret_cast<const void*>(Args.Source)))
+	{
+		if (Target.OptionalPropertyData.Num() == 0)
+		{
+			Target.OptionalPropertyData.AdjustSize(Context, Descriptor->InternalSize, Descriptor->InternalAlignment);
+		}
+		
+		const FReplicationStateMemberSerializerDescriptor& MemberSerializerDescriptor = Descriptor->MemberSerializerDescriptors[0];
+		const FNetSerializer* Serializer = MemberSerializerDescriptor.Serializer;
+
+		FNetQuantizeArgs MemberArgs;
+		MemberArgs.Version = 0;
+		MemberArgs.NetSerializerConfig = MemberSerializerDescriptor.SerializerConfig;
+		MemberArgs.Source = NetSerializerValuePointer(OptionalProperty->GetValuePointerForRead(reinterpret_cast<const void*>(Args.Source)));
+		MemberArgs.Target = NetSerializerValuePointer(Target.OptionalPropertyData.GetData());
+		MemberArgs.ChangeMaskInfo = Args.ChangeMaskInfo;
+
+		Serializer->Quantize(Context, MemberArgs);
+
+		return;
+	}
+	else
+	{
+		// Reset target
+		Reset(Context, Config, Target);
+	}
+}
+
+void FOptionalPropertyNetSerializer::Dequantize(FNetSerializationContext& Context, const FNetDequantizeArgs& Args)
+{
+	const QuantizedType& Source = *reinterpret_cast<const QuantizedType*>(Args.Source);
+	const ConfigType* Config = static_cast<const ConfigType*>(Args.NetSerializerConfig);
+	const FReplicationStateDescriptor* Descriptor = Config->StateDescriptor;
+
+	const FOptionalProperty* OptionalProperty = Config->Property.Get();
+	
+	if (Source.OptionalPropertyData.Num() > 0)
+	{
+		const FReplicationStateMemberSerializerDescriptor& MemberSerializerDescriptor = Descriptor->MemberSerializerDescriptors[0];
+		const FNetSerializer* Serializer = MemberSerializerDescriptor.Serializer;
+
+		FNetDequantizeArgs MemberArgs;
+		MemberArgs.Version = 0;
+		MemberArgs.NetSerializerConfig = MemberSerializerDescriptor.SerializerConfig;
+		MemberArgs.Source = NetSerializerValuePointer(Source.OptionalPropertyData.GetData());
+		MemberArgs.Target = NetSerializerValuePointer(OptionalProperty->MarkSetAndGetInitializedValuePointerToReplace(reinterpret_cast<void*>(Args.Target)));
+		MemberArgs.ChangeMaskInfo = Args.ChangeMaskInfo;
+
+		Serializer->Dequantize(Context, MemberArgs);
+
+		if (!Context.HasError())
+		{
+			return;
+		}
+	}
+
+	OptionalProperty->MarkUnset(reinterpret_cast<void*>(Args.Target));
+}
+
+void FOptionalPropertyNetSerializer::Apply(FNetSerializationContext& Context, const FNetApplyArgs& Args)
+{
+	const ConfigType* Config = static_cast<const ConfigType*>(Args.NetSerializerConfig);
+	const FReplicationStateDescriptor* Descriptor = Config->StateDescriptor;
+
+	const FOptionalProperty* OptionalProperty = Config->Property.Get();	
+	if (OptionalProperty->IsSet(reinterpret_cast<const void*>(Args.Source)))
+	{
+		const FReplicationStateMemberSerializerDescriptor& MemberSerializerDescriptor = Descriptor->MemberSerializerDescriptors[0];
+		const FNetSerializer* Serializer = MemberSerializerDescriptor.Serializer;
+
+		UE::Net::Private::InternalApplyPropertyValue(Descriptor, 0, static_cast<uint8*>(OptionalProperty->MarkSetAndGetInitializedValuePointerToReplace(reinterpret_cast<void*>(Args.Target))), static_cast<uint8*>(OptionalProperty->MarkSetAndGetInitializedValuePointerToReplace(reinterpret_cast<void*>(Args.Source))));
+	}
+	else
+	{
+		OptionalProperty->MarkUnset(reinterpret_cast<void*>(Args.Target));
+	}
+}
+
+
+bool FOptionalPropertyNetSerializer::IsEqual(FNetSerializationContext& Context, const FNetIsEqualArgs& Args)
+{
+	const ConfigType* Config = static_cast<const ConfigType*>(Args.NetSerializerConfig);
+	const FReplicationStateDescriptor* Descriptor = Config->StateDescriptor;
+	
+	if (Args.bStateIsQuantized)
+	{
+		const QuantizedType& Source0 = *reinterpret_cast<const QuantizedType*>(Args.Source0);
+		const QuantizedType& Source1 = *reinterpret_cast<const QuantizedType*>(Args.Source1);
+
+		const bool bIsSource0Set = Source0.OptionalPropertyData.Num() > 0;
+		const bool bIsSource1Set = Source1.OptionalPropertyData.Num() > 0;
+
+		if (bIsSource0Set != bIsSource1Set)
+		{
+			return false;
+		}
+		if (bIsSource0Set)
+		{
+			const FReplicationStateMemberSerializerDescriptor& MemberSerializerDescriptor = Descriptor->MemberSerializerDescriptors[0];
+			const FNetSerializer* Serializer = MemberSerializerDescriptor.Serializer;
+
+			FNetIsEqualArgs MemberArgs = Args;
+			MemberArgs.Version = 0;
+			MemberArgs.NetSerializerConfig = MemberSerializerDescriptor.SerializerConfig;
+			MemberArgs.Source0 = NetSerializerValuePointer(Source0.OptionalPropertyData.GetData());
+			MemberArgs.Source1 = NetSerializerValuePointer(Source1.OptionalPropertyData.GetData());
+			MemberArgs.ChangeMaskInfo = Args.ChangeMaskInfo;
+
+			return Serializer->IsEqual(Context, MemberArgs);
+		}
+	}
+	else
+	{
+		const FOptionalProperty* OptionalProperty = Config->Property.Get();
+		const bool bIsEqual = OptionalProperty->Identical(reinterpret_cast<const void*>(Args.Source0), reinterpret_cast<const void*>(Args.Source1), 0);
+
+		return bIsEqual;
+	}
+	return true;
+}
+
+bool FOptionalPropertyNetSerializer::Validate(FNetSerializationContext& Context, const FNetValidateArgs& Args)
+{
+	const ConfigType* Config = static_cast<const ConfigType*>(Args.NetSerializerConfig);
+	const FReplicationStateDescriptor* Descriptor = Config->StateDescriptor;
+	if (Descriptor == nullptr)
+	{
+		return false;
+	}
+
+	const FOptionalProperty* OptionalProperty = Config->Property.Get();
+	if (OptionalProperty->IsSet(reinterpret_cast<const void*>(Args.Source)))
+	{
+		const FReplicationStateMemberSerializerDescriptor& MemberSerializerDescriptor = Descriptor->MemberSerializerDescriptors[0];
+		const FNetSerializer* Serializer = MemberSerializerDescriptor.Serializer;
+
+		FNetValidateArgs MemberArgs;
+		MemberArgs.Version = 0;
+		MemberArgs.NetSerializerConfig = MemberSerializerDescriptor.SerializerConfig;
+		MemberArgs.Source = NetSerializerValuePointer(OptionalProperty->GetValuePointerForReadOrReplace(reinterpret_cast<const void*>(Args.Source)));
+
+		if (!Serializer->Validate(Context, MemberArgs))
+		{
+			return false;
+		}
+	}
+
+	return true;
+}
+
+void FOptionalPropertyNetSerializer::CloneDynamicState(FNetSerializationContext& Context, const FNetCloneDynamicStateArgs& Args)
+{
+	const QuantizedType& Source = *reinterpret_cast<const QuantizedType*>(Args.Source);
+	QuantizedType& Target = *reinterpret_cast<QuantizedType*>(Args.Target);
+
+	Target.OptionalPropertyData.Clone(Context, Source.OptionalPropertyData);
+
+	// Forward call if necessary
+	const ConfigType* Config = static_cast<const ConfigType*>(Args.NetSerializerConfig);
+	const FReplicationStateDescriptor* Descriptor = Config->StateDescriptor;
+
+	const bool bIsSourceSet = Source.OptionalPropertyData.Num() > 0;
+	if (bIsSourceSet && EnumHasAnyFlags(Descriptor->Traits, EReplicationStateTraits::HasDynamicState))
+	{
+		const FReplicationStateMemberSerializerDescriptor& MemberSerializerDescriptor = Descriptor->MemberSerializerDescriptors[0];
+		const FNetSerializer* Serializer = MemberSerializerDescriptor.Serializer;
+
+		FNetCloneDynamicStateArgs MemberArgs;
+		MemberArgs.Version = 0;
+		MemberArgs.NetSerializerConfig = MemberSerializerDescriptor.SerializerConfig;
+		MemberArgs.Source = NetSerializerValuePointer(Source.OptionalPropertyData.GetData());
+		MemberArgs.Target = NetSerializerValuePointer(Target.OptionalPropertyData.GetData());
+
+		Serializer->CloneDynamicState(Context, MemberArgs);
+	}
+}
+
+void FOptionalPropertyNetSerializer::FreeDynamicState(FNetSerializationContext& Context, const FNetFreeDynamicStateArgs& Args)
+{
+	QuantizedType& Value = *reinterpret_cast<QuantizedType*>(Args.Source);
+	const ConfigType* Config = static_cast<const ConfigType*>(Args.NetSerializerConfig);
+
+	if (!Value.OptionalPropertyData.GetData())
+	{
+		return;
+	}
+
+	const FReplicationStateDescriptor* Descriptor = Config->StateDescriptor;
+	if (EnumHasAnyFlags(Descriptor->MemberTraitsDescriptors[0].Traits, EReplicationStateMemberTraits::HasDynamicState))
+	{
+		const FReplicationStateMemberSerializerDescriptor& MemberSerializerDescriptor = Descriptor->MemberSerializerDescriptors[0];
+
+		const FNetSerializer* Serializer = MemberSerializerDescriptor.Serializer;
+
+		FNetFreeDynamicStateArgs MemberArgs;
+		MemberArgs.Version = 0;
+		MemberArgs.NetSerializerConfig = MemberSerializerDescriptor.SerializerConfig;
+		MemberArgs.Source = NetSerializerValuePointer(Value.OptionalPropertyData.GetData());
+
+		Serializer->FreeDynamicState(Context, MemberArgs);
+	}
+	
+	Value.OptionalPropertyData.Free(Context);
+}
+
+void FOptionalPropertyNetSerializer::CollectNetReferences(FNetSerializationContext& Context, const FNetCollectReferencesArgs& Args)
+{
+	const ConfigType* Config = static_cast<const ConfigType*>(Args.NetSerializerConfig);
+	const FReplicationStateDescriptor* Descriptor = Config->StateDescriptor;
+	const QuantizedType& Value = *reinterpret_cast<QuantizedType*>(Args.Source);
+
+	if (!Descriptor->HasObjectReference() || Value.OptionalPropertyData.Num() == 0)
+	{
+		return;
+	}
+
+	Private::FReplicationStateOperationsInternal::CollectReferences(Context, *reinterpret_cast<FNetReferenceCollector*>(Args.Collector), Args.ChangeMaskInfo, Value.OptionalPropertyData.GetData(), Descriptor);
+}
+
+}
