dIris
- Fixed ensure in prune stale objects triggered by SubObjects of objects pending tear-off not being properly detached until destroyed.
- Unit tests for the above.

diff --git a/Engine/Plugins/Runtime/ReplicationSystemTestPlugin/Source/Private/Tests/ReplicationSystem/TestReplicationSystem.cpp b/Engine/Plugins/Runtime/ReplicationSystemTestPlugin/Source/Private/Tests/ReplicationSystem/TestReplicationSystem.cpp
--- a/Engine/Plugins/Runtime/ReplicationSystemTestPlugin/Source/Private/Tests/ReplicationSystem/TestReplicationSystem.cpp
+++ b/Engine/Plugins/Runtime/ReplicationSystemTestPlugin/Source/Private/Tests/ReplicationSystem/TestReplicationSystem.cpp
@@ -2504,6 +2504,144 @@ UE_NET_TEST_FIXTURE(FTearOffAndFlushTestFixture, TestTearOffNextUpdateExistingOb
 	UE_NET_ASSERT_TRUE(Cast<UTestReplicatedIrisObject>(Client->GetReplicationBridge()->GetReplicatedObject(ServerSubObject->NetRefHandle)) == nullptr);
 }
 
+// Test TearOff and destroy 
+UE_NET_TEST_FIXTURE(FTearOffAndFlushTestFixture, TestTearOffNextUpdateWithUnexpectedDestroyObject)
+{
+	UReplicationSystem* ReplicationSystem = Server->ReplicationSystem;
+	FNetRefHandleManager* NetRefHandleManager = &ReplicationSystem->GetReplicationSystemInternal()->GetNetRefHandleManager();
+
+	// Add a client
+	FReplicationSystemTestClient* Client = CreateClient();
+
+	// Create a pushbased object on server
+	UTestReplicatedIrisPushModelObject* ServerObject = NewObject<UTestReplicatedIrisPushModelObject>();
+	ServerObject->NetRefHandle = Server->ReplicationBridge->BeginReplication(ServerObject);
+
+	const FInternalNetRefIndex ServerObjectInternalIndex = NetRefHandleManager->GetInternalIndex(ServerObject->NetRefHandle);
+
+	Server->UpdateAndSend({Client});
+
+	// Verify that we no longer have any references to the object
+	UE_NET_ASSERT_EQ(uint16(1), NetRefHandleManager->GetNetObjectRefCount(ServerObjectInternalIndex));
+
+	// Initiate tear-off
+	Server->ReplicationSystem->TearOffNextUpdate(ServerObject->NetRefHandle);
+
+	// Put tear-off in flight
+	Server->NetUpdate();
+	Server->SendTo(Client);
+	Server->PostSendUpdate();
+
+	// Destroy object in a bad way.
+	ServerObject->MarkAsGarbage();
+	constexpr bool bPerformFullPurge = false;
+	CollectGarbage(GARBAGE_COLLECTION_KEEPFLAGS, bPerformFullPurge);
+
+	Server->DeliverTo(Client, true);
+
+	// As we have issued tear-off, this should not cause an ensure
+	{
+		FEnsureScope EnsureScope;
+
+		Server->UpdateAndSend({Client});
+
+		UE_NET_ASSERT_EQ(EnsureScope.GetCount(), 0);
+	}
+
+	// Verify that we no longer have any references to the object
+	UE_NET_ASSERT_EQ(uint16(0), NetRefHandleManager->GetNetObjectRefCount(ServerObjectInternalIndex));
+}
+
+// Test TearOff and destroy of SubObjects that are still pending create/tearoff
+UE_NET_TEST_FIXTURE(FTearOffAndFlushTestFixture, TestTearOffNextUpdateWithUnexpectedDestroySubObject)
+{
+	UReplicationSystem* ReplicationSystem = Server->ReplicationSystem;
+	FNetRefHandleManager* NetRefHandleManager = &ReplicationSystem->GetReplicationSystemInternal()->GetNetRefHandleManager();
+
+	// Add a client
+	FReplicationSystemTestClient* Client = CreateClient();
+
+	// Spawn object on server
+	UTestReplicatedIrisObject* ServerObject = Server->CreateObject(0,0);
+
+	// Spawn second object on server as a subobject without adding a strong reference
+	UTestReplicatedIrisObject* ServerSubObject = NewObject<UTestReplicatedIrisObject>();
+	ServerSubObject->NetRefHandle = Server->ReplicationBridge->BeginReplication(ServerObject->NetRefHandle, ServerSubObject);
+	
+	const FInternalNetRefIndex ServerObjectInternalIndex = NetRefHandleManager->GetInternalIndex(ServerObject->NetRefHandle);
+	const FInternalNetRefIndex SubObjectObjectInternalIndex = NetRefHandleManager->GetInternalIndex(ServerSubObject->NetRefHandle);
+
+	Server->UpdateAndSend({Client});
+
+	// TearOff the object this will also tear-off subobject
+	Server->ReplicationSystem->TearOffNextUpdate(ServerObject->NetRefHandle);
+
+	// Just logic update, object will be pending tearoff.
+	Server->UpdateAndSend({});
+
+	// Destroy subobject in a bad way.
+	ServerSubObject->MarkAsGarbage();
+	constexpr bool bPerformFullPurge = false;
+	CollectGarbage(GARBAGE_COLLECTION_KEEPFLAGS, bPerformFullPurge);
+
+	// As we have issued tear-off, this should not cause an ensure
+	{
+		FEnsureScope EnsureScope;
+
+		Server->UpdateAndSend({Client});
+
+		UE_NET_ASSERT_EQ(EnsureScope.GetCount(), 0);
+	}
+
+	// Verify that we no longer have any references to the objects
+	UE_NET_ASSERT_EQ(uint16(0), NetRefHandleManager->GetNetObjectRefCount(ServerObjectInternalIndex));
+	UE_NET_ASSERT_EQ(uint16(0), NetRefHandleManager->GetNetObjectRefCount(SubObjectObjectInternalIndex));
+}
+
+// Test TearOff and destroy of SubObjects that are still pending create/tearoff
+UE_NET_TEST_FIXTURE(FTearOffAndFlushTestFixture, TestImmediateTearOffWithUnexpectedDestroySubObject)
+{
+	UReplicationSystem* ReplicationSystem = Server->ReplicationSystem;
+	FNetRefHandleManager* NetRefHandleManager = &ReplicationSystem->GetReplicationSystemInternal()->GetNetRefHandleManager();
+
+	// Add a client
+	FReplicationSystemTestClient* Client = CreateClient();
+
+	// Spawn object on server
+	UTestReplicatedIrisObject* ServerObject = Server->CreateObject(0,0);
+
+	// Explicit spawn of second object on server as a subobject to avoid hard refernce
+	UTestReplicatedIrisObject* ServerSubObject = NewObject<UTestReplicatedIrisObject>();
+	FNetRefHandle SubObjectHandle = Server->ReplicationBridge->BeginReplication(ServerObject->NetRefHandle, ServerSubObject);
+	ServerSubObject->NetRefHandle = SubObjectHandle;
+
+	const FInternalNetRefIndex ServerObjectInternalIndex = NetRefHandleManager->GetInternalIndex(ServerObject->NetRefHandle);
+	const FInternalNetRefIndex SubObjectObjectInternalIndex = NetRefHandleManager->GetInternalIndex(ServerSubObject->NetRefHandle);
+
+	Server->UpdateAndSend({Client});
+
+	// Immediate tearoff
+	Server->ReplicationBridge->EndReplication(ServerObject, EEndReplicationFlags::TearOff);
+
+	// Destroy subobject in a bad way.
+	ServerSubObject->MarkAsGarbage();
+	constexpr bool bPerformFullPurge = false;
+	CollectGarbage(GARBAGE_COLLECTION_KEEPFLAGS, bPerformFullPurge);
+
+	// As we have issued tear-off, this should not cause an ensure
+	{
+		FEnsureScope EnsureScope;
+
+		Server->UpdateAndSend({Client});
+
+		UE_NET_ASSERT_EQ(EnsureScope.GetCount(), 0);
+	}
+
+	// Verify that we no longer have any references to the objects
+	UE_NET_ASSERT_EQ(uint16(0), NetRefHandleManager->GetNetObjectRefCount(ServerObjectInternalIndex));
+	UE_NET_ASSERT_EQ(uint16(0), NetRefHandleManager->GetNetObjectRefCount(SubObjectObjectInternalIndex));
+}
+
 // Test TearOff and destroy of SubObjects that are still pending create/tearoff
 UE_NET_TEST_FIXTURE(FTearOffAndFlushTestFixture, TestTearOffNextUpdateExistingObjectWithSubObjectPendingCreation)
 {
diff --git a/Engine/Source/Runtime/Net/Iris/Private/Iris/ReplicationSystem/ReplicationBridge.cpp b/Engine/Source/Runtime/Net/Iris/Private/Iris/ReplicationSystem/ReplicationBridge.cpp
--- a/Engine/Source/Runtime/Net/Iris/Private/Iris/ReplicationSystem/ReplicationBridge.cpp
+++ b/Engine/Source/Runtime/Net/Iris/Private/Iris/ReplicationSystem/ReplicationBridge.cpp
@@ -487,10 +487,6 @@ void UReplicationBridge::StopReplicatingNetRefHandle(FNetRefHandle Handle, EEndR
 			// Add handle to list of objects pending EndReplication indicate that it should be destroyed during next update
 			// We need to do this to cover the case where the torn off object not yet has been added to the scope.
 			AddPendingEndReplication(Handle, InternalReplicationIndex, EndReplicationFlags);
-
-			// Detach instance as we must assume that we should not access the object after this call.
-			ObjectData.bPendingEndReplication = 1U;
-			InternalDetachInstanceFromNetRefHandle(InternalReplicationIndex);
 		}
 		else 
 		{
@@ -906,6 +902,10 @@ void UReplicationBridge::InternalTearOff(FNetRefHandle Handle)
 	// Mark object as being torn-off and that we should no longer propagate state changes
 	ObjectData.bTearOff = 1U;
 	ObjectData.bShouldPropagateChangedStates = 0U;
+
+	// Detach instance as we must assume that we should not access the object after this call.
+	ObjectData.bPendingEndReplication = 1U;
+	InternalDetachInstanceFromNetRefHandle(InternalObjectIndex);
 }
 
 bool UReplicationBridge::IsReplicatedHandle(FNetRefHandle Handle) const
