Iris - Dynamic config loading/unloading

#Synced-CL 51332317

diff --git a/Engine/Plugins/Runtime/ReplicationSystemTestPlugin/Source/Private/Tests/Config/TestDynamicConfig.cpp b/Engine/Plugins/Runtime/ReplicationSystemTestPlugin/Source/Private/Tests/Config/TestDynamicConfig.cpp
new file mode 100644
--- /dev/null
+++ b/Engine/Plugins/Runtime/ReplicationSystemTestPlugin/Source/Private/Tests/Config/TestDynamicConfig.cpp
@@ -0,0 +1,218 @@
+// Copyright Epic Games, Inc. All Rights Reserved.
+
+#include "NetworkAutomationTest.h"
+#include "NetworkAutomationTestMacros.h"
+#include "Containers/Set.h"
+#include "Iris/Config/IrisDynamicConfig.h"
+#include "Iris/Config/IrisDynamicConfig.h"
+#include "Tests/ReplicationSystem/ReplicationSystemTestFixture.h"
+#include "Templates/UnrealTemplate.h"
+
+namespace UE::Net::Private
+{
+
+class FTestDynamicConfigFixture : public FReplicationSystemTestFixture
+{
+protected:
+	virtual void SetUp() override
+	{
+		FReplicationSystemTestFixture::SetUp();
+
+		FIrisDynamicConfig::OnIrisDynamicConfigChange().AddRaw(this, &FTestDynamicConfigFixture::OnDynamicConfigChange);
+	}
+
+	virtual void TearDown() override
+	{
+		FIrisDynamicConfig::OnIrisDynamicConfigChange().RemoveAll(this);
+
+		FReplicationSystemTestFixture::TearDown();
+
+		OnDynamicConfigChangeModifiedSections.Reset();
+	}
+
+
+	void OnDynamicConfigChange(const TSet<FString>& ModifiedSections)
+	{
+		++DynamicConfigChangeCounter;
+		OnDynamicConfigChangeModifiedSections = ModifiedSections;
+	}
+
+	const TSet<FString>& GetModifiedSections() const
+	{
+		return OnDynamicConfigChangeModifiedSections;
+	}
+
+	TSet<FString> GetAndClearModifiedSections() 
+	{
+		TSet<FString> ModifiedSections = MoveTemp(OnDynamicConfigChangeModifiedSections);
+		OnDynamicConfigChangeModifiedSections = TSet<FString>();
+		return ModifiedSections;
+	}
+
+	unsigned GetDynamicConfigChangeCounter() const
+	{
+		return DynamicConfigChangeCounter;
+	}
+
+protected:
+	// This helper class keeps track of registered configs so that they're always properly unregistered, even when a test fails.
+	class FScopedConfigRegistrar
+	{
+	public:
+		~FScopedConfigRegistrar()
+		{
+			for (FName ConfigName : RegisteredConfigs)
+			{
+				FIrisDynamicConfig::UnregisterDynamicConfig(ConfigName);
+			}
+		}
+
+		void RegisterDynamicConfig(FName UniqueConfigName, const FString& FilePath)
+		{
+			RegisteredConfigs.Add(UniqueConfigName);
+			FIrisDynamicConfig::RegisterDynamicConfig(UniqueConfigName, FilePath);
+		}
+
+		void RegisterDynamicConfigBuffer(FName UniqueConfigName, const FString& VirtualFilePath, const FString& ConfigContents)
+		{
+			RegisteredConfigs.Add(UniqueConfigName);
+			FIrisDynamicConfig::RegisterDynamicConfigBuffer(UniqueConfigName, VirtualFilePath, ConfigContents);
+		}
+
+		void RegisterDynamicConfigBuffer(FName UniqueConfigName, const FString& VirtualFilePath, FUtf8StringView ConfigContents)
+		{
+			RegisteredConfigs.Add(UniqueConfigName);
+			FIrisDynamicConfig::RegisterDynamicConfigBuffer(UniqueConfigName, VirtualFilePath, ConfigContents);
+		}
+
+		void UnregisterDynamicConfig(FName UniqueConfigName)
+		{
+			RegisteredConfigs.Remove(UniqueConfigName);
+			FIrisDynamicConfig::UnregisterDynamicConfig(UniqueConfigName);
+		}
+
+	private:
+		TSet<FName, DefaultKeyFuncs<FName>, TInlineSetAllocator<8>> RegisteredConfigs;
+	};
+
+private:
+	// The set is cleared when the config system broadcasts a change
+	TSet<FString> OnDynamicConfigChangeModifiedSections;
+	unsigned DynamicConfigChangeCounter = 0;
+};
+
+UE_NET_TEST_FIXTURE(FTestDynamicConfigFixture, CanRegisterDynamicConfigBuffer)
+{
+	FScopedConfigRegistrar Registrar;
+	Registrar.RegisterDynamicConfigBuffer(FName("TestDynamicConfig"), "TestDynamicConfig/Iris/Iris.ini", UTF8TEXTVIEW("[SectionName]"));
+
+	const TSet<FString>& ModifiedSections = GetModifiedSections();
+	UE_NET_ASSERT_EQ(ModifiedSections.Num(), 1);
+	UE_NET_ASSERT_NE(ModifiedSections.Find("SectionName"), nullptr);
+}
+
+UE_NET_TEST_FIXTURE(FTestDynamicConfigFixture, RegisteringEmptyDynamicConfigIsNotBroadcasted)
+{
+	FScopedConfigRegistrar Registrar;
+
+	const unsigned OldCounter = GetDynamicConfigChangeCounter();
+	Registrar.RegisterDynamicConfigBuffer(FName("TestDynamicConfig1"), "TestDynamicConfig/Iris/Iris1.ini", UTF8TEXTVIEW(""));
+	UE_NET_ASSERT_EQ(GetDynamicConfigChangeCounter(), OldCounter);
+
+	Registrar.RegisterDynamicConfigBuffer(FName("TestDynamicConfig2"), "TestDynamicConfig/Iris/Iris2.ini", FString());
+	UE_NET_ASSERT_EQ(GetDynamicConfigChangeCounter(), OldCounter);
+}
+
+UE_NET_TEST_FIXTURE(FTestDynamicConfigFixture, CanUnregisterDynamicConfigBuffer)
+{
+	FScopedConfigRegistrar Registrar;
+
+	Registrar.RegisterDynamicConfigBuffer(FName("TestDynamicConfig"), "TestDynamicConfig/Iris/Iris.ini", UTF8TEXTVIEW("[SectionName]"));
+
+	GetAndClearModifiedSections();
+
+	Registrar.UnregisterDynamicConfig(FName("TestDynamicConfig"));
+	const TSet<FString>& ModifiedSections = GetModifiedSections();
+	UE_NET_ASSERT_EQ(ModifiedSections.Num(), 1);
+	UE_NET_ASSERT_NE(ModifiedSections.Find("SectionName"), nullptr);
+}
+
+UE_NET_TEST_FIXTURE(FTestDynamicConfigFixture, CanRegisterAndUnregisterDynamicConfigsInDifferentOrder)
+{
+	FScopedConfigRegistrar Registrar;
+
+	Registrar.RegisterDynamicConfigBuffer(FName("TestDynamicConfig1"), "TestDynamicConfig/Iris/Iris1.ini", UTF8TEXTVIEW("[SectionName1]"));
+	Registrar.RegisterDynamicConfigBuffer(FName("TestDynamicConfig2"), "TestDynamicConfig/Iris/Iris2.ini", UTF8TEXTVIEW("[SectionName2]"));
+
+	GetAndClearModifiedSections();
+
+	Registrar.UnregisterDynamicConfig(FName("TestDynamicConfig2"));
+	const TSet<FString>& ModifiedSections1 = GetModifiedSections();
+	UE_NET_ASSERT_EQ(ModifiedSections1.Num(), 1);
+	UE_NET_ASSERT_NE(ModifiedSections1.Find("SectionName2"), nullptr);
+
+	Registrar.UnregisterDynamicConfig(FName("TestDynamicConfig1"));
+	const TSet<FString>& ModifiedSections2 = GetModifiedSections();
+	UE_NET_ASSERT_EQ(ModifiedSections2.Num(), 1);
+	UE_NET_ASSERT_NE(ModifiedSections2.Find("SectionName1"), nullptr);
+}
+
+UE_NET_TEST_FIXTURE(FTestDynamicConfigFixture, CanRegisterAndUnregisterConfigsWithOverlappingSectionsInDifferentOrder)
+{
+	FScopedConfigRegistrar Registrar;
+
+	Registrar.RegisterDynamicConfigBuffer(FName("TestDynamicConfig1"), "TestDynamicConfig/Iris/Iris1.ini", UTF8TEXTVIEW("[SectionName1]\n; Hello\n[SectionName3]\n; World"));
+	Registrar.RegisterDynamicConfigBuffer(FName("TestDynamicConfig2"), "TestDynamicConfig/Iris/Iris2.ini", UTF8TEXTVIEW("[SectionName3]\n; Hello\n[SectionName2]\n; World"));
+
+	GetAndClearModifiedSections();
+
+	Registrar.UnregisterDynamicConfig(FName("TestDynamicConfig2"));
+	const TSet<FString>& ModifiedSections1 = GetModifiedSections();
+	UE_NET_ASSERT_EQ(ModifiedSections1.Num(), 2);
+	UE_NET_ASSERT_NE(ModifiedSections1.Find("SectionName2"), nullptr);
+	UE_NET_ASSERT_NE(ModifiedSections1.Find("SectionName3"), nullptr);
+
+	Registrar.UnregisterDynamicConfig(FName("TestDynamicConfig1"));
+	const TSet<FString>& ModifiedSections2 = GetModifiedSections();
+	UE_NET_ASSERT_EQ(ModifiedSections2.Num(), 2);
+	UE_NET_ASSERT_NE(ModifiedSections2.Find("SectionName1"), nullptr);
+	UE_NET_ASSERT_NE(ModifiedSections2.Find("SectionName3"), nullptr);
+}
+
+UE_NET_TEST_FIXTURE(FTestDynamicConfigFixture, CanRetrieveSectionConfig)
+{
+	FScopedConfigRegistrar Registrar;
+
+	Registrar.RegisterDynamicConfigBuffer(FName("TestDynamicConfig"), "TestDynamicConfig/Iris/Iris.ini", UTF8TEXTVIEW("[SectionName]\nKey1=Value\nKey2=Value"));
+
+	TArray<FString> KeyValuePairsArray;
+	FIrisDynamicConfig::GetSection(TEXT("SectionName"), KeyValuePairsArray);
+
+	TSet<FString> KeyValuePairsSet(KeyValuePairsArray);
+	UE_NET_ASSERT_EQ(KeyValuePairsSet.Num(), 2);
+
+	UE_NET_ASSERT_NE(KeyValuePairsSet.Find("Key1=Value"), nullptr); 
+	UE_NET_ASSERT_NE(KeyValuePairsSet.Find("Key2=Value"), nullptr); 
+}
+
+UE_NET_TEST_FIXTURE(FTestDynamicConfigFixture, CanRetrieveSectionConfigMulti)
+{
+	FScopedConfigRegistrar Registrar;
+
+	// Register multiple configs for the same section
+	Registrar.RegisterDynamicConfigBuffer(FName("TestDynamicConfig1"), "TestDynamicConfig/Iris/Iris1.ini", UTF8TEXTVIEW("[SectionName]\nKey1=Value\nKey2=Value"));
+	Registrar.RegisterDynamicConfigBuffer(FName("TestDynamicConfig2"), "TestDynamicConfig/Iris/Iris2.ini", UTF8TEXTVIEW("[SectionName]\nKey3=Value\nKey4=Value"));
+
+	TArray<FString> KeyValuePairsArray;
+	FIrisDynamicConfig::GetSection(TEXT("SectionName"), KeyValuePairsArray);
+
+	TSet<FString> KeyValuePairsSet(KeyValuePairsArray);
+	UE_NET_ASSERT_EQ(KeyValuePairsSet.Num(), 4);
+
+	UE_NET_ASSERT_NE(KeyValuePairsSet.Find("Key1=Value"), nullptr); 
+	UE_NET_ASSERT_NE(KeyValuePairsSet.Find("Key2=Value"), nullptr); 
+	UE_NET_ASSERT_NE(KeyValuePairsSet.Find("Key3=Value"), nullptr); 
+	UE_NET_ASSERT_NE(KeyValuePairsSet.Find("Key4=Value"), nullptr); 
+}
+
+}
diff --git a/Engine/Source/Runtime/Net/Iris/Private/Iris/Config/IrisDynamicConfig.cpp b/Engine/Source/Runtime/Net/Iris/Private/Iris/Config/IrisDynamicConfig.cpp
new file mode 100644
--- /dev/null
+++ b/Engine/Source/Runtime/Net/Iris/Private/Iris/Config/IrisDynamicConfig.cpp
@@ -0,0 +1,205 @@
+// Copyright Epic Games, Inc. All Rights Reserved.
+
+#include "Iris/Config/IrisDynamicConfig.h"
+
+#include "CoreGlobals.h" // for IsEngineExitRequested()
+#include "Containers/Map.h"
+#include "Iris/Core/IrisLog.h"
+#include "Misc/ConfigCacheIni.h"
+
+namespace UE::Net::Private
+{
+
+class FIrisDynamicConfigImpl
+{
+public:
+	void RegisterDynamicConfig(FName UniqueConfigName, const FString& FilePath);
+	void RegisterDynamicConfigBuffer(FName UniqueConfigName, const FString& VirtualFilePath, const FString& ConfigContents);
+	void RegisterDynamicConfigBuffer(FName UniqueConfigName, const FString& VirtualFilePath, FUtf8StringView ConfigContents);
+	void UnregisterDynamicConfig(FName UniqueConfigName);
+
+	void GetSection(const TCHAR* SectionName, TArray<FString>& OutKeyValuePairs);
+	void GetArray(const TCHAR* SectionName, const TCHAR* ArrayName, TArray<FString>& OutValues);
+
+	FIrisDynamicConfig::FOnIrisDynamicConfigChange OnIrisDynamicConfigChange;
+
+private:
+	FConfigBranch* GetOrCreateConfig();
+
+	// The config system uses three different strings types to retrieve things...
+	inline static FString VirtualBranchFileName = "Iris.ini";
+
+	// The cache is only needed to be able to construct the FConfigBranch.
+	TOptional<FConfigCacheIni> ConfigCache;
+	FConfigBranch* IrisConfig = nullptr;
+	TCompactMap<FName, FString> ConfigNameToFilePath;
+};
+
+static FIrisDynamicConfigImpl IrisDynamicConfig;
+
+}
+
+namespace UE::Net
+{
+
+void FIrisDynamicConfig::RegisterDynamicConfig(FName UniqueConfigName, const FString& FilePath)
+{
+	Private::IrisDynamicConfig.RegisterDynamicConfig(UniqueConfigName, FilePath);
+}
+
+void FIrisDynamicConfig::RegisterDynamicConfigBuffer(FName UniqueConfigName, const FString& VirtualFilePath, const FString& ConfigContents)
+{
+	Private::IrisDynamicConfig.RegisterDynamicConfigBuffer(UniqueConfigName, VirtualFilePath, ConfigContents);
+}
+
+void FIrisDynamicConfig::RegisterDynamicConfigBuffer(FName UniqueConfigName, const FString& VirtualFilePath, FUtf8StringView ConfigContents)
+{
+	Private::IrisDynamicConfig.RegisterDynamicConfigBuffer(UniqueConfigName, VirtualFilePath, ConfigContents);
+}
+
+void FIrisDynamicConfig::UnregisterDynamicConfig(FName UniqueConfigName)
+{
+	Private::IrisDynamicConfig.UnregisterDynamicConfig(UniqueConfigName);
+}
+
+FIrisDynamicConfig::FOnIrisDynamicConfigChange::RegistrationType& FIrisDynamicConfig::OnIrisDynamicConfigChange()
+{
+	return Private::IrisDynamicConfig.OnIrisDynamicConfigChange;
+}
+
+void FIrisDynamicConfig::GetSection(const TCHAR* SectionName, TArray<FString>& OutKeyValuePairs)
+{
+	Private::IrisDynamicConfig.GetSection(SectionName, OutKeyValuePairs);
+}
+
+void FIrisDynamicConfig::GetArray(const TCHAR* SectionName, const TCHAR* ArrayName, TArray<FString>& OutValues)
+{
+	Private::IrisDynamicConfig.GetArray(SectionName, ArrayName, OutValues);
+}
+
+}
+
+// FIrisDynamicConfigImpl
+namespace UE::Net::Private
+{
+
+FConfigBranch* FIrisDynamicConfigImpl::GetOrCreateConfig()
+{
+	if (IrisConfig == nullptr)
+	{
+		ConfigCache.Emplace(EConfigCacheType::Temporary);
+		IrisConfig = &ConfigCache->AddNewBranch(VirtualBranchFileName);
+	}
+
+	UE_ASSUME(IrisConfig != nullptr);
+	return IrisConfig;
+}
+
+void FIrisDynamicConfigImpl::RegisterDynamicConfig(FName UniqueConfigName, const FString& FilePath)
+{
+	if (!ensure(IsInGameThread()))
+	{
+		UE_LOGF(LogIris, Error, "FIrisDynamicConfig::RegisterDynamicConfig called from non-game thread. IGNORING dynamic config named '%ls'.", ToCStr(UniqueConfigName.ToString()));
+		return;
+	}
+	
+	if (ensureMsgf(!ConfigNameToFilePath.Find(UniqueConfigName), TEXT("FIrisDynamicConfig::RegisterDynamicConfig A dynamic config named '%s' has already been registered. IGNORING it."), ToCStr(UniqueConfigName.ToString())))
+	{
+		ConfigNameToFilePath.Add(UniqueConfigName, FilePath);
+
+		// Add config name to map and load it
+		FConfigModificationTracker ModTracker;
+		GetOrCreateConfig()->AddDynamicLayerToHierarchy(FilePath, &ModTracker);
+
+		// We're only modifying our branch so any modified section would be in there.
+		if (const TPair<FName, TSet<FString>>* ModifiedSections = ModTracker.ModifiedSectionsPerBranch.FindArbitraryElement())
+		{
+			OnIrisDynamicConfigChange.Broadcast(ModifiedSections->Value);
+		}
+	}
+}
+
+void FIrisDynamicConfigImpl::RegisterDynamicConfigBuffer(FName UniqueConfigName, const FString& VirtualFilePath, const FString& ConfigContents)
+{
+	if (!ensure(IsInGameThread()))
+	{
+		UE_LOGF(LogIris, Error, "FIrisDynamicConfig::RegisterDynamicConfigBuffer called from non-game thread. IGNORING dynamic config named '%ls'.", ToCStr(UniqueConfigName.ToString()));
+		return;
+	}
+
+	if (ConfigContents.IsEmpty())
+	{
+		UE_LOGF(LogIris, Verbose, "FIrisDynamicConfig::RegisterDynamicConfigBuffer Ignoring empty config named '%ls'", ToCStr(UniqueConfigName.ToString()));
+		return;
+	}
+
+	if (ensureMsgf(!ConfigNameToFilePath.Find(UniqueConfigName), TEXT("FIrisDynamicConfig::RegisterDynamicConfigBuffer A dynamic config named '%s' has already been registered. IGNORING it."), ToCStr(UniqueConfigName.ToString())))
+	{
+		ConfigNameToFilePath.Add(UniqueConfigName, VirtualFilePath);
+
+		// Add config name to map and process the config contents
+		constexpr EName Tag = NAME_None;
+		constexpr DynamicLayerPriority LayerPriority = DynamicLayerPriority::Unknown;
+		FConfigModificationTracker ModTracker;
+		GetOrCreateConfig()->AddDynamicLayerStringToHierarchy(VirtualFilePath, ConfigContents, Tag, LayerPriority, &ModTracker);
+
+		// We're only modifying our branch so any modified section would be in there.
+		if (const TPair<FName, TSet<FString>>* ModifiedSections = ModTracker.ModifiedSectionsPerBranch.FindArbitraryElement())
+		{
+			OnIrisDynamicConfigChange.Broadcast(ModifiedSections->Value);
+		}
+	}
+}
+
+void FIrisDynamicConfigImpl::RegisterDynamicConfigBuffer(FName UniqueConfigName, const FString& VirtualFilePath, FUtf8StringView ConfigContents)
+{
+	const FString ContentsAsString(ConfigContents);
+	return RegisterDynamicConfigBuffer(UniqueConfigName, VirtualFilePath, ContentsAsString);
+}
+
+void FIrisDynamicConfigImpl::UnregisterDynamicConfig(FName UniqueConfigName)
+{
+	if (IsEngineExitRequested())
+	{
+		return;
+	}
+
+	if (!ensure(IsInGameThread()))
+	{
+		UE_LOGF(LogIris, Error, "FIrisDynamicConfig::UnregisterDynamicConfig called from non-game thread. IGNORING unregistering config named '%ls'.", ToCStr(UniqueConfigName.ToString()));
+		return;
+	}
+
+	FString FilePath;
+	if (!ConfigNameToFilePath.RemoveAndCopyValue(UniqueConfigName, FilePath))
+	{
+		return;
+	}
+
+	FConfigModificationTracker ModTracker;
+	GetOrCreateConfig()->RemoveDynamicLayerFromHierarchy(FilePath, &ModTracker);
+
+	// We're only modifying our branch so any modified section would be in there.
+	if (const TPair<FName,TSet<FString>>* ModifiedSections = ModTracker.ModifiedSectionsPerBranch.FindArbitraryElement())
+	{
+		OnIrisDynamicConfigChange.Broadcast(ModifiedSections->Value);
+	}
+}
+
+void FIrisDynamicConfigImpl::GetSection(const TCHAR* SectionName, TArray<FString>& OutKeyValuePairs)
+{
+	if (FConfigCacheIni* Config = ConfigCache.GetPtrOrNull())
+	{
+		Config->GetSection(SectionName, OutKeyValuePairs, VirtualBranchFileName);
+	}
+}
+
+void FIrisDynamicConfigImpl::GetArray(const TCHAR* SectionName, const TCHAR* ArrayName, TArray<FString>& OutValues)
+{
+	if (FConfigCacheIni* Config = ConfigCache.GetPtrOrNull())
+	{
+		Config->GetArray(SectionName, ArrayName, OutValues, VirtualBranchFileName);
+	}
+}
+
+}
diff --git a/Engine/Source/Runtime/Net/Iris/Private/Iris/IrisCoreModule.cpp b/Engine/Source/Runtime/Net/Iris/Private/Iris/IrisCoreModule.cpp
--- a/Engine/Source/Runtime/Net/Iris/Private/Iris/IrisCoreModule.cpp
+++ b/Engine/Source/Runtime/Net/Iris/Private/Iris/IrisCoreModule.cpp
@@ -10,6 +10,7 @@
 #include "Net/Core/Connection/NetEnums.h"
 #include "Net/Core/Trace/Private/NetTraceInternal.h"
 
+#include "Iris/Config/IrisDynamicConfig.h"
 #include "Iris/Core/IrisLog.h"
 #include "Iris/ReplicationState/DefaultPropertyNetSerializerInfos.h"
 #include "Iris/ReplicationState/PropertyNetSerializerInfoRegistry.h"
@@ -89,7 +90,7 @@ private:
 	{
 		return false;
 	}
-
+	 
 	void OnAllModuleLoadingPhasesComplete()
 	{
 		FCoreDelegates::OnAllModuleLoadingPhasesComplete.RemoveAll(this);
diff --git a/Engine/Source/Runtime/Net/Iris/Public/Iris/Config/IrisDynamicConfig.h b/Engine/Source/Runtime/Net/Iris/Public/Iris/Config/IrisDynamicConfig.h
new file mode 100644
--- /dev/null
+++ b/Engine/Source/Runtime/Net/Iris/Public/Iris/Config/IrisDynamicConfig.h
@@ -0,0 +1,43 @@
+// Copyright Epic Games, Inc. All Rights Reserved.
+
+#pragma once
+
+#include "HAL/Platform.h"
+#include "Containers/StringFwd.h"
+#include "Delegates/Delegate.h"
+#include "UObject/NameTypes.h"
+
+namespace UE::Net
+{
+
+/**
+ * FIrisDynamicConfig is a system designed for Iris to handle runtime dynamic configuration of features based on dynamic loading/unloading or activation/deactivation of executable code, such as a GameFeaturePlugin.
+ * It will only support specific configuration sections rather than being generic like DefaultEngine.ini. This means that each subsystem will decide exactly what will be supported in terms of dynamic configuration.
+ */
+class UE_EXPERIMENTAL(5.8, "Iris dynamic config API is experimental and subject to change.") FIrisDynamicConfig
+{
+public:
+	/** Register an Iris config file for dynamically loaded or activated executable code. Make sure to use a unique name so that the config can be unloaded properly. Config must be unregistered via UnregisterDynamicConfig except during app shutdown. */
+	IRISCORE_API static void RegisterDynamicConfig(FName UniqueConfigName, const FString& FilePath);
+
+	/** Register an Iris config buffer for dynamically loaded or activated executable code. Make sure to use a unique name and virtual file path so that the config can be unloaded properly. Config must be unregistered via UnregisterDynamicConfig except during app shutdown. */
+	IRISCORE_API static void RegisterDynamicConfigBuffer(FName UniqueConfigName, const FString& VirtualFilePath, const FString& ConfigContents);
+
+	/** Register an Iris config buffer for dynamically loaded or activated executable code. Make sure to use a unique name and virtual file path so that the config can be unloaded properly. The contents will be converted to an FString internally. Config must be unregistered via UnregisterDynamicConfig except during app shutdown. */
+	IRISCORE_API static void RegisterDynamicConfigBuffer(FName UniqueConfigName, const FString& VirtualFilePath, FUtf8StringView ConfigContents);
+
+	/** Unregister the config file or buffer associated with the name. */
+	IRISCORE_API static void UnregisterDynamicConfig(FName UniqueConfigName);
+
+	using FOnIrisDynamicConfigChange = TMulticastDelegate<void(const TSet<FString>& ModifiedSections)>;
+	/** Dynamic config loaded/unloaded delegate registration/unregistration */
+	IRISCORE_API static FOnIrisDynamicConfigChange::RegistrationType& OnIrisDynamicConfigChange();
+
+	/** Retrieve the contents for a config section. */
+	IRISCORE_API static void GetSection(const TCHAR* SectionName, TArray<FString>& OutKeyValuePairs);
+
+	/** Retrieve the values for an array in a config section. */
+	IRISCORE_API static void GetArray(const TCHAR* SectionName, const TCHAR* ArrayName, TArray<FString>& OutValues);
+};
+
+}
