Iris
- Added unittest for invalidation of stale references to stable subobjects of dynamic replicated objects

#Synced-CL 51535387

diff --git a/Engine/Plugins/Runtime/ReplicationSystemTestPlugin/Source/Private/Tests/ReplicationSystem/TestObjectReferences.cpp b/Engine/Plugins/Runtime/ReplicationSystemTestPlugin/Source/Private/Tests/ReplicationSystem/TestObjectReferences.cpp
--- a/Engine/Plugins/Runtime/ReplicationSystemTestPlugin/Source/Private/Tests/ReplicationSystem/TestObjectReferences.cpp
+++ b/Engine/Plugins/Runtime/ReplicationSystemTestPlugin/Source/Private/Tests/ReplicationSystem/TestObjectReferences.cpp
@@ -47,9 +47,9 @@ void UTestObjectReferences_TestClassWithDefaultSubObject::RegisterReplicationFra
 	UE::Net::FReplicationFragmentUtil::CreateAndRegisterFragmentsForObject(this, Context, RegistrationFlags);
 
 	// Include subobject as well in the protocol
-	if (CreatedSubObjectRef)
+	if (CreatedSubObjectRef.Get())
 	{
-		static_cast<UTestReplicatedIrisObject*>(CreatedSubObjectRef)->RegisterReplicationFragments(Context, RegistrationFlags);
+		static_cast<UTestReplicatedIrisObject*>(CreatedSubObjectRef.Get())->RegisterReplicationFragments(Context, RegistrationFlags);
 	}
 }
 
@@ -58,6 +58,29 @@ void UTestObjectReferences_TestClassWithDefaultSubObject::GetLifetimeReplicatedP
 	DOREPLIFETIME(UTestObjectReferences_TestClassWithDefaultSubObject, ObjectRef);
 }
 
+UTestObjectReferences_TestClassWithNestedDefaultSubObjects::UTestObjectReferences_TestClassWithNestedDefaultSubObjects()
+: UReplicatedTestObject()
+{
+	CreatedSubObjectRef = CreateDefaultSubobject<UTestObjectReferences_TestClassWithDefaultSubObject>("TestDefaultSubObjectRef");
+}
+
+void UTestObjectReferences_TestClassWithNestedDefaultSubObjects::RegisterReplicationFragments(UE::Net::FFragmentRegistrationContext& Context, UE::Net::EFragmentRegistrationFlags RegistrationFlags)
+{
+	UE::Net::FReplicationFragmentUtil::CreateAndRegisterFragmentsForObject(this, Context, RegistrationFlags);
+
+	// Include subobject as well in the protocol
+	if (CreatedSubObjectRef.Get())
+	{
+		CreatedSubObjectRef.Get()->RegisterReplicationFragments(Context, RegistrationFlags);
+	}
+}
+
+void UTestObjectReferences_TestClassWithNestedDefaultSubObjects::GetLifetimeReplicatedProps(TArray< class FLifetimeProperty > & OutLifetimeProps) const
+{
+	DOREPLIFETIME(UTestObjectReferences_TestClassWithNestedDefaultSubObjects, ObjectRef);
+}
+
+
 namespace UE::Net::Private
 {
 
@@ -798,7 +821,7 @@ UE_NET_TEST_FIXTURE(FReplicationSystemServerClientTestFixture, TestInlinedSubObj
 
 	// Verify that we could resolve the subobject ref and that we got the expected reference
 	UE_NET_ASSERT_TRUE(ClientObject->ObjectRef != nullptr);
-	UE_NET_ASSERT_EQ(ClientObject->ObjectRef, ClientObject->CreatedSubObjectRef);
+	UE_NET_ASSERT_EQ(ClientObject->ObjectRef.Get(), ClientObject->CreatedSubObjectRef.Get());
 }
 
 // Test that reference to separate subobject works as expected
@@ -823,8 +846,8 @@ UE_NET_TEST_FIXTURE(FReplicationSystemServerClientTestFixture, TestExternalSubOb
 	Server->PostSendUpdate();
 
 	// Verify that object has been spawned on client
-	auto ClientObject = Client->GetObjectAs<UTestObjectReferences_TestClassWithDefaultSubObject>(Object->NetRefHandle);
-	auto ClientExternalSubObject = Client->GetObjectAs<UTestObjectReferences_TestClassWithDefaultSubObject>(ExternalSubObject->NetRefHandle);
+	UTestObjectReferences_TestClassWithDefaultSubObject* ClientObject = Client->GetObjectAs<UTestObjectReferences_TestClassWithDefaultSubObject>(Object->NetRefHandle);
+	UTestObjectReferences_TestClassWithDefaultSubObject* ClientExternalSubObject = Client->GetObjectAs<UTestObjectReferences_TestClassWithDefaultSubObject>(ExternalSubObject->NetRefHandle);
 
 	// Verify that default subobject is created
 	UE_NET_ASSERT_TRUE(ClientObject->CreatedSubObjectRef != nullptr);
@@ -841,7 +864,7 @@ UE_NET_TEST_FIXTURE(FReplicationSystemServerClientTestFixture, TestExternalSubOb
 
 	// Verify that we could resolve the subobject ref and that we got the expected reference
 	UE_NET_ASSERT_TRUE(ClientObject->ObjectRef != nullptr);
-	UE_NET_ASSERT_EQ(ClientObject->ObjectRef, TObjectPtr<UObject>(ClientExternalSubObject));
+	UE_NET_ASSERT_EQ(static_cast<UTestObjectReferences_TestClassWithDefaultSubObject*>(ClientObject->ObjectRef.Get()), ClientExternalSubObject);
 }
 
 // Test partial resolve of references in array
@@ -1515,4 +1538,170 @@ UE_NET_TEST_FIXTURE(FReplicationSystemServerClientTestFixture, TestTearOffDoesIn
 	UE_NET_ASSERT_EQ(ClientObject->StructWithRef.ObjectRef, TObjectPtr<UObject>(nullptr));
 }
 
+// Test that references to subobjects of stale dynamic objects works as epxected
+UE_NET_TEST_FIXTURE(FReplicationSystemServerClientTestFixture, TestStaleDefaultSubObjectReferences)
+{
+	TGuardConsoleVariable<bool> CVarOverride(TEXT("net.iris.VerifyOuterForStaleDynamicReferences"), true);
+
+	// Add a client
+	FReplicationSystemTestClient* Client = CreateClient();
+
+	// Spawn objects on server
+	UTestObjectReferences_TestClassWithDefaultSubObject* Object1 = Server->CreateObject<UTestObjectReferences_TestClassWithDefaultSubObject>();
+	UTestObjectReferences_TestClassWithDefaultSubObject* Object2 = Server->CreateObject<UTestObjectReferences_TestClassWithDefaultSubObject>();
+	// Have object2 reference the default subobject of object1
+	Object2->ObjectRef = Object1->CreatedSubObjectRef;
+
+	Server->UpdateAndSend({Client});
+
+	// Store original references
+	const FNetObjectReference OriginalRefToReplicatedObject1 = Server->ReplicationBridge->GetOrCreateObjectReference(Object1);
+	const FNetObjectReference OriginalRefToSubObject = Server->ReplicationBridge->GetOrCreateObjectReference(Object1->CreatedSubObjectRef);
+
+	// Verify that objects has been spawned on client and that we can resolve the replicated subobject via ref
+	{
+		UTestObjectReferences_TestClassWithDefaultSubObject* ClientObject1 = Client->GetObjectAs<UTestObjectReferences_TestClassWithDefaultSubObject>(Object1->NetRefHandle);
+		UTestObjectReferences_TestClassWithDefaultSubObject* ClientObject2 = Client->GetObjectAs<UTestObjectReferences_TestClassWithDefaultSubObject>(Object2->NetRefHandle);
+	
+		UE_NET_ASSERT_NE(ClientObject1, nullptr);
+		UE_NET_ASSERT_NE(ClientObject2, nullptr);
+		UE_NET_ASSERT_NE(ClientObject2->ObjectRef.Get(), nullptr);
+	}
+
+	// Stop replicating object1, Note: Object is not destroyed on server, but we tell server to destroy it on client
+	Server->ReplicationBridge->EndReplication(Object1, EEndReplicationFlags::Destroy);
+
+	Server->UpdateAndSend({Client});
+
+	constexpr bool bPerformFullPurge = true;
+	CollectGarbage(GARBAGE_COLLECTION_KEEPFLAGS, bPerformFullPurge);
+
+	// Verify new status
+	{
+		// This will assign a new reference
+		const FNetObjectReference NewRefToNotReplicatingObject1 = Server->ReplicationBridge->GetOrCreateObjectReference(Object1);
+		UE_NET_ASSERT_FALSE(NewRefToNotReplicatingObject1 == OriginalRefToReplicatedObject1);
+
+		// This will assign a new reference as the outer has changed.
+		const FNetObjectReference NewRefToNotReplicatingSubObject = Server->ReplicationBridge->GetOrCreateObjectReference(Object1->CreatedSubObjectRef);
+		UE_NET_ASSERT_FALSE(NewRefToNotReplicatingSubObject == OriginalRefToSubObject);
+
+		UTestObjectReferences_TestClassWithDefaultSubObject* ClientObject1 = Client->GetObjectAs<UTestObjectReferences_TestClassWithDefaultSubObject>(Object1->NetRefHandle);
+		UTestObjectReferences_TestClassWithDefaultSubObject* ClientObject2 = Client->GetObjectAs<UTestObjectReferences_TestClassWithDefaultSubObject>(Object2->NetRefHandle);
+	
+		UE_NET_ASSERT_EQ(ClientObject1, nullptr);
+		UE_NET_ASSERT_NE(ClientObject2, nullptr);
+		UE_NET_ASSERT_EQ(ClientObject2->ObjectRef.Get(), nullptr);
+	}
+
+	// Restart replication of object1, Note: since it is dynamic, it will get a new NetRefHandle
+	Server->ReplicationBridge->BeginReplication(Object1);
+
+	// Spawn objects on server
+	UTestObjectReferences_TestClassWithDefaultSubObject* Object3 = Server->CreateObject<UTestObjectReferences_TestClassWithDefaultSubObject>();
+
+	// Have object3 reference object1
+	Object3->ObjectRef = Object1->CreatedSubObjectRef;
+
+	Server->UpdateAndSend({Client});
+
+	// Verify that objects has been spawned on client and that we can resolve the new replicated subobject for client object 3
+	{
+		UTestObjectReferences_TestClassWithDefaultSubObject* ClientObject1 = Client->GetObjectAs<UTestObjectReferences_TestClassWithDefaultSubObject>(Object1->NetRefHandle);
+		UTestObjectReferences_TestClassWithDefaultSubObject* ClientObject2 = Client->GetObjectAs<UTestObjectReferences_TestClassWithDefaultSubObject>(Object2->NetRefHandle);
+		UTestObjectReferences_TestClassWithDefaultSubObject* ClientObject3 = Client->GetObjectAs<UTestObjectReferences_TestClassWithDefaultSubObject>(Object3->NetRefHandle);
+	
+		UE_NET_ASSERT_NE(ClientObject1, nullptr);
+		UE_NET_ASSERT_NE(ClientObject2, nullptr);
+		// This will not resolve as it references a subobject of the old no longer existing instance of the dynamic object
+		UE_NET_ASSERT_EQ(ClientObject2->ObjectRef.Get(), nullptr);
+		UE_NET_ASSERT_NE(ClientObject3->ObjectRef.Get(), nullptr);
+	}
+}
+
+// Test that references to nested default subobjects of stale replicated dynamic objects works as expected
+UE_NET_TEST_FIXTURE(FReplicationSystemServerClientTestFixture, TestStaleNestedDefaultSubObjectReferences)
+{
+	TGuardConsoleVariable<bool> CVarOverride(TEXT("net.iris.VerifyOuterForStaleDynamicReferences"), true);
+
+	// Add a client
+	FReplicationSystemTestClient* Client = CreateClient();
+
+	// Spawn objects on server
+	UTestObjectReferences_TestClassWithNestedDefaultSubObjects* Object1 = Server->CreateObject<UTestObjectReferences_TestClassWithNestedDefaultSubObjects>();
+	UTestObjectReferences_TestClassWithNestedDefaultSubObjects* Object2 = Server->CreateObject<UTestObjectReferences_TestClassWithNestedDefaultSubObjects>();
+
+	// Have object2 reference nested default subobject of object1->
+	Object2->ObjectRef = Object1->CreatedSubObjectRef->CreatedSubObjectRef;
+
+	Server->UpdateAndSend({Client});
+
+	// Store original references
+	const FNetObjectReference OriginalRefToReplicatedObject1 = Server->ReplicationBridge->GetOrCreateObjectReference(Object1);
+	const FNetObjectReference OriginalRefToSubObject = Server->ReplicationBridge->GetOrCreateObjectReference(Object1->CreatedSubObjectRef->CreatedSubObjectRef);
+
+	// Verify that objects has been spawned on client and that we can resolve the replicated subobject via ref
+	{
+		UTestObjectReferences_TestClassWithNestedDefaultSubObjects* ClientObject1 = Client->GetObjectAs<UTestObjectReferences_TestClassWithNestedDefaultSubObjects>(Object1->NetRefHandle);
+		UTestObjectReferences_TestClassWithNestedDefaultSubObjects* ClientObject2 = Client->GetObjectAs<UTestObjectReferences_TestClassWithNestedDefaultSubObjects>(Object2->NetRefHandle);
+	
+		UE_NET_ASSERT_NE(ClientObject1, nullptr);
+		UE_NET_ASSERT_NE(ClientObject2, nullptr);
+		UE_NET_ASSERT_NE(ClientObject2->ObjectRef.Get(), nullptr);
+	}
+
+	// Stop replicating object1, Note: Object is not destroyed on server, but we tell server to destroy it on client
+	Server->ReplicationBridge->EndReplication(Object1, EEndReplicationFlags::Destroy);
+
+	Server->UpdateAndSend({Client});
+
+	constexpr bool bPerformFullPurge = true;
+	CollectGarbage(GARBAGE_COLLECTION_KEEPFLAGS, bPerformFullPurge);
+
+	// Verify new status
+	{
+		// This will assign a new reference
+		const FNetObjectReference NewRefToNotReplicatingObject1 = Server->ReplicationBridge->GetOrCreateObjectReference(Object1);
+		UE_NET_ASSERT_FALSE(NewRefToNotReplicatingObject1 == OriginalRefToReplicatedObject1);
+
+		// This will assign a new reference as the outer has changed.
+		const FNetObjectReference NewRefToNotReplicatingSubObject = Server->ReplicationBridge->GetOrCreateObjectReference(Object1->CreatedSubObjectRef->CreatedSubObjectRef);
+		UE_NET_ASSERT_FALSE(NewRefToNotReplicatingSubObject == OriginalRefToSubObject);
+
+		UTestObjectReferences_TestClassWithNestedDefaultSubObjects* ClientObject1 = Client->GetObjectAs<UTestObjectReferences_TestClassWithNestedDefaultSubObjects>(Object1->NetRefHandle);
+		UTestObjectReferences_TestClassWithNestedDefaultSubObjects* ClientObject2 = Client->GetObjectAs<UTestObjectReferences_TestClassWithNestedDefaultSubObjects>(Object2->NetRefHandle);
+	
+		UE_NET_ASSERT_EQ(ClientObject1, nullptr);
+		UE_NET_ASSERT_NE(ClientObject2, nullptr);
+		UE_NET_ASSERT_EQ(ClientObject2->ObjectRef.Get(), nullptr);
+	}
+
+	// Restart replication of object1, Note: since it is dynamic, it will get a new NetRefHandle
+	Server->ReplicationBridge->BeginReplication(Object1);
+
+	// Spawn objects on server
+	UTestObjectReferences_TestClassWithNestedDefaultSubObjects* Object3 = Server->CreateObject<UTestObjectReferences_TestClassWithNestedDefaultSubObjects>();
+
+	// Have object3 reference object1
+	Object3->ObjectRef = Object1->CreatedSubObjectRef->CreatedSubObjectRef;
+
+	Server->UpdateAndSend({Client});
+
+	// Verify that objects has been spawned on client and that we can resolve the new replicated subobject for client object 3
+	{
+		UTestObjectReferences_TestClassWithNestedDefaultSubObjects* ClientObject1 = Client->GetObjectAs<UTestObjectReferences_TestClassWithNestedDefaultSubObjects>(Object1->NetRefHandle);
+		UTestObjectReferences_TestClassWithNestedDefaultSubObjects* ClientObject2 = Client->GetObjectAs<UTestObjectReferences_TestClassWithNestedDefaultSubObjects>(Object2->NetRefHandle);
+		UTestObjectReferences_TestClassWithNestedDefaultSubObjects* ClientObject3 = Client->GetObjectAs<UTestObjectReferences_TestClassWithNestedDefaultSubObjects>(Object3->NetRefHandle);
+	
+		UE_NET_ASSERT_NE(ClientObject1, nullptr);
+		UE_NET_ASSERT_NE(ClientObject2, nullptr);
+		UE_NET_ASSERT_NE(ClientObject3, nullptr);
+
+		// This will not resolve as it references a subobject of the old no longer existing instance of the dynamic object
+		UE_NET_ASSERT_EQ(ClientObject2->ObjectRef.Get(), nullptr);
+
+		UE_NET_ASSERT_NE(ClientObject3->ObjectRef.Get(), nullptr);
+	}
+}
+
 }
diff --git a/Engine/Plugins/Runtime/ReplicationSystemTestPlugin/Source/Private/Tests/ReplicationSystem/TestObjectReferences.h b/Engine/Plugins/Runtime/ReplicationSystemTestPlugin/Source/Private/Tests/ReplicationSystem/TestObjectReferences.h
--- a/Engine/Plugins/Runtime/ReplicationSystemTestPlugin/Source/Private/Tests/ReplicationSystem/TestObjectReferences.h
+++ b/Engine/Plugins/Runtime/ReplicationSystemTestPlugin/Source/Private/Tests/ReplicationSystem/TestObjectReferences.h
@@ -161,10 +161,29 @@ public:
 	virtual void RegisterReplicationFragments(UE::Net::FFragmentRegistrationContext& Fragments, UE::Net::EFragmentRegistrationFlags RegistrationFlags) override;
 
 	UPROPERTY(Transient, Replicated)
-	TObjectPtr<UObject> ObjectRef;
+	TWeakObjectPtr<UObject> ObjectRef;
 
 	// Used to store subobject created before we start replication
 	UPROPERTY(Transient)
 	TObjectPtr<UObject> CreatedSubObjectRef;
 };
 
+UCLASS()
+class UTestObjectReferences_TestClassWithNestedDefaultSubObjects : public UReplicatedTestObject
+{
+	GENERATED_BODY()
+
+public:
+	UTestObjectReferences_TestClassWithNestedDefaultSubObjects(); /* : UReplicatedTestObject() {} */
+
+	virtual void RegisterReplicationFragments(UE::Net::FFragmentRegistrationContext& Fragments, UE::Net::EFragmentRegistrationFlags RegistrationFlags) override;
+
+	UPROPERTY(Transient, Replicated)
+	TWeakObjectPtr<UObject> ObjectRef;
+
+	// Used to store subobject created before we start replication
+	UPROPERTY(Transient)
+	TObjectPtr<UTestObjectReferences_TestClassWithDefaultSubObject> CreatedSubObjectRef;
+};
+
+
