ReplicationSystemTestPlugin: Add FUserSuppliedArchivePlatformFile, an IPlatformFile implementation that
maps virtual file paths to caller-owned FArchive objects.

FScopedUserSuppliedArchivePlatformFile is a RAII helper that installs
FUserSuppliedArchivePlatformFile at the top of the platform file chain on
construction and restores the previous IPlatformFile on destruction.

#Synced-CL 51999447

diff --git a/Engine/Plugins/Runtime/ReplicationSystemTestPlugin/Source/Private/ReplicationSystemTestPlugin/HAL/UserSuppliedArchivePlatformFile.cpp b/Engine/Plugins/Runtime/ReplicationSystemTestPlugin/Source/Private/ReplicationSystemTestPlugin/HAL/UserSuppliedArchivePlatformFile.cpp
new file mode 100644
--- /dev/null
+++ b/Engine/Plugins/Runtime/ReplicationSystemTestPlugin/Source/Private/ReplicationSystemTestPlugin/HAL/UserSuppliedArchivePlatformFile.cpp
@@ -0,0 +1,191 @@
+// Copyright Epic Games, Inc. All Rights Reserved.
+
+#include "ReplicationSystemTestPlugin/HAL/UserSuppliedArchivePlatformFile.h"
+#include "Serialization/Archive.h"
+
+/**
+ * IFileHandle implementation backed by a borrowed FArchive.
+ *
+ * Each handle maintains its own Position and seeks the archive to that
+ * position before every read. This means multiple handles opened on the
+ * same archive each have an independent logical cursor, provided they are
+ * not used concurrently from multiple threads.
+ */
+class FUserSuppliedArchiveFileHandle : public IFileHandle
+{
+public:
+	explicit FUserSuppliedArchiveFileHandle(FArchive* InArchive)
+		: Archive(InArchive)
+	{}
+
+	virtual int64 Tell() override
+	{
+		return Position;
+	}
+
+	virtual bool Seek(int64 NewPosition) override
+	{
+		Position = NewPosition;
+		return true;
+	}
+
+	virtual bool SeekFromEnd(int64 NewPositionRelativeToEnd = 0) override
+	{
+		Position = Archive->TotalSize() + NewPositionRelativeToEnd;
+		return true;
+	}
+
+	virtual bool Read(uint8* Destination, int64 BytesToRead) override
+	{
+		Archive->Seek(Position);
+		Archive->Serialize(Destination, BytesToRead);
+		if (!Archive->IsError())
+		{
+			Position += BytesToRead;
+			return true;
+		}
+		return false;
+	}
+
+	virtual bool ReadAt(uint8* Destination, int64 BytesToRead, int64 Offset) override
+	{
+		// Read at the given offset without affecting the handle's current position.
+		Archive->Seek(Offset);
+		Archive->Serialize(Destination, BytesToRead);
+		return !Archive->IsError();
+	}
+
+	virtual bool Write(const uint8* Source, int64 BytesToWrite) override { return false; }
+	virtual bool Flush(const bool bFullFlush = false) override { return true; }
+	virtual bool Truncate(int64 NewSize) override { return false; }
+
+	virtual int64 Size() override
+	{
+		return Archive->TotalSize();
+	}
+
+private:
+	FArchive* Archive = nullptr;
+	int64 Position = 0;
+};
+
+
+void FUserSuppliedArchivePlatformFile::AddArchive(const TCHAR* Filename, FArchive* Archive)
+{
+	Archives.Add(NormalizePath(Filename), Archive);
+}
+
+bool FUserSuppliedArchivePlatformFile::FileExists(const TCHAR* Filename)
+{
+	if (Archives.Contains(NormalizePath(Filename)))
+	{
+		return true;
+	}
+	return LowerLevel ? LowerLevel->FileExists(Filename) : false;
+}
+
+int64 FUserSuppliedArchivePlatformFile::FileSize(const TCHAR* Filename)
+{
+	FArchive** Archive = Archives.Find(NormalizePath(Filename));
+	if (Archive)
+	{
+		return (*Archive)->TotalSize();
+	}
+	return LowerLevel ? LowerLevel->FileSize(Filename) : -1LL;
+}
+
+bool FUserSuppliedArchivePlatformFile::DeleteFile(const TCHAR* Filename)
+{
+	if (Archives.Contains(NormalizePath(Filename)))
+	{
+		return false;
+	}
+	return LowerLevel ? LowerLevel->DeleteFile(Filename) : false;
+}
+
+bool FUserSuppliedArchivePlatformFile::IsReadOnly(const TCHAR* Filename)
+{
+	if (Archives.Contains(NormalizePath(Filename)))
+	{
+		return true;
+	}
+	return LowerLevel ? LowerLevel->IsReadOnly(Filename) : false;
+}
+
+bool FUserSuppliedArchivePlatformFile::SetReadOnly(const TCHAR* Filename, bool bNewReadOnlyValue)
+{
+	if (Archives.Contains(NormalizePath(Filename)))
+	{
+		return false;
+	}
+	return LowerLevel ? LowerLevel->SetReadOnly(Filename, bNewReadOnlyValue) : false;
+}
+
+FDateTime FUserSuppliedArchivePlatformFile::GetTimeStamp(const TCHAR* Filename)
+{
+	if (Archives.Contains(NormalizePath(Filename)))
+	{
+		return FDateTime::Now();
+	}
+	return LowerLevel ? LowerLevel->GetTimeStamp(Filename) : FDateTime::MinValue();
+}
+
+void FUserSuppliedArchivePlatformFile::SetTimeStamp(const TCHAR* Filename, FDateTime DateTime)
+{
+	if (!Archives.Contains(NormalizePath(Filename)) && LowerLevel)
+	{
+		LowerLevel->SetTimeStamp(Filename, DateTime);
+	}
+}
+
+FDateTime FUserSuppliedArchivePlatformFile::GetAccessTimeStamp(const TCHAR* Filename)
+{
+	if (Archives.Contains(NormalizePath(Filename)))
+	{
+		return FDateTime::Now();
+	}
+	return LowerLevel ? LowerLevel->GetAccessTimeStamp(Filename) : FDateTime::MinValue();
+}
+
+FString FUserSuppliedArchivePlatformFile::GetFilenameOnDisk(const TCHAR* Filename)
+{
+	if (Archives.Contains(NormalizePath(Filename)))
+	{
+		return FString(Filename);
+	}
+	return LowerLevel ? LowerLevel->GetFilenameOnDisk(Filename) : FString(TEXT(""));
+}
+
+IFileHandle* FUserSuppliedArchivePlatformFile::OpenRead(const TCHAR* Filename, bool bAllowWrite)
+{
+	FArchive** Archive = Archives.Find(NormalizePath(Filename));
+	if (Archive)
+	{
+		return new FUserSuppliedArchiveFileHandle(*Archive);
+	}
+	return LowerLevel ? LowerLevel->OpenRead(Filename, bAllowWrite) : nullptr;
+}
+
+FFileStatData FUserSuppliedArchivePlatformFile::GetStatData(const TCHAR* FilenameOrDirectory)
+{
+	FArchive** Archive = Archives.Find(NormalizePath(FilenameOrDirectory));
+	if (Archive)
+	{
+		const FDateTime Now = FDateTime::Now();
+		return FFileStatData(
+			Now,	// CreationTime
+			Now,	// AccessTime
+			Now,	// ModificationTime
+			(*Archive)->TotalSize(),
+			/*bIsDirectory=*/false,
+			/*bIsReadOnly=*/true
+		);
+	}
+	if (LowerLevel)
+	{
+		return LowerLevel->GetStatData(FilenameOrDirectory);
+	}
+	FFileStatData Data;
+	Data.bIsValid = false;
+	return Data;
+}
diff --git a/Engine/Plugins/Runtime/ReplicationSystemTestPlugin/Source/Public/ReplicationSystemTestPlugin/HAL/UserSuppliedArchivePlatformFile.h b/Engine/Plugins/Runtime/ReplicationSystemTestPlugin/Source/Public/ReplicationSystemTestPlugin/HAL/UserSuppliedArchivePlatformFile.h
new file mode 100644
--- /dev/null
+++ b/Engine/Plugins/Runtime/ReplicationSystemTestPlugin/Source/Public/ReplicationSystemTestPlugin/HAL/UserSuppliedArchivePlatformFile.h
@@ -0,0 +1,104 @@
+// Copyright Epic Games, Inc. All Rights Reserved.
+
+#pragma once
+
+#include "Containers/Map.h"
+#include "Containers/UnrealString.h"
+#include "HAL/PlatformFile.h"
+#include "HAL/PlatformFileManager.h"
+#include "Misc/Paths.h"
+
+class FArchive;
+
+/**
+ * IPlatformFile implementation for use in tests.
+ *
+ * Register a non-owned FArchive under a virtual path with AddArchive. OpenRead
+ * returns a valid IFileHandle for any registered path. Each handle maintains
+ * its own read position by seeking the underlying archive before every read,
+ * so multiple handles opened on the same archive work correctly when used
+ * serially. Timestamps reported for registered paths are FDateTime::Now().
+ *
+ * All operations on paths not present in the archive map are forwarded to the
+ * lower-level IPlatformFile.
+ */
+class FUserSuppliedArchivePlatformFile : public IPlatformFile
+{
+public:
+	FUserSuppliedArchivePlatformFile() : LowerLevel(nullptr) {}
+
+	FUserSuppliedArchivePlatformFile(FUserSuppliedArchivePlatformFile&&) = delete;
+	FUserSuppliedArchivePlatformFile& operator=(FUserSuppliedArchivePlatformFile&&) = delete;
+	FUserSuppliedArchivePlatformFile(const FUserSuppliedArchivePlatformFile&) = delete;
+	FUserSuppliedArchivePlatformFile& operator=(const FUserSuppliedArchivePlatformFile&) = delete;
+
+	virtual bool ShouldBeUsed(IPlatformFile* Inner, const TCHAR* CmdLine) const override { return true; }
+	virtual bool Initialize(IPlatformFile* Inner, const TCHAR* CmdLine) override { LowerLevel = Inner; return true; }
+	virtual IPlatformFile* GetLowerLevel() override { return LowerLevel; }
+	virtual void SetLowerLevel(IPlatformFile* NewLowerLevel) override { LowerLevel = NewLowerLevel; }
+	virtual const TCHAR* GetName() const override { return TEXT("UserSuppliedArchivePlatformFile"); }
+
+	virtual bool FileExists(const TCHAR* Filename) override;
+	virtual int64 FileSize(const TCHAR* Filename) override;
+	virtual bool DeleteFile(const TCHAR* Filename) override;
+	virtual bool IsReadOnly(const TCHAR* Filename) override;
+	virtual bool MoveFile(const TCHAR* To, const TCHAR* From) override { return LowerLevel ? LowerLevel->MoveFile(To, From) : false; }
+	virtual bool SetReadOnly(const TCHAR* Filename, bool bNewReadOnlyValue) override;
+	virtual FDateTime GetTimeStamp(const TCHAR* Filename) override;
+	virtual void SetTimeStamp(const TCHAR* Filename, FDateTime DateTime) override;
+	virtual FDateTime GetAccessTimeStamp(const TCHAR* Filename) override;
+	virtual FString GetFilenameOnDisk(const TCHAR* Filename) override;
+	virtual IFileHandle* OpenRead(const TCHAR* Filename, bool bAllowWrite = false) override;
+	virtual IFileHandle* OpenWrite(const TCHAR* Filename, bool bAppend = false, bool bAllowRead = false) override { return LowerLevel ? LowerLevel->OpenWrite(Filename, bAppend, bAllowRead) : nullptr; }
+	virtual bool DirectoryExists(const TCHAR* Directory) override { return LowerLevel ? LowerLevel->DirectoryExists(Directory) : false; }
+	virtual bool CreateDirectory(const TCHAR* Directory) override { return LowerLevel ? LowerLevel->CreateDirectory(Directory) : false; }
+	virtual bool DeleteDirectory(const TCHAR* Directory) override { return LowerLevel ? LowerLevel->DeleteDirectory(Directory) : false; }
+	virtual FFileStatData GetStatData(const TCHAR* FilenameOrDirectory) override;
+	virtual bool IterateDirectory(const TCHAR* Directory, FDirectoryVisitor& Visitor) override { return LowerLevel ? LowerLevel->IterateDirectory(Directory, Visitor) : false; }
+	virtual bool IterateDirectoryStat(const TCHAR* Directory, FDirectoryStatVisitor& Visitor) override { return LowerLevel ? LowerLevel->IterateDirectoryStat(Directory, Visitor) : false; }
+
+	/** Register a non-owned FArchive under a virtual file path. */
+	void AddArchive(const TCHAR* Filename, FArchive* Archive);
+
+private:
+	static FString NormalizePath(const TCHAR* Filename)
+	{
+		FString Path(Filename);
+		FPaths::NormalizeFilename(Path);
+		return Path;
+	}
+
+	IPlatformFile* LowerLevel;
+	TMap<FString, FArchive*> Archives;
+};
+
+/**
+ * RAII helper that installs an FUserSuppliedArchivePlatformFile at the top of
+ * the platform file chain and restores the original chain on destruction.
+ *
+ * FArchive objects registered via AddArchive are served by OpenRead for the
+ * duration of the scope.
+ */
+class FScopedUserSuppliedArchivePlatformFile
+{
+public:
+	FScopedUserSuppliedArchivePlatformFile()
+	{
+		IPlatformFile& Current = FPlatformFileManager::Get().GetPlatformFile();
+		PlatformFile.Initialize(&Current, TEXT(""));
+		FPlatformFileManager::Get().SetPlatformFile(PlatformFile);
+	}
+
+	~FScopedUserSuppliedArchivePlatformFile()
+	{
+		FPlatformFileManager::Get().SetPlatformFile(*PlatformFile.GetLowerLevel());
+	}
+
+	void AddArchive(const TCHAR* Filename, FArchive* Archive)
+	{
+		PlatformFile.AddArchive(Filename, Archive);
+	}
+
+private:
+	FUserSuppliedArchivePlatformFile PlatformFile;
+};
