Iris - Add file-based tests to TestDynamicConfig.cpp that exercise FIrisDynamicConfig::RegisterDynamicConfig by
registering in-memory FMemoryReaderView archives under virtual INI paths.

#Synced-CL 51999447

diff --git a/Engine/Plugins/Runtime/ReplicationSystemTestPlugin/Source/Private/Tests/Config/TestDynamicConfig.cpp b/Engine/Plugins/Runtime/ReplicationSystemTestPlugin/Source/Private/Tests/Config/TestDynamicConfig.cpp
--- a/Engine/Plugins/Runtime/ReplicationSystemTestPlugin/Source/Private/Tests/Config/TestDynamicConfig.cpp
+++ b/Engine/Plugins/Runtime/ReplicationSystemTestPlugin/Source/Private/Tests/Config/TestDynamicConfig.cpp
@@ -4,7 +4,8 @@
 #include "NetworkAutomationTestMacros.h"
 #include "Containers/Set.h"
 #include "Iris/Config/IrisDynamicConfig.h"
-#include "Iris/Config/IrisDynamicConfig.h"
+#include "ReplicationSystemTestPlugin/HAL/UserSuppliedArchivePlatformFile.h"
+#include "Serialization/MemoryReader.h"
 #include "Tests/ReplicationSystem/ReplicationSystemTestFixture.h"
 #include "Templates/UnrealTemplate.h"
 
@@ -215,4 +216,101 @@ UE_NET_TEST_FIXTURE(FTestDynamicConfigFixture, CanRetrieveSectionConfigMulti)
 	UE_NET_ASSERT_NE(KeyValuePairsSet.Find("Key4=Value"), nullptr); 
 }
 
+// RegisterDynamicConfig (file-based) tests - these exercise FIrisDynamicConfig::RegisterDynamicConfig
+// using FScopedUserSuppliedArchivePlatformFile to serve INI content from in-memory archives.
+
+UE_NET_TEST_FIXTURE(FTestDynamicConfigFixture, CanRegisterDynamicConfig)
+{
+	const FUtf8StringView ConfigContent = UTF8TEXTVIEW("[SectionName]\nKey=Value");
+	FMemoryReaderView Archive(FMemoryView(ConfigContent.GetData(), ConfigContent.Len()));
+
+	FScopedUserSuppliedArchivePlatformFile StubFile;
+	StubFile.AddArchive(TEXT("TestDynamicConfig/Iris/Iris.ini"), &Archive);
+
+	FScopedConfigRegistrar Registrar;
+	Registrar.RegisterDynamicConfig(FName("TestDynamicConfig"), "TestDynamicConfig/Iris/Iris.ini");
+
+	const TSet<FString>& ModifiedSections = GetModifiedSections();
+	UE_NET_ASSERT_EQ(ModifiedSections.Num(), 1);
+	UE_NET_ASSERT_NE(ModifiedSections.Find("SectionName"), nullptr);
+}
+
+UE_NET_TEST_FIXTURE(FTestDynamicConfigFixture, RegisteringNonexistentFileIsNotBroadcasted)
+{
+	// No archive is registered for this path, so the config system will not find the file.
+	FScopedUserSuppliedArchivePlatformFile StubFile;
+
+	const unsigned OldCounter = GetDynamicConfigChangeCounter();
+
+	FScopedConfigRegistrar Registrar;
+	Registrar.RegisterDynamicConfig(FName("TestDynamicConfig"), "TestDynamicConfig/Iris/Missing.ini");
+
+	UE_NET_ASSERT_EQ(GetDynamicConfigChangeCounter(), OldCounter);
+}
+
+UE_NET_TEST_FIXTURE(FTestDynamicConfigFixture, CanUnregisterDynamicConfig)
+{
+	const FUtf8StringView ConfigContent = UTF8TEXTVIEW("[SectionName]\nKey=Value");
+	FMemoryReaderView Archive(FMemoryView(ConfigContent.GetData(), ConfigContent.Len()));
+
+	FScopedUserSuppliedArchivePlatformFile StubFile;
+	StubFile.AddArchive(TEXT("TestDynamicConfig/Iris/Iris.ini"), &Archive);
+
+	FScopedConfigRegistrar Registrar;
+	Registrar.RegisterDynamicConfig(FName("TestDynamicConfig"), "TestDynamicConfig/Iris/Iris.ini");
+
+	GetAndClearModifiedSections();
+
+	Registrar.UnregisterDynamicConfig(FName("TestDynamicConfig"));
+	const TSet<FString>& ModifiedSections = GetModifiedSections();
+	UE_NET_ASSERT_EQ(ModifiedSections.Num(), 1);
+	UE_NET_ASSERT_NE(ModifiedSections.Find("SectionName"), nullptr);
+}
+
+UE_NET_TEST_FIXTURE(FTestDynamicConfigFixture, CanRetrieveSectionConfigFromFile)
+{
+	const FUtf8StringView ConfigContent = UTF8TEXTVIEW("[SectionName]\nKey1=Value\nKey2=Value");
+	FMemoryReaderView Archive(FMemoryView(ConfigContent.GetData(), ConfigContent.Len()));
+
+	FScopedUserSuppliedArchivePlatformFile StubFile;
+	StubFile.AddArchive(TEXT("TestDynamicConfig/Iris/Iris.ini"), &Archive);
+
+	FScopedConfigRegistrar Registrar;
+	Registrar.RegisterDynamicConfig(FName("TestDynamicConfig"), "TestDynamicConfig/Iris/Iris.ini");
+
+	TArray<FString> KeyValuePairsArray;
+	FIrisDynamicConfig::GetSection(TEXT("SectionName"), KeyValuePairsArray);
+
+	TSet<FString> KeyValuePairsSet(KeyValuePairsArray);
+	UE_NET_ASSERT_EQ(KeyValuePairsSet.Num(), 2);
+	UE_NET_ASSERT_NE(KeyValuePairsSet.Find("Key1=Value"), nullptr);
+	UE_NET_ASSERT_NE(KeyValuePairsSet.Find("Key2=Value"), nullptr);
+}
+
+UE_NET_TEST_FIXTURE(FTestDynamicConfigFixture, CanRegisterMultipleDynamicConfigsFromFiles)
+{
+	const FUtf8StringView ConfigContent1 = UTF8TEXTVIEW("[SectionName]\nKey1=Value\nKey2=Value");
+	const FUtf8StringView ConfigContent2 = UTF8TEXTVIEW("[SectionName]\nKey3=Value\nKey4=Value");
+	FMemoryReaderView Archive1(FMemoryView(ConfigContent1.GetData(), ConfigContent1.Len()));
+	FMemoryReaderView Archive2(FMemoryView(ConfigContent2.GetData(), ConfigContent2.Len()));
+
+	FScopedUserSuppliedArchivePlatformFile StubFile;
+	StubFile.AddArchive(TEXT("TestDynamicConfig/Iris/Iris1.ini"), &Archive1);
+	StubFile.AddArchive(TEXT("TestDynamicConfig/Iris/Iris2.ini"), &Archive2);
+
+	FScopedConfigRegistrar Registrar;
+	Registrar.RegisterDynamicConfig(FName("TestDynamicConfig1"), "TestDynamicConfig/Iris/Iris1.ini");
+	Registrar.RegisterDynamicConfig(FName("TestDynamicConfig2"), "TestDynamicConfig/Iris/Iris2.ini");
+
+	TArray<FString> KeyValuePairsArray;
+	FIrisDynamicConfig::GetSection(TEXT("SectionName"), KeyValuePairsArray);
+
+	TSet<FString> KeyValuePairsSet(KeyValuePairsArray);
+	UE_NET_ASSERT_EQ(KeyValuePairsSet.Num(), 4);
+	UE_NET_ASSERT_NE(KeyValuePairsSet.Find("Key1=Value"), nullptr);
+	UE_NET_ASSERT_NE(KeyValuePairsSet.Find("Key2=Value"), nullptr);
+	UE_NET_ASSERT_NE(KeyValuePairsSet.Find("Key3=Value"), nullptr);
+	UE_NET_ASSERT_NE(KeyValuePairsSet.Find("Key4=Value"), nullptr);
+}
+
 }
