Iris - Added FScopedIrisDynamicConfig for easy and proper dynamic config management.

#Synced-CL 52460203

diff --git a/Engine/Plugins/Runtime/ReplicationSystemTestPlugin/Source/Private/Tests/Config/TestDynamicConfig.cpp b/Engine/Plugins/Runtime/ReplicationSystemTestPlugin/Source/Private/Tests/Config/TestDynamicConfig.cpp
--- a/Engine/Plugins/Runtime/ReplicationSystemTestPlugin/Source/Private/Tests/Config/TestDynamicConfig.cpp
+++ b/Engine/Plugins/Runtime/ReplicationSystemTestPlugin/Source/Private/Tests/Config/TestDynamicConfig.cpp
@@ -4,6 +4,7 @@
 #include "NetworkAutomationTestMacros.h"
 #include "Containers/Set.h"
 #include "Iris/Config/IrisDynamicConfig.h"
+#include "Iris/Config/ScopedIrisDynamicConfig.h"
 #include "ReplicationSystemTestPlugin/HAL/UserSuppliedArchivePlatformFile.h"
 #include "Serialization/MemoryReader.h"
 #include "Tests/ReplicationSystem/ReplicationSystemTestFixture.h"
@@ -55,47 +56,6 @@ protected:
 		return DynamicConfigChangeCounter;
 	}
 
-protected:
-	// This helper class keeps track of registered configs so that they're always properly unregistered, even when a test fails.
-	class FScopedConfigRegistrar
-	{
-	public:
-		~FScopedConfigRegistrar()
-		{
-			for (FName ConfigName : RegisteredConfigs)
-			{
-				FIrisDynamicConfig::UnregisterDynamicConfig(ConfigName);
-			}
-		}
-
-		void RegisterDynamicConfig(FName UniqueConfigName, const FString& FilePath)
-		{
-			RegisteredConfigs.Add(UniqueConfigName);
-			FIrisDynamicConfig::RegisterDynamicConfig(UniqueConfigName, FilePath);
-		}
-
-		void RegisterDynamicConfigBuffer(FName UniqueConfigName, const FString& VirtualFilePath, const FString& ConfigContents)
-		{
-			RegisteredConfigs.Add(UniqueConfigName);
-			FIrisDynamicConfig::RegisterDynamicConfigBuffer(UniqueConfigName, VirtualFilePath, ConfigContents);
-		}
-
-		void RegisterDynamicConfigBuffer(FName UniqueConfigName, const FString& VirtualFilePath, FUtf8StringView ConfigContents)
-		{
-			RegisteredConfigs.Add(UniqueConfigName);
-			FIrisDynamicConfig::RegisterDynamicConfigBuffer(UniqueConfigName, VirtualFilePath, ConfigContents);
-		}
-
-		void UnregisterDynamicConfig(FName UniqueConfigName)
-		{
-			RegisteredConfigs.Remove(UniqueConfigName);
-			FIrisDynamicConfig::UnregisterDynamicConfig(UniqueConfigName);
-		}
-
-	private:
-		TSet<FName, DefaultKeyFuncs<FName>, TInlineSetAllocator<8>> RegisteredConfigs;
-	};
-
 private:
 	// The set is cleared when the config system broadcasts a change
 	TSet<FString> OnDynamicConfigChangeModifiedSections;
@@ -104,8 +64,8 @@ private:
 
 UE_NET_TEST_FIXTURE(FTestDynamicConfigFixture, CanRegisterDynamicConfigBuffer)
 {
-	FScopedConfigRegistrar Registrar;
-	Registrar.RegisterDynamicConfigBuffer(FName("TestDynamicConfig"), "TestDynamicConfig/Iris/Iris.ini", UTF8TEXTVIEW("[SectionName]"));
+	FScopedIrisDynamicConfig ScopedConfig;
+	ScopedConfig.RegisterDynamicConfigBuffer(FName("TestDynamicConfig"), "TestDynamicConfig/Iris/Iris.ini", UTF8TEXTVIEW("[SectionName]"));
 
 	const TSet<FString>& ModifiedSections = GetModifiedSections();
 	UE_NET_ASSERT_EQ(ModifiedSections.Num(), 1);
@@ -114,25 +74,25 @@ UE_NET_TEST_FIXTURE(FTestDynamicConfigFixture, CanRegisterDynamicConfigBuffer)
 
 UE_NET_TEST_FIXTURE(FTestDynamicConfigFixture, RegisteringEmptyDynamicConfigIsNotBroadcasted)
 {
-	FScopedConfigRegistrar Registrar;
+	FScopedIrisDynamicConfig ScopedConfig;
 
 	const unsigned OldCounter = GetDynamicConfigChangeCounter();
-	Registrar.RegisterDynamicConfigBuffer(FName("TestDynamicConfig1"), "TestDynamicConfig/Iris/Iris1.ini", UTF8TEXTVIEW(""));
+	ScopedConfig.RegisterDynamicConfigBuffer(FName("TestDynamicConfig1"), "TestDynamicConfig/Iris/Iris1.ini", UTF8TEXTVIEW(""));
 	UE_NET_ASSERT_EQ(GetDynamicConfigChangeCounter(), OldCounter);
 
-	Registrar.RegisterDynamicConfigBuffer(FName("TestDynamicConfig2"), "TestDynamicConfig/Iris/Iris2.ini", FString());
+	ScopedConfig.RegisterDynamicConfigBuffer(FName("TestDynamicConfig2"), "TestDynamicConfig/Iris/Iris2.ini", FString());
 	UE_NET_ASSERT_EQ(GetDynamicConfigChangeCounter(), OldCounter);
 }
 
 UE_NET_TEST_FIXTURE(FTestDynamicConfigFixture, CanUnregisterDynamicConfigBuffer)
 {
-	FScopedConfigRegistrar Registrar;
+	FScopedIrisDynamicConfig ScopedConfig;
 
-	Registrar.RegisterDynamicConfigBuffer(FName("TestDynamicConfig"), "TestDynamicConfig/Iris/Iris.ini", UTF8TEXTVIEW("[SectionName]"));
+	ScopedConfig.RegisterDynamicConfigBuffer(FName("TestDynamicConfig"), "TestDynamicConfig/Iris/Iris.ini", UTF8TEXTVIEW("[SectionName]"));
 
 	GetAndClearModifiedSections();
 
-	Registrar.UnregisterDynamicConfig(FName("TestDynamicConfig"));
+	ScopedConfig.UnregisterDynamicConfig(FName("TestDynamicConfig"));
 	const TSet<FString>& ModifiedSections = GetModifiedSections();
 	UE_NET_ASSERT_EQ(ModifiedSections.Num(), 1);
 	UE_NET_ASSERT_NE(ModifiedSections.Find("SectionName"), nullptr);
@@ -140,19 +100,19 @@ UE_NET_TEST_FIXTURE(FTestDynamicConfigFixture, CanUnregisterDynamicConfigBuffer)
 
 UE_NET_TEST_FIXTURE(FTestDynamicConfigFixture, CanRegisterAndUnregisterDynamicConfigsInDifferentOrder)
 {
-	FScopedConfigRegistrar Registrar;
+	FScopedIrisDynamicConfig ScopedConfig;
 
-	Registrar.RegisterDynamicConfigBuffer(FName("TestDynamicConfig1"), "TestDynamicConfig/Iris/Iris1.ini", UTF8TEXTVIEW("[SectionName1]"));
-	Registrar.RegisterDynamicConfigBuffer(FName("TestDynamicConfig2"), "TestDynamicConfig/Iris/Iris2.ini", UTF8TEXTVIEW("[SectionName2]"));
+	ScopedConfig.RegisterDynamicConfigBuffer(FName("TestDynamicConfig1"), "TestDynamicConfig/Iris/Iris1.ini", UTF8TEXTVIEW("[SectionName1]"));
+	ScopedConfig.RegisterDynamicConfigBuffer(FName("TestDynamicConfig2"), "TestDynamicConfig/Iris/Iris2.ini", UTF8TEXTVIEW("[SectionName2]"));
 
 	GetAndClearModifiedSections();
 
-	Registrar.UnregisterDynamicConfig(FName("TestDynamicConfig2"));
+	ScopedConfig.UnregisterDynamicConfig(FName("TestDynamicConfig2"));
 	const TSet<FString>& ModifiedSections1 = GetModifiedSections();
 	UE_NET_ASSERT_EQ(ModifiedSections1.Num(), 1);
 	UE_NET_ASSERT_NE(ModifiedSections1.Find("SectionName2"), nullptr);
 
-	Registrar.UnregisterDynamicConfig(FName("TestDynamicConfig1"));
+	ScopedConfig.UnregisterDynamicConfig(FName("TestDynamicConfig1"));
 	const TSet<FString>& ModifiedSections2 = GetModifiedSections();
 	UE_NET_ASSERT_EQ(ModifiedSections2.Num(), 1);
 	UE_NET_ASSERT_NE(ModifiedSections2.Find("SectionName1"), nullptr);
@@ -160,20 +120,20 @@ UE_NET_TEST_FIXTURE(FTestDynamicConfigFixture, CanRegisterAndUnregisterDynamicCo
 
 UE_NET_TEST_FIXTURE(FTestDynamicConfigFixture, CanRegisterAndUnregisterConfigsWithOverlappingSectionsInDifferentOrder)
 {
-	FScopedConfigRegistrar Registrar;
+	FScopedIrisDynamicConfig ScopedConfig;
 
-	Registrar.RegisterDynamicConfigBuffer(FName("TestDynamicConfig1"), "TestDynamicConfig/Iris/Iris1.ini", UTF8TEXTVIEW("[SectionName1]\n; Hello\n[SectionName3]\n; World"));
-	Registrar.RegisterDynamicConfigBuffer(FName("TestDynamicConfig2"), "TestDynamicConfig/Iris/Iris2.ini", UTF8TEXTVIEW("[SectionName3]\n; Hello\n[SectionName2]\n; World"));
+	ScopedConfig.RegisterDynamicConfigBuffer(FName("TestDynamicConfig1"), "TestDynamicConfig/Iris/Iris1.ini", UTF8TEXTVIEW("[SectionName1]\n; Hello\n[SectionName3]\n; World"));
+	ScopedConfig.RegisterDynamicConfigBuffer(FName("TestDynamicConfig2"), "TestDynamicConfig/Iris/Iris2.ini", UTF8TEXTVIEW("[SectionName3]\n; Hello\n[SectionName2]\n; World"));
 
 	GetAndClearModifiedSections();
 
-	Registrar.UnregisterDynamicConfig(FName("TestDynamicConfig2"));
+	ScopedConfig.UnregisterDynamicConfig(FName("TestDynamicConfig2"));
 	const TSet<FString>& ModifiedSections1 = GetModifiedSections();
 	UE_NET_ASSERT_EQ(ModifiedSections1.Num(), 2);
 	UE_NET_ASSERT_NE(ModifiedSections1.Find("SectionName2"), nullptr);
 	UE_NET_ASSERT_NE(ModifiedSections1.Find("SectionName3"), nullptr);
 
-	Registrar.UnregisterDynamicConfig(FName("TestDynamicConfig1"));
+	ScopedConfig.UnregisterDynamicConfig(FName("TestDynamicConfig1"));
 	const TSet<FString>& ModifiedSections2 = GetModifiedSections();
 	UE_NET_ASSERT_EQ(ModifiedSections2.Num(), 2);
 	UE_NET_ASSERT_NE(ModifiedSections2.Find("SectionName1"), nullptr);
@@ -182,9 +142,9 @@ UE_NET_TEST_FIXTURE(FTestDynamicConfigFixture, CanRegisterAndUnregisterConfigsWi
 
 UE_NET_TEST_FIXTURE(FTestDynamicConfigFixture, CanRetrieveSectionConfig)
 {
-	FScopedConfigRegistrar Registrar;
+	FScopedIrisDynamicConfig ScopedConfig;
 
-	Registrar.RegisterDynamicConfigBuffer(FName("TestDynamicConfig"), "TestDynamicConfig/Iris/Iris.ini", UTF8TEXTVIEW("[SectionName]\nKey1=Value\nKey2=Value"));
+	ScopedConfig.RegisterDynamicConfigBuffer(FName("TestDynamicConfig"), "TestDynamicConfig/Iris/Iris.ini", UTF8TEXTVIEW("[SectionName]\nKey1=Value\nKey2=Value"));
 
 	TArray<FString> KeyValuePairsArray;
 	FIrisDynamicConfig::GetSection(TEXT("SectionName"), KeyValuePairsArray);
@@ -198,11 +158,11 @@ UE_NET_TEST_FIXTURE(FTestDynamicConfigFixture, CanRetrieveSectionConfig)
 
 UE_NET_TEST_FIXTURE(FTestDynamicConfigFixture, CanRetrieveSectionConfigMulti)
 {
-	FScopedConfigRegistrar Registrar;
+	FScopedIrisDynamicConfig ScopedConfig;
 
 	// Register multiple configs for the same section
-	Registrar.RegisterDynamicConfigBuffer(FName("TestDynamicConfig1"), "TestDynamicConfig/Iris/Iris1.ini", UTF8TEXTVIEW("[SectionName]\nKey1=Value\nKey2=Value"));
-	Registrar.RegisterDynamicConfigBuffer(FName("TestDynamicConfig2"), "TestDynamicConfig/Iris/Iris2.ini", UTF8TEXTVIEW("[SectionName]\nKey3=Value\nKey4=Value"));
+	ScopedConfig.RegisterDynamicConfigBuffer(FName("TestDynamicConfig1"), "TestDynamicConfig/Iris/Iris1.ini", UTF8TEXTVIEW("[SectionName]\nKey1=Value\nKey2=Value"));
+	ScopedConfig.RegisterDynamicConfigBuffer(FName("TestDynamicConfig2"), "TestDynamicConfig/Iris/Iris2.ini", UTF8TEXTVIEW("[SectionName]\nKey3=Value\nKey4=Value"));
 
 	TArray<FString> KeyValuePairsArray;
 	FIrisDynamicConfig::GetSection(TEXT("SectionName"), KeyValuePairsArray);
@@ -227,8 +187,8 @@ UE_NET_TEST_FIXTURE(FTestDynamicConfigFixture, CanRegisterDynamicConfig)
 	FScopedUserSuppliedArchivePlatformFile StubFile;
 	StubFile.AddArchive(TEXT("TestDynamicConfig/Iris/Iris.ini"), &Archive);
 
-	FScopedConfigRegistrar Registrar;
-	Registrar.RegisterDynamicConfig(FName("TestDynamicConfig"), "TestDynamicConfig/Iris/Iris.ini");
+	FScopedIrisDynamicConfig ScopedConfig;
+	ScopedConfig.RegisterDynamicConfig(FName("TestDynamicConfig"), "TestDynamicConfig/Iris/Iris.ini");
 
 	const TSet<FString>& ModifiedSections = GetModifiedSections();
 	UE_NET_ASSERT_EQ(ModifiedSections.Num(), 1);
@@ -242,8 +202,8 @@ UE_NET_TEST_FIXTURE(FTestDynamicConfigFixture, RegisteringNonexistentFileIsNotBr
 
 	const unsigned OldCounter = GetDynamicConfigChangeCounter();
 
-	FScopedConfigRegistrar Registrar;
-	Registrar.RegisterDynamicConfig(FName("TestDynamicConfig"), "TestDynamicConfig/Iris/Missing.ini");
+	FScopedIrisDynamicConfig ScopedConfig;
+	ScopedConfig.RegisterDynamicConfig(FName("TestDynamicConfig"), "TestDynamicConfig/Iris/Missing.ini");
 
 	UE_NET_ASSERT_EQ(GetDynamicConfigChangeCounter(), OldCounter);
 }
@@ -256,12 +216,12 @@ UE_NET_TEST_FIXTURE(FTestDynamicConfigFixture, CanUnregisterDynamicConfig)
 	FScopedUserSuppliedArchivePlatformFile StubFile;
 	StubFile.AddArchive(TEXT("TestDynamicConfig/Iris/Iris.ini"), &Archive);
 
-	FScopedConfigRegistrar Registrar;
-	Registrar.RegisterDynamicConfig(FName("TestDynamicConfig"), "TestDynamicConfig/Iris/Iris.ini");
+	FScopedIrisDynamicConfig ScopedConfig;
+	ScopedConfig.RegisterDynamicConfig(FName("TestDynamicConfig"), "TestDynamicConfig/Iris/Iris.ini");
 
 	GetAndClearModifiedSections();
 
-	Registrar.UnregisterDynamicConfig(FName("TestDynamicConfig"));
+	ScopedConfig.UnregisterDynamicConfig(FName("TestDynamicConfig"));
 	const TSet<FString>& ModifiedSections = GetModifiedSections();
 	UE_NET_ASSERT_EQ(ModifiedSections.Num(), 1);
 	UE_NET_ASSERT_NE(ModifiedSections.Find("SectionName"), nullptr);
@@ -275,8 +235,8 @@ UE_NET_TEST_FIXTURE(FTestDynamicConfigFixture, CanRetrieveSectionConfigFromFile)
 	FScopedUserSuppliedArchivePlatformFile StubFile;
 	StubFile.AddArchive(TEXT("TestDynamicConfig/Iris/Iris.ini"), &Archive);
 
-	FScopedConfigRegistrar Registrar;
-	Registrar.RegisterDynamicConfig(FName("TestDynamicConfig"), "TestDynamicConfig/Iris/Iris.ini");
+	FScopedIrisDynamicConfig ScopedConfig;
+	ScopedConfig.RegisterDynamicConfig(FName("TestDynamicConfig"), "TestDynamicConfig/Iris/Iris.ini");
 
 	TArray<FString> KeyValuePairsArray;
 	FIrisDynamicConfig::GetSection(TEXT("SectionName"), KeyValuePairsArray);
@@ -298,9 +258,9 @@ UE_NET_TEST_FIXTURE(FTestDynamicConfigFixture, CanRegisterMultipleDynamicConfigs
 	StubFile.AddArchive(TEXT("TestDynamicConfig/Iris/Iris1.ini"), &Archive1);
 	StubFile.AddArchive(TEXT("TestDynamicConfig/Iris/Iris2.ini"), &Archive2);
 
-	FScopedConfigRegistrar Registrar;
-	Registrar.RegisterDynamicConfig(FName("TestDynamicConfig1"), "TestDynamicConfig/Iris/Iris1.ini");
-	Registrar.RegisterDynamicConfig(FName("TestDynamicConfig2"), "TestDynamicConfig/Iris/Iris2.ini");
+	FScopedIrisDynamicConfig ScopedConfig;
+	ScopedConfig.RegisterDynamicConfig(FName("TestDynamicConfig1"), "TestDynamicConfig/Iris/Iris1.ini");
+	ScopedConfig.RegisterDynamicConfig(FName("TestDynamicConfig2"), "TestDynamicConfig/Iris/Iris2.ini");
 
 	TArray<FString> KeyValuePairsArray;
 	FIrisDynamicConfig::GetSection(TEXT("SectionName"), KeyValuePairsArray);
diff --git a/Engine/Source/Runtime/Net/Iris/Public/Iris/Config/ScopedIrisDynamicConfig.h b/Engine/Source/Runtime/Net/Iris/Public/Iris/Config/ScopedIrisDynamicConfig.h
new file mode 100644
--- /dev/null
+++ b/Engine/Source/Runtime/Net/Iris/Public/Iris/Config/ScopedIrisDynamicConfig.h
@@ -0,0 +1,59 @@
+// Copyright Epic Games, Inc. All Rights Reserved.
+
+#pragma once
+
+#include "Iris/Config/IrisDynamicConfig.h"
+#include "Containers/Set.h"
+
+namespace UE::Net
+{
+
+/**
+ * FScopedIrisDynamicConfig is a helper class for FIrisDynamicConfig that unregisters all registered dynamic configs in its destructor.
+ * Configs can be unregistered manually if desired. The API mimics FIrisDynamicConfig when it comes to registering and unregistering individual configs.
+ */
+class UE_EXPERIMENTAL(5.8, "Iris dynamic config API is experimental and subject to change.") FScopedIrisDynamicConfig
+{
+public:
+	~FScopedIrisDynamicConfig()
+	{
+		UnregisterAllDynamicConfigs();
+	}
+
+	void RegisterDynamicConfig(FName UniqueConfigName, const FString& FilePath)
+	{
+		RegisteredConfigs.Add(UniqueConfigName);
+		FIrisDynamicConfig::RegisterDynamicConfig(UniqueConfigName, FilePath);
+	}
+
+	void RegisterDynamicConfigBuffer(FName UniqueConfigName, const FString& VirtualFilePath, const FString& ConfigContents)
+	{
+		RegisteredConfigs.Add(UniqueConfigName);
+		FIrisDynamicConfig::RegisterDynamicConfigBuffer(UniqueConfigName, VirtualFilePath, ConfigContents);
+	}
+
+	void RegisterDynamicConfigBuffer(FName UniqueConfigName, const FString& VirtualFilePath, FUtf8StringView ConfigContents)
+	{
+		RegisteredConfigs.Add(UniqueConfigName);
+		FIrisDynamicConfig::RegisterDynamicConfigBuffer(UniqueConfigName, VirtualFilePath, ConfigContents);
+	}
+
+	void UnregisterDynamicConfig(FName UniqueConfigName)
+	{
+		RegisteredConfigs.Remove(UniqueConfigName);
+		FIrisDynamicConfig::UnregisterDynamicConfig(UniqueConfigName);
+	}
+
+	void UnregisterAllDynamicConfigs()
+	{
+		for (FName ConfigName : RegisteredConfigs)
+		{
+			FIrisDynamicConfig::UnregisterDynamicConfig(ConfigName);
+		}
+		RegisteredConfigs.Empty();
+	}
+private:
+	TSet<FName, DefaultKeyFuncs<FName>, TInlineSetAllocator<1>> RegisteredConfigs;
+};
+
+}
